Added ability to extend the default context menu in the scene outliner
This tidys up some logic and duplication between the level editor viewport context menu and the scene outliner, since the level viewport menu can now extend the scene outliner, rather than having to duplicate its functionality. The ability to completely override the menu still remains by way of an override delegate.

diff --git a/Engine/Source/Editor/Layers/Private/SLayerBrowser.h b/Engine/Source/Editor/Layers/Private/SLayerBrowser.h
--- a/Engine/Source/Editor/Layers/Private/SLayerBrowser.h
+++ b/Engine/Source/Editor/Layers/Private/SLayerBrowser.h
@@ -172,7 +172,7 @@ public:
 		.BorderImage( FEditorStyle::GetBrush( "NoBrush" ) )
 		.Content()
 		[
-			SceneOutlinerModule.CreateSceneOutliner( InitOptions, FOnContextMenuOpening(), FOnActorPicked() )
+			SceneOutlinerModule.CreateSceneOutliner( InitOptions, FOnActorPicked() )
 		];
 
 
diff --git a/Engine/Source/Editor/LevelEditor/Private/LevelEditorActions.cpp b/Engine/Source/Editor/LevelEditor/Private/LevelEditorActions.cpp
--- a/Engine/Source/Editor/LevelEditor/Private/LevelEditorActions.cpp
+++ b/Engine/Source/Editor/LevelEditor/Private/LevelEditorActions.cpp
@@ -423,18 +423,6 @@ void FLevelEditorActionCallbacks::AttachToSocketSelection(const FName SocketName
 	}	
 }
 
-void FLevelEditorActionCallbacks::SetActorFolder(FName Folder)
-{
-	FSlateApplication::Get().DismissAllMenus();
-
-	FScopedTransaction Transaction(LOCTEXT("SetActorFolder", "Set Actor Folder"));
-	for ( FSelectionIterator It( GEditor->GetSelectedActorIterator() ) ; It ; ++It )
-	{
-		AActor* Actor = Cast<AActor>( *It );
-		Actor->SetFolderPath(Folder);
-	}
-}
-
 void FLevelEditorActionCallbacks::SetMaterialQualityLevel( EMaterialQualityLevel::Type NewQualityLevel )
 {
 	GUnrealEd->AccessEditorUserSettings().MaterialQualityLevel = NewQualityLevel;
diff --git a/Engine/Source/Editor/LevelEditor/Private/LevelEditorToolBar.cpp b/Engine/Source/Editor/LevelEditor/Private/LevelEditorToolBar.cpp
--- a/Engine/Source/Editor/LevelEditor/Private/LevelEditorToolBar.cpp
+++ b/Engine/Source/Editor/LevelEditor/Private/LevelEditorToolBar.cpp
@@ -1027,7 +1027,6 @@ TSharedRef< SWidget > FLevelEditorToolBar::GenerateMatineeMenuContent( TSharedRe
 		[
 			SceneOutlinerModule.CreateSceneOutliner(
 				InitOptions,
-				FOnContextMenuOpening(), //no context menu allowed here
 				FOnActorPicked::CreateStatic( &FLevelEditorToolBar::OnMatineeActorPicked ) )
 		];
 
diff --git a/Engine/Source/Editor/LevelEditor/Private/LevelViewportContextMenu.cpp b/Engine/Source/Editor/LevelEditor/Private/LevelViewportContextMenu.cpp
--- a/Engine/Source/Editor/LevelEditor/Private/LevelViewportContextMenu.cpp
+++ b/Engine/Source/Editor/LevelEditor/Private/LevelViewportContextMenu.cpp
@@ -28,7 +28,6 @@
 #include "LevelEditorCreateActorMenu.h"
 #include "SourceCodeNavigation.h"
 #include "Developer/MeshUtilities/Public/MeshUtilities.h"
-#include "SceneOutlinerTreeItems.h"
 
 #define LOCTEXT_NAMESPACE "LevelViewportContextMenu"
 
@@ -85,13 +84,6 @@ public:
 	 */
 	static void FillActorMenu( class FMenuBuilder& MenuBuilder );
 
-	/**
-	 * Fills in menu options for the actor folder menu
-	 *
-	 * @param MenuBuilder	The menu to add items to
-	 */
-	static void FillActorFoldersMenu( class FMenuBuilder& MenuBuilder );
-
 	/**
 	 * Fills in menu options for the snap menu
 	 *
@@ -152,13 +144,23 @@ struct FLevelScriptEventMenuHelper
 //       payload contains a weak reference to a level editor instance
 TSharedPtr< SWidget > FLevelViewportContextMenu::BuildMenuWidget( TWeakPtr< SLevelEditor > LevelEditor, TSharedPtr<FExtender> Extender )
 {
-	TSharedRef< FAssetThumbnailPool > ThumbnailPool = MakeShareable(new FAssetThumbnailPool(5, true));
+	// Build up the menu
+	const bool bShouldCloseWindowAfterMenuSelection = true;
+	FMenuBuilder MenuBuilder(bShouldCloseWindowAfterMenuSelection, TSharedPtr<const FUICommandList>());
 
+	FillMenu(MenuBuilder, LevelEditor, Extender);
+	
+	return MenuBuilder.MakeWidget();
+}
+
+void FLevelViewportContextMenu::FillMenu( FMenuBuilder& MenuBuilder, TWeakPtr<SLevelEditor> LevelEditor, TSharedPtr<FExtender> Extender )
+{
 	// Generate information about our selection
 	TArray<AActor*> SelectedActors;
 	GEditor->GetSelectedActors()->GetSelectedObjects<AActor>( SelectedActors );
-	FLevelViewportContextMenuImpl::SelectionInfo = AssetSelectionUtils::BuildSelectedActorInfo( SelectedActors );
 
+	FSelectedActorInfo& SelectionInfo = FLevelViewportContextMenuImpl::SelectionInfo;
+	SelectionInfo = AssetSelectionUtils::BuildSelectedActorInfo( SelectedActors );
 
 	// Get all menu extenders for this context menu from the level editor module
 	FLevelEditorModule& LevelEditorModule = FModuleManager::GetModuleChecked<FLevelEditorModule>( TEXT("LevelEditor") );
@@ -170,21 +172,17 @@ TSharedPtr< SWidget > FLevelViewportContextMenu::BuildMenuWidget( TWeakPtr< SLev
 		Extenders.Add(Extender);
 	}
 
+	auto LevelEditorActions = LevelEditor.Pin()->GetLevelEditorActions().ToSharedRef();
 	for (int32 i = 0; i < MenuExtenderDelegates.Num(); ++i)
 	{
 		if (MenuExtenderDelegates[i].IsBound())
 		{
-			Extenders.Add(MenuExtenderDelegates[i].Execute(LevelEditor.Pin()->GetLevelEditorActions().ToSharedRef(), SelectedActors));
+			Extenders.Add(MenuExtenderDelegates[i].Execute(LevelEditorActions, SelectedActors));
 		}
 	}
-	TSharedPtr<FExtender> MenuExtender = FExtender::Combine(Extenders);
 
-
-	FSelectedActorInfo& SelectionInfo = FLevelViewportContextMenuImpl::SelectionInfo;
-
-	// Build up the menu
-	const bool bShouldCloseWindowAfterMenuSelection = true;	// Set the menu to automatically close when the user commits to a choice
-	FMenuBuilder MenuBuilder( bShouldCloseWindowAfterMenuSelection, LevelEditor.Pin()->GetLevelEditorActions(), MenuExtender );
+	MenuBuilder.PushCommandList(LevelEditorActions);
+	MenuBuilder.PushExtender(FExtender::Combine(Extenders).ToSharedRef());
 
 	TArray<TWeakObjectPtr<UObject>> LabelObjects;
 	for ( FSelectionIterator SelItor(*GEditor->GetSelectedActors()) ; SelItor ; ++SelItor )
@@ -319,14 +317,6 @@ TSharedPtr< SWidget > FLevelViewportContextMenu::BuildMenuWidget( TWeakPtr< SLev
 				LOCTEXT( "ActorAttachToSubMenu_ToolTip", "Attach Actor as child" ),
 				FNewMenuDelegate::CreateStatic( &FLevelViewportContextMenuImpl::FillActorMenu ) );
 		}
-		
-		if ( GEditor->GetSelectedActorCount() )
-		{
-			MenuBuilder.AddSubMenu( 
-				LOCTEXT( "MoveActorsToFolder", "Move To Folder" ), 
-				LOCTEXT( "MoveActorsToFolder_Tooltip", "Move selected actors to a folder" ),
-				FNewMenuDelegate::CreateStatic( &FLevelViewportContextMenuImpl::FillActorFoldersMenu ) );
-		}
 
 		// Add a heading for "Movement" if an actor is selected
 		if ( GEditor->GetSelectedActorIterator() )
@@ -361,7 +351,8 @@ TSharedPtr< SWidget > FLevelViewportContextMenu::BuildMenuWidget( TWeakPtr< SLev
 
 	FLevelScriptEventMenuHelper::FillLevelBlueprintEventsMenu(MenuBuilder, SelectedActors);
 
-	return MenuBuilder.MakeWidget();
+	MenuBuilder.PopCommandList();
+	MenuBuilder.PopExtender();
 }
 
 
@@ -828,7 +819,6 @@ void FLevelViewportContextMenuImpl::FillActorMenu( FMenuBuilder& MenuBuilder )
 		[
 			SceneOutlinerModule.CreateSceneOutliner(
 			InitOptions,
-			FOnContextMenuOpening(), //no context menu allowed here
 			FOnActorPicked::CreateStatic( &FLevelEditorActionCallbacks::AttachToActor ) )
 		];
 
@@ -865,55 +855,6 @@ void FLevelViewportContextMenuImpl::FillActorMenu( FMenuBuilder& MenuBuilder )
 	MenuBuilder.AddWidget(MiniSceneOutliner, FText::GetEmpty(), false);
 }
 
-void FLevelViewportContextMenuImpl::FillActorFoldersMenu( FMenuBuilder& MenuBuilder )
-{
-	FSceneOutlinerInitializationOptions InitOptions;
-	InitOptions.bShowHeaderRow = false;
-	InitOptions.bFocusSearchBoxWhenOpened = true;
-	InitOptions.bOnlyShowFolders = true;
-
-	// Actor selector to allow the user to choose a folder
-	FSceneOutlinerModule& SceneOutlinerModule = FModuleManager::LoadModuleChecked<FSceneOutlinerModule>( "SceneOutliner" );
-	TSharedRef< SWidget > MiniSceneOutliner =
-		SNew( SVerticalBox )
-		+SVerticalBox::Slot()
-		.MaxHeight(400.0f)
-		[
-			SceneOutlinerModule.CreateSceneOutliner(
-			InitOptions,
-			FOnContextMenuOpening(), //no context menu allowed here
-			FOnSceneOutlinerItemPicked::CreateStatic([](TSharedRef<SceneOutliner::TOutlinerTreeItem> Item){
-				if (Item->Type == SceneOutliner::TOutlinerTreeItem::Folder)
-				{
-					auto FolderItem = StaticCastSharedRef<SceneOutliner::TOutlinerFolderTreeItem>(Item);
-					FLevelEditorActionCallbacks::SetActorFolder(FolderItem->Path);
-				}
-			}))
-		];
-
-	// Create New...
-	{
-		FText ToolTipText;
-		switch (SelectionInfo.NumSelected)
-		{
-		case 0:		ToolTipText = LOCTEXT( "CreateNewFolder_ToolTip0",		"Create a new root folder" );							break;
-		case 1:		ToolTipText = LOCTEXT( "CreateNewFolder_ToolTip1",		"Create a new root folder with this actor in" );		break;
-		default:	ToolTipText = LOCTEXT( "CreateNewFolder_ToolTipMany",	"Create a new root folder with these actors in" );		break;
-		}
-
-		MenuBuilder.AddMenuEntry(FLevelEditorCommands::Get().CreateNewOutlinerFolder, NAME_None, LOCTEXT( "CreateNewFolder", "Create New" ),
-			ToolTipText, FSlateIcon(FEditorStyle::GetStyleSetName(), "SceneOutliner.NewFolderIcon"));
-	}
-
-	MenuBuilder.AddMenuEntry(LOCTEXT( "MoveToRoot", "Move To Root" ), LOCTEXT( "MoveToRoot_ToolTip", "Move to the root of the scene outliner" ),
-		FSlateIcon(FEditorStyle::GetStyleSetName(), "SceneOutliner.MoveToRoot"), FExecuteAction::CreateStatic([](){ FLevelEditorActionCallbacks::SetActorFolder(FName()); }));
-
-	// Existing folders
-	MenuBuilder.BeginSection(FName(), LOCTEXT("ExistingFolders", "Existing:"));
-	MenuBuilder.AddWidget(MiniSceneOutliner, FText::GetEmpty(), false);
-	MenuBuilder.EndSection();
-}
-
 void FLevelViewportContextMenuImpl::FillSnapAlignMenu( FMenuBuilder& MenuBuilder )
 {
 	MenuBuilder.AddMenuEntry( FLevelEditorCommands::Get().SnapOriginToGrid );
diff --git a/Engine/Source/Editor/LevelEditor/Private/LevelViewportContextMenu.h b/Engine/Source/Editor/LevelEditor/Private/LevelViewportContextMenu.h
--- a/Engine/Source/Editor/LevelEditor/Private/LevelViewportContextMenu.h
+++ b/Engine/Source/Editor/LevelEditor/Private/LevelViewportContextMenu.h
@@ -27,8 +27,16 @@ public:
 	 * @param	Extender		Allows extension of this menu based on context.
 	 * @return	Widget for this context menu
 	 */
-	static TSharedPtr< SWidget > BuildMenuWidget( TWeakPtr< SLevelEditor > LevelEditor, TSharedPtr<FExtender> Extender = TSharedPtr<FExtender>() );
+	static TSharedPtr< SWidget > BuildMenuWidget(TWeakPtr< SLevelEditor > LevelEditor, TSharedPtr<FExtender> Extender = TSharedPtr<FExtender>());
 
+	/**
+	 * Populates the specified menu builder for the context menu that can be inserted into a pop-up window
+	 *
+	 * @param	MenuBuilder		The menu builder to fill the menu with
+	 * @param	LevelEditor		The level editor using this menu.
+	 * @param	Extender		Allows extension of this menu based on context.
+	 */
+	static void FillMenu(FMenuBuilder& MenuBuilder, TWeakPtr< SLevelEditor > LevelEditor, TSharedPtr<FExtender> Extender);
 
 private:
 
diff --git a/Engine/Source/Editor/LevelEditor/Private/SLevelEditor.cpp b/Engine/Source/Editor/LevelEditor/Private/SLevelEditor.cpp
--- a/Engine/Source/Editor/LevelEditor/Private/SLevelEditor.cpp
+++ b/Engine/Source/Editor/LevelEditor/Private/SLevelEditor.cpp
@@ -603,6 +603,21 @@ TSharedRef<SDockTab> SLevelEditor::SpawnLevelEditorTab( const FSpawnTabArgs& Arg
 	{
 		FSceneOutlinerInitializationOptions InitOptions;
 		InitOptions.Mode = ESceneOutlinerMode::ActorBrowsing;
+		{
+			TWeakPtr<SLevelEditor> WeakLevelEditor = SharedThis(this);
+			InitOptions.DefaultMenuExtender = MakeShareable(new FExtender);
+			InitOptions.DefaultMenuExtender->AddMenuExtension(
+				"FolderSection", EExtensionHook::Before, GetLevelEditorActions(),
+				FMenuExtensionDelegate::CreateStatic([](FMenuBuilder& MenuBuilder, TWeakPtr<SLevelEditor> WeakLevelEditor){
+					// Only extend the menu if we have actors selected
+					if (GEditor->GetSelectedActors()->Num())
+					{
+						FLevelViewportContextMenu::FillMenu(MenuBuilder, WeakLevelEditor, TSharedPtr<FExtender>());
+					}
+				}, WeakLevelEditor)
+			);
+		}
+
 
 		FText Label = NSLOCTEXT( "LevelEditor", "SceneOutlinerTabTitle", "Scene Outliner" );
 
@@ -621,7 +636,6 @@ TSharedRef<SDockTab> SLevelEditor::SpawnLevelEditorTab( const FSpawnTabArgs& Arg
 					[
 						SceneOutlinerModule.CreateSceneOutliner(
 							InitOptions,
-							FOnContextMenuOpening::CreateStatic< TWeakPtr< SLevelEditor >, TSharedPtr<FExtender> >( &FLevelViewportContextMenu::BuildMenuWidget, SharedThis( this ), TSharedPtr<FExtender>() ),
 							FOnActorPicked() /* Not used for outliner when in browsing mode */ )
 					]
 				]
diff --git a/Engine/Source/Editor/LevelEditor/Private/SLevelViewportToolBar.cpp b/Engine/Source/Editor/LevelEditor/Private/SLevelViewportToolBar.cpp
--- a/Engine/Source/Editor/LevelEditor/Private/SLevelViewportToolBar.cpp
+++ b/Engine/Source/Editor/LevelEditor/Private/SLevelViewportToolBar.cpp
@@ -336,7 +336,6 @@ static void OnGenerateActorLockingMenuSection( TWeakPtr<SLevelViewport> Viewport
 		[
 			SceneOutlinerModule.CreateSceneOutliner(
 				InitOptions,
-				FOnContextMenuOpening(), //no context menu allowed here
 				FOnActorPicked::CreateSP( Viewport.Pin().ToSharedRef(), &SLevelViewport::OnActorLockToggleFromMenu ) )
 		];
 
diff --git a/Engine/Source/Editor/LevelEditor/Public/LevelEditorActions.h b/Engine/Source/Editor/LevelEditor/Public/LevelEditorActions.h
--- a/Engine/Source/Editor/LevelEditor/Public/LevelEditorActions.h
+++ b/Engine/Source/Editor/LevelEditor/Public/LevelEditorActions.h
@@ -711,7 +711,6 @@ public:
 	static void ShowLightingStaticMeshInfo();
 	static void AttachToActor(AActor* ParentActorPtr );
 	static void AttachToSocketSelection(FName SocketName, AActor* ParentActorPtr);
-	static void SetActorFolder(FName Folder);
 	static void SetMaterialQualityLevel( EMaterialQualityLevel::Type NewQualityLevel );
 	static bool IsMaterialQualityLevelChecked( EMaterialQualityLevel::Type TestQualityLevel );
 	
diff --git a/Engine/Source/Editor/PropertyEditor/Private/UserInterface/PropertyEditor/SPropertyMenuActorPicker.cpp b/Engine/Source/Editor/PropertyEditor/Private/UserInterface/PropertyEditor/SPropertyMenuActorPicker.cpp
--- a/Engine/Source/Editor/PropertyEditor/Private/UserInterface/PropertyEditor/SPropertyMenuActorPicker.cpp
+++ b/Engine/Source/Editor/PropertyEditor/Private/UserInterface/PropertyEditor/SPropertyMenuActorPicker.cpp
@@ -84,7 +84,7 @@ void SPropertyMenuActorPicker::Construct( const FArguments& InArgs )
 				SNew( SBorder )
 				.BorderImage( FEditorStyle::GetBrush("Menu.Background") )
 				[
-					SceneOutlinerModule.CreateSceneOutliner(InitOptions, FOnContextMenuOpening(), FOnActorPicked::CreateSP(this, &SPropertyMenuActorPicker::OnActorSelected))
+					SceneOutlinerModule.CreateSceneOutliner(InitOptions, FOnActorPicked::CreateSP(this, &SPropertyMenuActorPicker::OnActorSelected))
 				]
 			];
 
diff --git a/Engine/Source/Editor/PropertyEditor/Private/UserInterface/PropertyEditor/SPropertySceneOutliner.cpp b/Engine/Source/Editor/PropertyEditor/Private/UserInterface/PropertyEditor/SPropertySceneOutliner.cpp
--- a/Engine/Source/Editor/PropertyEditor/Private/UserInterface/PropertyEditor/SPropertySceneOutliner.cpp
+++ b/Engine/Source/Editor/PropertyEditor/Private/UserInterface/PropertyEditor/SPropertySceneOutliner.cpp
@@ -68,7 +68,7 @@ TSharedRef<SWidget> SPropertySceneOutliner::OnGenerateSceneOutliner()
 			SNew( SBorder )
 			.BorderImage( FEditorStyle::GetBrush("Menu.Background") )
 			[
-				SceneOutlinerModule.CreateSceneOutliner(InitOptions, FOnContextMenuOpening(), FOnActorPicked::CreateSP(this, &SPropertySceneOutliner::OnActorSelectedFromOutliner))
+				SceneOutlinerModule.CreateSceneOutliner(InitOptions, FOnActorPicked::CreateSP(this, &SPropertySceneOutliner::OnActorSelectedFromOutliner))
 			]
 		];
 
diff --git a/Engine/Source/Editor/SceneOutliner/Private/SOutlinerTreeView.cpp b/Engine/Source/Editor/SceneOutliner/Private/SOutlinerTreeView.cpp
--- a/Engine/Source/Editor/SceneOutliner/Private/SOutlinerTreeView.cpp
+++ b/Engine/Source/Editor/SceneOutliner/Private/SOutlinerTreeView.cpp
@@ -424,33 +424,6 @@ namespace SceneOutliner
 		return ValidationInfo.IsValid();
 	}
 
-	void SOutlinerTreeView::MoveSelectionTo(FName NewParent)
-	{
-		FSlateApplication::Get().DismissAllMenus();
-
-		if (!ValidateMoveSelectionTo(NewParent))
-		{
-			return;
-		}
-
-		const FScopedTransaction Transaction( LOCTEXT("MoveOutlinerItems", "Move Scene Outliner Items") );
-		for (const auto& Item : GetSelectedItems())
-		{
-			if (Item->Type == TOutlinerTreeItem::Folder)
-			{
-				StaticCastSharedPtr<TOutlinerFolderTreeItem>(Item)->MoveTo(NewParent);
-			}
-			else
-			{
-				auto* Actor = StaticCastSharedPtr<TOutlinerActorTreeItem>(Item)->Actor.Get();
-				if (Actor)
-				{
-					Actor->SetFolderPath(NewParent);
-				}
-			}
-		}
-	}
-
 	FReply SOutlinerTreeView::OnDragOver(const FGeometry& MyGeometry, const FDragDropEvent& DragDropEvent)
 	{
 		auto SceneOutlinerPtr = SceneOutlinerWeak.Pin();
@@ -540,8 +513,7 @@ namespace SceneOutliner
 				auto FolderPtr = Folder.Pin();
 				if (FolderPtr.IsValid())
 				{
-					// Put it at the root
-					FolderPtr->MoveTo(FName());
+					SceneOutlinerPtr->MoveFolderTo(FolderPtr.ToSharedRef(), FName());
 				}
 			}
 		}
@@ -589,13 +561,12 @@ namespace SceneOutliner
 		if (Item->Type == TOutlinerTreeItem::Folder && DraggedFolders)
 		{
 			auto DropFolder = StaticCastSharedPtr<TOutlinerFolderTreeItem>(Item);
-
 			for (auto& FolderIt : *DraggedFolders)
 			{
 				auto FolderItemPtr = FolderIt.Pin();
 				if (FolderItemPtr.IsValid())
 				{
-					FolderItemPtr->MoveTo(DropFolder->Path);
+					SceneOutlinerPtr->MoveFolderTo(FolderItemPtr.ToSharedRef(), DropFolder->Path);
 				}
 			}
 		}
diff --git a/Engine/Source/Editor/SceneOutliner/Private/SOutlinerTreeView.h b/Engine/Source/Editor/SceneOutliner/Private/SOutlinerTreeView.h
--- a/Engine/Source/Editor/SceneOutliner/Private/SOutlinerTreeView.h
+++ b/Engine/Source/Editor/SceneOutliner/Private/SOutlinerTreeView.h
@@ -26,9 +26,6 @@ namespace SceneOutliner
 		/** Check that it is valid to move the current selection to the specified folder path */
 		bool ValidateMoveSelectionTo(FName NewParent);
 
-		/** Moves the current selection to the specified folder path */
-		void MoveSelectionTo(FName NewParent);
-
 	protected:
 
 		virtual FReply OnDragOver( const FGeometry& MyGeometry, const FDragDropEvent& DragDropEvent ) OVERRIDE;
diff --git a/Engine/Source/Editor/SceneOutliner/Private/SSceneOutliner.cpp b/Engine/Source/Editor/SceneOutliner/Private/SSceneOutliner.cpp
--- a/Engine/Source/Editor/SceneOutliner/Private/SSceneOutliner.cpp
+++ b/Engine/Source/Editor/SceneOutliner/Private/SSceneOutliner.cpp
@@ -67,10 +67,28 @@ namespace SceneOutliner
 		};
 	}
 
+	void PerformActionOnItem(TSharedPtr<SOutlinerTreeView> OutlinerTreeView, TSharedRef<TOutlinerTreeItem> NewItem, uint8 ActionMask)
+	{
+		if (ActionMask & NewItemActuator::Select)
+		{
+			OutlinerTreeView->ClearSelection();
+			OutlinerTreeView->SetItemSelection(NewItem, true);
+		}
+
+		if (ActionMask & NewItemActuator::Rename)
+		{
+			NewItem->Flags.RenameWhenInView = true;
+		}
+
+		if (ActionMask & (NewItemActuator::ScrollIntoView | NewItemActuator::Rename))
+		{
+			OutlinerTreeView->RequestScrollIntoView(NewItem);
+		}
+	}
+
 	void SSceneOutliner::Construct( const FArguments& InArgs, const FSceneOutlinerInitializationOptions& InInitOptions )
 	{
 		InitOptions = InInitOptions;
-		OnContextMenuOpening = InArgs._MakeContextMenuWidgetDelegate;
 
 		OnItemPicked = InArgs._OnItemPickedDelegate;
 
@@ -727,6 +745,8 @@ namespace SceneOutliner
 			RequestSort();
 		}
 
+		NewItemActions.Empty();
+
 		// We're fully refreshed now.
 		bNeedsRefresh = false;
 	}
@@ -955,10 +975,7 @@ namespace SceneOutliner
 			}
 		}
 
-		if (InActorItem->Flags.RenameWhenInView)
-		{
-			OutlinerTreeView->RequestScrollIntoView(InActorItem);
-		}
+		NewItemActions.ItemHasBeenCreated(OutlinerTreeView.ToSharedRef(), InActorItem);
 
 		return true;
 	}
@@ -1001,10 +1018,7 @@ namespace SceneOutliner
 		// Folders are always expanded by default
 		OutlinerTreeView->SetItemExpansion(InFolderItem, true);
 
-		if (InFolderItem->Flags.RenameWhenInView)
-		{
-			OutlinerTreeView->RequestScrollIntoView(InFolderItem);
-		}
+		NewItemActions.ItemHasBeenCreated(OutlinerTreeView.ToSharedRef(), InFolderItem);
 
 		return true;
 	}
@@ -1418,7 +1432,21 @@ namespace SceneOutliner
 			auto FolderTreeItem = StaticCastSharedPtr<TOutlinerFolderTreeItem>(TreeItem);
 			if (!InLabel.EqualTo(FText::FromName(FolderTreeItem->LeafName)))
 			{
-				FolderTreeItem->Rename(FName(*InLabel.ToString()));
+				// Rename the item
+				FName NewPath = FolderTreeItem->GetParentPath();
+				if (NewPath.IsNone())
+				{
+					NewPath = FName(*InLabel.ToString());
+				}
+				else
+				{
+					NewPath = FName(*(NewPath.ToString() / InLabel.ToString()));
+				}
+
+				if (FActorFolders::Get().RenameFolderInWorld(*RepresentingWorld, FolderTreeItem->Path, NewPath))
+				{
+					NewItemActions.WhenCreated(NewPath, NewItemActuator::Select | NewItemActuator::ScrollIntoView);
+				}
 			}
 		}
 		
@@ -1493,27 +1521,24 @@ namespace SceneOutliner
 		TArray<AActor*> SelectedActors;
 		GEditor->GetSelectedActors()->GetSelectedObjects<AActor>( SelectedActors );
 
-		if (SelectedActors.Num())
+		if (SelectedActors.Num() && InitOptions.ContextMenuOverride.IsBound())
 		{
-			if (OnContextMenuOpening.IsBound())
-			{
-				return OnContextMenuOpening.Execute();
-			}
+			return InitOptions.ContextMenuOverride.Execute();
 		}
-		else if (InitOptions.Mode == ESceneOutlinerMode::ActorBrowsing)
+
+		if (InitOptions.Mode == ESceneOutlinerMode::ActorBrowsing)
 		{
-			return BuildFolderContextMenu();
+			return BuildDefaultContextMenu();
 		}
 
 		return TSharedPtr<SWidget>();
 	}
 
-	/** Build a context menu for right-clicking a folder */
-	TSharedPtr<SWidget> SSceneOutliner::BuildFolderContextMenu()  const
+	TSharedPtr<SWidget> SSceneOutliner::BuildDefaultContextMenu()  const
 	{
 		// Build up the menu
 		const bool bCloseAfterSelection = true;
-		FMenuBuilder MenuBuilder(bCloseAfterSelection, TSharedPtr<FUICommandList>());
+		FMenuBuilder MenuBuilder(bCloseAfterSelection, TSharedPtr<FUICommandList>(), InitOptions.DefaultMenuExtender);
 
 		TSharedPtr<TOutlinerFolderTreeItem> ParentFolder;
 		auto SelectedItems = OutlinerTreeView->GetSelectedItems();
@@ -1522,23 +1547,32 @@ namespace SceneOutliner
 			ParentFolder = StaticCastSharedPtr<TOutlinerFolderTreeItem>(SelectedItems[0]);
 		}
 
-		MenuBuilder.BeginSection(FName(), LOCTEXT("FolderClassName", "Folder"));
+		const FSlateIcon NewFolderIcon(FEditorStyle::GetStyleSetName(), "SceneOutliner.NewFolderIcon");
+		MenuBuilder.BeginSection("FolderSection", LOCTEXT("ThisFolderSection", "This Folder"));
 		{
-			MenuBuilder.AddMenuEntry(LOCTEXT("CreateFolder", "Create Folder"), FText(), FSlateIcon(FEditorStyle::GetStyleSetName(), "SceneOutliner.NewFolderIcon"),
-				FUIAction(FExecuteAction::CreateSP(this, &SSceneOutliner::CreateFolder, ParentFolder.IsValid() ? ParentFolder->Path : FName())));
-
 			if (ParentFolder.IsValid())
 			{
+				MenuBuilder.AddMenuEntry(LOCTEXT("CreateSubFolder", "Create Sub Folder"), FText(), NewFolderIcon, FUIAction(FExecuteAction::CreateSP(this, &SSceneOutliner::CreateSubFolder, ParentFolder->Path)));
 				MenuBuilder.AddMenuEntry(LOCTEXT("RenameFolder", "Rename Folder"), FText(), FSlateIcon(), FUIAction(FExecuteAction::CreateSP(this, &SSceneOutliner::RenameFolder, ParentFolder.ToSharedRef())));
 				MenuBuilder.AddMenuEntry(LOCTEXT("DeleteFolder", "Delete Folder"), FText(), FSlateIcon(), FUIAction(FExecuteAction::CreateSP(this, &SSceneOutliner::DeleteFolder, ParentFolder.ToSharedRef())));
 			}
 		}
 		MenuBuilder.EndSection();
 
-		MenuBuilder.AddSubMenu( 
-			LOCTEXT("MoveActorsToFolder", "Move To"),
-			LOCTEXT("MoveActorsToFolder_Tooltip", "Move selection to a folder"),
-			FNewMenuDelegate::CreateSP(this,  &SSceneOutliner::FillFoldersSubMenu));
+		if (SelectedItems.Num() != 0)
+		{
+			MenuBuilder.BeginSection("SelectionSection", LOCTEXT("OutlinerSelection", "Current Outliner Selection"));
+			{
+				MenuBuilder.AddSubMenu(
+					LOCTEXT("MoveActorsToFolder", "Move To Folder"),
+					LOCTEXT("MoveActorsToFolder_Tooltip", "Move selection to a folder"),
+					FNewMenuDelegate::CreateSP(this,  &SSceneOutliner::FillFoldersSubMenu));
+			}
+		}
+		else
+		{
+			MenuBuilder.AddMenuEntry(LOCTEXT("CreateFolder", "Create Folder"), FText(), NewFolderIcon, FUIAction(FExecuteAction::CreateSP(this, &SSceneOutliner::CreateFolder)));
+		}
 
 		return MenuBuilder.MakeWidget();
 	}
@@ -1563,10 +1597,13 @@ namespace SceneOutliner
 				.OnItemPickedDelegate(FOnSceneOutlinerItemPicked::CreateSP(this, &SSceneOutliner::MoveSelectionTo))
 			];
 
+		MenuBuilder.AddMenuEntry(LOCTEXT( "CreateNew", "Create New" ), LOCTEXT( "MoveToRoot_ToolTip", "Move to a new folder" ),
+			FSlateIcon(FEditorStyle::GetStyleSetName(), "SceneOutliner.NewFolderIcon"), FExecuteAction::CreateSP(this, &SSceneOutliner::CreateFolder));
+
 		if (OutlinerTreeView->ValidateMoveSelectionTo(FName()))
 		{
 			MenuBuilder.AddMenuEntry(LOCTEXT( "MoveToRoot", "Move To Root" ), LOCTEXT( "MoveToRoot_ToolTip", "Move to the root of the scene outliner" ),
-				FSlateIcon(FEditorStyle::GetStyleSetName(), "SceneOutliner.MoveToRoot"), FExecuteAction::CreateSP(OutlinerTreeView.Get(), &SOutlinerTreeView::MoveSelectionTo, FName()));
+				FSlateIcon(FEditorStyle::GetStyleSetName(), "SceneOutliner.MoveToRoot"), FExecuteAction::CreateSP(this, &SSceneOutliner::MoveSelectionTo, FName()));
 		}
 
 		// Existing folders - it would be good to filter out the existing selected folders from this mini-outliner, but that's currently not possible
@@ -1575,11 +1612,71 @@ namespace SceneOutliner
 		MenuBuilder.EndSection();
 	}
 
+	FName SSceneOutliner::MoveFolderTo(TSharedRef<TOutlinerFolderTreeItem> Folder, FName NewParent)
+	{
+		FName NewPath;
+		if (NewParent.IsNone())
+		{
+			NewPath = Folder->LeafName;
+		}
+		else
+		{
+			NewPath = FName(*(NewParent.ToString() / Folder->LeafName.ToString()));
+		}
+
+		if (FActorFolders::Get().RenameFolderInWorld(*RepresentingWorld, Folder->Path, NewPath))
+		{
+			return NewPath;
+		}
+
+		return FName();
+	}
+
 	void SSceneOutliner::MoveSelectionTo(TSharedRef<TOutlinerTreeItem> NewParent)
 	{
 		if (NewParent->Type == TOutlinerTreeItem::Folder)
 		{
-			OutlinerTreeView->MoveSelectionTo(StaticCastSharedRef<TOutlinerFolderTreeItem>(NewParent)->Path);
+			MoveSelectionTo(StaticCastSharedRef<TOutlinerFolderTreeItem>(NewParent)->Path);
+		}
+	}
+
+	void SSceneOutliner::MoveSelectionTo(FName NewParent)
+	{
+		FSlateApplication::Get().DismissAllMenus();
+
+		if (!OutlinerTreeView->ValidateMoveSelectionTo(NewParent))
+		{
+			return;
+		}
+
+		auto SelectedItems = GetSelectedItems();
+
+		const FScopedTransaction Transaction( LOCTEXT("MoveOutlinerItems", "Move Scene Outliner Items") );
+		for (const auto& Item : SelectedItems)
+		{
+			if (Item->Type == TOutlinerTreeItem::Folder)
+			{
+				auto Folder = StaticCastSharedPtr<TOutlinerFolderTreeItem>(Item).ToSharedRef();
+				auto NewPath = MoveFolderTo(Folder, NewParent);
+				if (!NewPath.IsNone() && SelectedItems.Num() == 1)
+				{
+					NewItemActions.WhenCreated(NewPath, NewItemActuator::Select | NewItemActuator::ScrollIntoView);
+				}
+			}
+			else
+			{
+				auto ActorItem = StaticCastSharedPtr<TOutlinerActorTreeItem>(Item);
+				auto* Actor = ActorItem->Actor.Get();
+				if (Actor)
+				{
+					Actor->SetFolderPath(NewParent);
+
+					if (SelectedItems.Num() == 1)
+					{
+						NewItemActions.WhenCreated(ActorItem.ToSharedRef(), NewItemActuator::Select | NewItemActuator::ScrollIntoView);
+					}
+				}
+			}
 		}
 	}
 
@@ -1613,21 +1710,54 @@ namespace SceneOutliner
 			return FReply::Unhandled();
 		}
 
-		const FName NewFolderName = FActorFolders::Get().GetDefaultFolderNameForSelection(*RepresentingWorld);
-		FActorFolders::Get().CreateFolderContainingSelection(*RepresentingWorld, NewFolderName);
+		CreateFolder();
 
 		return FReply::Handled();
 	}
 
-	void SSceneOutliner::CreateFolder(FName ParentPath)
+	void SSceneOutliner::CreateFolder()
 	{
 		if (!RepresentingWorld)
 		{
 			return;
 		}
 
-		const FName NewFolderName = FActorFolders::Get().GetDefaultFolderName(*RepresentingWorld, ParentPath);
+		const FScopedTransaction Transaction(LOCTEXT("UndoAction_CreateFolder", "Create Folder"));
+
+		const FName NewFolderName = FActorFolders::Get().GetDefaultFolderNameForSelection(*RepresentingWorld);
+		FActorFolders::Get().CreateFolderContainingSelection(*RepresentingWorld, NewFolderName);
+
+		// Clear the current actor selection.
+		GEditor->SelectNone(false, true, true);
+
+		// Move any selected folders into the new folder name
+		for (auto& Item : GetSelectedItems())
+		{
+			if (Item->Type == TOutlinerTreeItem::Folder)
+			{
+				auto Folder = StaticCastSharedPtr<TOutlinerFolderTreeItem>(Item).ToSharedRef();
+				MoveFolderTo(Folder, NewFolderName);
+			}
+		}
+
+		// At this point the new folder will be in our newly added list, so select it and open a rename when it gets refreshed
+		NewItemActions.WhenCreated(NewFolderName, NewItemActuator::Select | NewItemActuator::Rename);
+	}
+	
+	void SSceneOutliner::CreateSubFolder(FName Parent)
+	{
+		if (!RepresentingWorld)
+		{
+			return;
+		}
+
+		const FScopedTransaction Transaction(LOCTEXT("UndoAction_CreateFolder", "Create Folder"));
+
+		const FName NewFolderName = FActorFolders::Get().GetDefaultFolderName(*RepresentingWorld, Parent);
 		FActorFolders::Get().CreateFolder(*RepresentingWorld, NewFolderName);
+
+		// At this point the new folder will be in our newly added list, so select it and open a rename when it gets refreshed
+		NewItemActions.WhenCreated(NewFolderName, NewItemActuator::Select | NewItemActuator::Rename);
 	}
 
 	void SSceneOutliner::OnBroadcastFolderCreate(const UWorld& InWorld, FName NewPath)
@@ -2692,7 +2822,69 @@ namespace SceneOutliner
 			CustomColumn->SortItems(RootTreeItems, SortMode);
 		}
 	}
+	
+	void NewItemActuator::Empty()
+	{
+		PendingFolderActions.Empty();
+		PendingActorActions.Empty();
+	}
+
+	void NewItemActuator::WhenCreated(FName FolderPath, uint8 InActionMask)
+	{
+		const FolderAction NewAction = { InActionMask, FolderPath };
+		PendingFolderActions.Add(NewAction);
+	}
+
+	void NewItemActuator::WhenCreated(FOutlinerTreeItemRef Item, uint8 InActionMask)
+	{
+		if (Item->Type == TOutlinerTreeItem::Actor)
+		{
+			const ActorAction NewAction = { InActionMask, StaticCastSharedRef<TOutlinerActorTreeItem>(Item) };
+			PendingActorActions.Add(NewAction);
+		}
+		else
+		{
+			WhenCreated(StaticCastSharedRef<TOutlinerFolderTreeItem>(Item)->Path, InActionMask);
+		}
+	}
+
+	void NewItemActuator::ItemHasBeenCreated(TSharedRef<SOutlinerTreeView> OutlinerTreeView, TSharedRef<TOutlinerActorTreeItem> NewItem)
+	{
+		for (int32 Index = 0; Index < PendingActorActions.Num(); )
+		{
+			auto PinnedPendingItem = PendingActorActions[Index].ActorItem.Pin();
+			if (!PinnedPendingItem.IsValid())
+			{
+				PendingActorActions.RemoveAt(Index);
+			}
+			else if(PinnedPendingItem == NewItem)
+			{
+				PerformActionOnItem(OutlinerTreeView, NewItem, PendingActorActions[Index].Actions);
 
+				PendingActorActions.RemoveAt(Index);
+			}
+			else
+			{
+				++Index;
+			}
+		}
+	}
+
+	void NewItemActuator::ItemHasBeenCreated(TSharedRef<SOutlinerTreeView> OutlinerTreeView, TSharedRef<TOutlinerFolderTreeItem> NewItem)
+	{
+		for (int32 Index = 0; Index < PendingFolderActions.Num(); )
+		{
+			if (PendingFolderActions[Index].FolderPath == NewItem->Path)
+			{
+				PerformActionOnItem(OutlinerTreeView, NewItem, PendingFolderActions[Index].Actions);
+				PendingFolderActions.RemoveAt(Index);
+			}
+			else
+			{
+				++Index;
+			}
+		}
+	}
 
 }	// namespace SceneOutliner
 
diff --git a/Engine/Source/Editor/SceneOutliner/Private/SSceneOutliner.h b/Engine/Source/Editor/SceneOutliner/Private/SSceneOutliner.h
--- a/Engine/Source/Editor/SceneOutliner/Private/SSceneOutliner.h
+++ b/Engine/Source/Editor/SceneOutliner/Private/SSceneOutliner.h
@@ -18,6 +18,47 @@ namespace SceneOutliner
 
 	FText GetLabelForItem( const TSharedRef<TOutlinerTreeItem> TreeItem );
 
+	/** An array of actions to apply to newly added items in the scene outliner */
+	struct NewItemActuator
+	{
+		enum ActionType
+		{
+			/** Select the item when it is created */
+			Select			= 1 << 0,
+			/** Scroll the item into view when it is created */
+			ScrollIntoView	= 1 << 1,
+			/** Interactively rename the item when it is created (implies the above) */
+			Rename			= 1 << 2,
+		};
+
+		/** Empty this actuator (called after every refresh of the tree) */
+		void Empty();
+
+		/** Specify an action to perform on the specified actor or folder when it is created (see ActionType enum) */
+		void WhenCreated(FName FolderPath, uint8 InActionMask);
+		void WhenCreated(FOutlinerTreeItemRef, uint8 InActionMask);
+
+		/** Called by the scene outliner when an item has been added to the tree */
+		void ItemHasBeenCreated(TSharedRef<SOutlinerTreeView> OutlinerTreeView, TSharedRef<TOutlinerActorTreeItem> NewItem);
+		void ItemHasBeenCreated(TSharedRef<SOutlinerTreeView> OutlinerTreeView, TSharedRef<TOutlinerFolderTreeItem> NewItem);
+
+	private:
+
+		struct ActorAction
+		{	
+			uint8 Actions;
+			TWeakPtr<TOutlinerActorTreeItem> ActorItem;
+		};
+		struct FolderAction
+		{	
+			uint8 Actions;
+			FName FolderPath;
+		};
+
+		TArray<ActorAction>		PendingActorActions;
+		TArray<FolderAction>	PendingFolderActions;
+	};
+
 	/**
 	 * Scene Outliner widget
 	 */
@@ -28,7 +69,6 @@ namespace SceneOutliner
 
 		SLATE_BEGIN_ARGS( SSceneOutliner ){}
 
-			SLATE_ARGUMENT( FOnContextMenuOpening, MakeContextMenuWidgetDelegate );
 			SLATE_ARGUMENT( FOnSceneOutlinerItemPicked, OnItemPickedDelegate )
 
 		SLATE_END_ARGS()
@@ -76,6 +116,9 @@ namespace SceneOutliner
 		/** Find a folder by its path */
 		const TSharedPtr<TOutlinerFolderTreeItem> FindFolderByPath(FName Path) const;
 
+		/** Move the specified folder to a new parent path, returning the new full path */
+		FName MoveFolderTo(TSharedRef<TOutlinerFolderTreeItem> Folder, FName NewParent);
+
 	private:
 
 		/** Populates our data set */
@@ -171,21 +214,24 @@ namespace SceneOutliner
 		/** Open a context menu for this scene outliner */
 		TSharedPtr<SWidget> OnOpenContextMenu() const;
 
-		/** Build a context menu for right-clicking a folder */
-		TSharedPtr<SWidget> BuildFolderContextMenu() const;
+		/** Build a context menu for right-clicking an item in the tree */
+		TSharedPtr<SWidget> BuildDefaultContextMenu() const;
 		void FillFoldersSubMenu(FMenuBuilder& MenuBuilder) const;
 
 		/** Move the selected items to the specified parent */
 		void MoveSelectionTo(TSharedRef<TOutlinerTreeItem> NewParent);
 
-		/** Move the selected items to the specified parent */
+		/** Moves the current selection to the specified folder path */
 		void MoveSelectionTo(FName NewParent);
 
 		/** Called when the user has clicked the button to add a new folder */
 		FReply OnCreateFolderClicked();
 
 		/** Create a new folder under the specified parent name (NAME_None for root) */
-		void CreateFolder(FName ParentPath);
+		void CreateFolder();
+
+		/** Create a sub folder inside the specified parent */
+		void CreateSubFolder(FName Parent);
 
 		/** Initiate a rename operation for the specified folder */
 		void RenameFolder(TSharedRef<TOutlinerFolderTreeItem> Folder);
@@ -415,6 +461,9 @@ namespace SceneOutliner
 		/** true if the Scene Outliner should do a full refresh. */
 		bool bFullRefresh;
 
+		/** An instance of a class responsible for operating on newly added tree items */
+		NewItemActuator NewItemActions;
+
 		/** Timer for PIE/SIE mode to sort the outliner. */
 		float SortOutlinerTimer;
 
diff --git a/Engine/Source/Editor/SceneOutliner/Private/SceneOutlinerModule.cpp b/Engine/Source/Editor/SceneOutliner/Private/SceneOutlinerModule.cpp
--- a/Engine/Source/Editor/SceneOutliner/Private/SceneOutlinerModule.cpp
+++ b/Engine/Source/Editor/SceneOutliner/Private/SceneOutlinerModule.cpp
@@ -31,16 +31,15 @@ namespace SceneOutliner
 	}
 }
 
-TSharedRef< ISceneOutliner > FSceneOutlinerModule::CreateSceneOutliner( const FSceneOutlinerInitializationOptions& InitOptions, const FOnContextMenuOpening& MakeContextMenuWidgetDelegate, const FOnActorPicked& OnActorPickedDelegate ) const
+TSharedRef< ISceneOutliner > FSceneOutlinerModule::CreateSceneOutliner( const FSceneOutlinerInitializationOptions& InitOptions, const FOnActorPicked& OnActorPickedDelegate ) const
 {
 	auto OnItemPicked = FOnSceneOutlinerItemPicked::CreateStatic( &SceneOutliner::OnSceneOutlinerItemClicked, OnActorPickedDelegate );
-	return CreateSceneOutliner(InitOptions, MakeContextMenuWidgetDelegate, OnItemPicked);
+	return CreateSceneOutliner(InitOptions, OnItemPicked);
 }
 
-TSharedRef< ISceneOutliner > FSceneOutlinerModule::CreateSceneOutliner( const FSceneOutlinerInitializationOptions& InitOptions, const FOnContextMenuOpening& MakeContextMenuWidgetDelegate, const FOnSceneOutlinerItemPicked& OnItemPickedDelegate ) const
+TSharedRef< ISceneOutliner > FSceneOutlinerModule::CreateSceneOutliner( const FSceneOutlinerInitializationOptions& InitOptions, const FOnSceneOutlinerItemPicked& OnItemPickedDelegate ) const
 {
 	return SNew( SceneOutliner::SSceneOutliner, InitOptions )
 		.IsEnabled( FSlateApplication::Get().GetNormalExecutionAttribute() )
-		.MakeContextMenuWidgetDelegate( MakeContextMenuWidgetDelegate )
 		.OnItemPickedDelegate( OnItemPickedDelegate );
 }
diff --git a/Engine/Source/Editor/SceneOutliner/Private/SceneOutlinerTreeItems.cpp b/Engine/Source/Editor/SceneOutliner/Private/SceneOutlinerTreeItems.cpp
--- a/Engine/Source/Editor/SceneOutliner/Private/SceneOutlinerTreeItems.cpp
+++ b/Engine/Source/Editor/SceneOutliner/Private/SceneOutlinerTreeItems.cpp
@@ -55,43 +55,6 @@ namespace SceneOutliner
 		}
 	}
 
-	void TOutlinerFolderTreeItem::Rename(FName NewLabel)
-	{
-		FName NewPath = GetParentPath();
-		if (NewPath.IsNone())
-		{
-			NewPath = FName(*NewLabel.ToString());
-		}
-		else
-		{
-			NewPath = FName(*(NewPath.ToString() / NewLabel.ToString()));
-		}
-
-		if (NewPath != Path && !FActorFolders::PathIsChildOf(NewPath.ToString(), Path.ToString()))
-		{
-			const FScopedTransaction Transaction(LOCTEXT("UndoAction_RenameFolder", "Rename Folder"));
-			FActorFolders::Get().RenameFolderInWorld(*GWorld, Path, NewPath);
-		}
-	}
-
-	void TOutlinerFolderTreeItem::MoveTo(FName NewPath)
-	{
-		if (NewPath.IsNone())
-		{
-			NewPath = LeafName;
-		}
-		else
-		{
-			NewPath = FName(*(NewPath.ToString() / LeafName.ToString()));
-		}
-
-		if (NewPath != Path && !FActorFolders::PathIsChildOf(NewPath.ToString(), Path.ToString()))
-		{
-			const FScopedTransaction Transaction(LOCTEXT("UndoAction_MoveFolder", "Move Folder"));
-			FActorFolders::Get().RenameFolderInWorld(*GWorld, Path, NewPath);
-		}
-	}
-
 	TSharedRef<FDecoratedDragDropOp> CreateDragDropOperation(TArray<TSharedPtr<TOutlinerTreeItem>>& InTreeItems)
 	{
 		TArray<TWeakObjectPtr<AActor>> Actors;
diff --git a/Engine/Source/Editor/SceneOutliner/Public/SceneOutlinerInitializationOptions.h b/Engine/Source/Editor/SceneOutliner/Public/SceneOutlinerInitializationOptions.h
--- a/Engine/Source/Editor/SceneOutliner/Public/SceneOutlinerInitializationOptions.h
+++ b/Engine/Source/Editor/SceneOutliner/Public/SceneOutlinerInitializationOptions.h
@@ -104,6 +104,13 @@ public:
 
 	/** True to only show folders in this outliner */
 	bool bOnlyShowFolders;
+
+	/** Override default context menu handling */
+	FOnContextMenuOpening ContextMenuOverride;
+
+	/** Extend default context menu handling */
+	TSharedPtr<FExtender> DefaultMenuExtender;
+
 public:
 
 	/** Constructor */
diff --git a/Engine/Source/Editor/SceneOutliner/Public/SceneOutlinerModule.h b/Engine/Source/Editor/SceneOutliner/Public/SceneOutlinerModule.h
--- a/Engine/Source/Editor/SceneOutliner/Public/SceneOutlinerModule.h
+++ b/Engine/Source/Editor/SceneOutliner/Public/SceneOutlinerModule.h
@@ -35,7 +35,6 @@ public:
 	 */
 	virtual TSharedRef< class ISceneOutliner > CreateSceneOutliner(
 		const FSceneOutlinerInitializationOptions& InitOptions,
-		const FOnContextMenuOpening& MakeContextMenuWidgetDelegate,
 		const FOnActorPicked& OnActorPickedDelegate ) const;
 
 	/**
@@ -50,7 +49,6 @@ public:
 
 	virtual TSharedRef< class ISceneOutliner > CreateSceneOutliner(
 		const FSceneOutlinerInitializationOptions& InitOptions,
-		const FOnContextMenuOpening& MakeContextMenuWidgetDelegate,
 		const FOnSceneOutlinerItemPicked& OnItemPickedDelegate ) const;
 };
 
diff --git a/Engine/Source/Editor/SceneOutliner/Public/SceneOutlinerTreeItems.h b/Engine/Source/Editor/SceneOutliner/Public/SceneOutlinerTreeItems.h
--- a/Engine/Source/Editor/SceneOutliner/Public/SceneOutlinerTreeItems.h
+++ b/Engine/Source/Editor/SceneOutliner/Public/SceneOutlinerTreeItems.h
@@ -114,20 +114,6 @@ struct SCENEOUTLINER_API TOutlinerFolderTreeItem : TOutlinerTreeItem
 		}
 	}
 
-	/** Move this folder to a new parent path.
-	 *	Broadcasts notifications to ensure all scene outliners are kept up to date.
-	 *
-	 *	@param	NewParentPath	The new parent path to put this folder in (NAME_None for root)
-	 */
-	void MoveTo(FName NewParentPath);
-
-	/** Rename this folder leaf name.
-	 *	Broadcasts notifications to ensure all scene outliners are kept up to date.
-	 *
-	 *	@param	NewLabel	The new label for this folder
-	 */
-	void Rename(FName NewLabel);
-
 	/** Get the parent path for this folder */
 	FName GetParentPath() const
 	{
diff --git a/Engine/Source/Editor/UnrealEd/Private/EditorActorFolders.cpp b/Engine/Source/Editor/UnrealEd/Private/EditorActorFolders.cpp
--- a/Engine/Source/Editor/UnrealEd/Private/EditorActorFolders.cpp
+++ b/Engine/Source/Editor/UnrealEd/Private/EditorActorFolders.cpp
@@ -276,11 +276,11 @@ void FActorFolders::DeleteFolder(const UWorld& InWorld, FName FolderToDelete)
 	}
 }
 
-void FActorFolders::RenameFolderInWorld(UWorld& World, FName OldPath, FName NewPath)
+bool FActorFolders::RenameFolderInWorld(UWorld& World, FName OldPath, FName NewPath)
 {
-	if (OldPath.IsNone())
+	if (OldPath.IsNone() || OldPath == NewPath || PathIsChildOf(NewPath.ToString(), OldPath.ToString()))
 	{
-		return;
+		return false;
 	}
 
 	const FScopedTransaction Transaction(LOCTEXT("UndoAction_RenameFolder", "Rename Folder"));
@@ -334,6 +334,8 @@ void FActorFolders::RenameFolderInWorld(UWorld& World, FName OldPath, FName NewP
 		FoldersInWorld.Paths.Remove(Path);
 		OnFolderDelete.Broadcast(World, Path);
 	}
+
+	return RenamedFolders.Num() != 0;
 }
 
 #undef LOCTEXT_NAMESPACE
\ No newline at end of file
diff --git a/Engine/Source/Editor/UnrealEd/Public/EditorActorFolders.h b/Engine/Source/Editor/UnrealEd/Public/EditorActorFolders.h
--- a/Engine/Source/Editor/UnrealEd/Public/EditorActorFolders.h
+++ b/Engine/Source/Editor/UnrealEd/Public/EditorActorFolders.h
@@ -64,7 +64,7 @@ struct UNREALED_API FActorFolders : public FGCObject
 	void DeleteFolder(const UWorld& InWorld, FName FolderToDelete);
 
 	/** Rename the specified path to a new name */
-	void RenameFolderInWorld(UWorld& World, FName OldPath, FName NewPath);
+	bool RenameFolderInWorld(UWorld& World, FName OldPath, FName NewPath);
 
 private:
 
