#ttp 328753 - EDITOR: Add "Select Actors" option to Scene Outliner folder context menu
#summary Added ability to select all children and descendants of a folder selection in the scene outliner
#proj Editor
#branch UE4

diff --git a/Engine/Source/Editor/SceneOutliner/Private/SSceneOutliner.cpp b/Engine/Source/Editor/SceneOutliner/Private/SSceneOutliner.cpp
--- a/Engine/Source/Editor/SceneOutliner/Private/SSceneOutliner.cpp
+++ b/Engine/Source/Editor/SceneOutliner/Private/SSceneOutliner.cpp
@@ -26,6 +26,30 @@ DEFINE_LOG_CATEGORY_STATIC(LogSceneOutliner, Log, All);
 
 namespace SceneOutliner
 {
+
+	/** Recursively set the selection of all actor descendants of the specified folder */
+	void SelectDescedantActors(TSharedRef<TOutlinerFolderTreeItem> FolderItem, bool bSelected)
+	{
+		for (const auto& WeakChild : FolderItem->ChildItems)
+		{
+			auto ChildItem = WeakChild.Pin();
+			if (ChildItem.IsValid())
+			{
+				if (ChildItem->Type == TOutlinerTreeItem::Actor)
+				{
+					if (AActor* Actor = StaticCastSharedPtr<TOutlinerActorTreeItem>(ChildItem)->Actor.Get())
+					{
+						GEditor->SelectActor(Actor, true, /*bNotify=*/false);
+					}
+				}
+				else
+				{
+					SelectDescedantActors(StaticCastSharedPtr<TOutlinerFolderTreeItem>(ChildItem).ToSharedRef(), bSelected);
+				}
+			}
+		}
+	}
+
 	/** Gets the label for a tree item */
 	FText GetLabelForItem( const TSharedRef<TOutlinerTreeItem> TreeItem )
 	{
@@ -1515,7 +1539,19 @@ namespace SceneOutliner
 		OutlinerTreeView->SetItemExpansion(*ActorToTreeItemMap.Find(Actor), true);
 	}
 
-	/** Open a context menu for this scene outliner */
+	TArray<TSharedRef<TOutlinerFolderTreeItem>> SSceneOutliner::GetSelectedFolders() const
+	{
+		TArray<TSharedRef<TOutlinerFolderTreeItem>> Folders;
+		for (auto& Item : OutlinerTreeView->GetSelectedItems())
+		{
+			if (Item->Type == TOutlinerTreeItem::Folder)
+			{
+				Folders.Add(StaticCastSharedPtr<TOutlinerFolderTreeItem>(Item).ToSharedRef());
+			}
+		}
+		return Folders;
+	}
+
 	TSharedPtr<SWidget> SSceneOutliner::OnOpenContextMenu() const
 	{
 		TArray<AActor*> SelectedActors;
@@ -1534,17 +1570,19 @@ namespace SceneOutliner
 		return TSharedPtr<SWidget>();
 	}
 
-	TSharedPtr<SWidget> SSceneOutliner::BuildDefaultContextMenu()  const
+	TSharedPtr<SWidget> SSceneOutliner::BuildDefaultContextMenu() const
 	{
 		// Build up the menu
 		const bool bCloseAfterSelection = true;
 		FMenuBuilder MenuBuilder(bCloseAfterSelection, TSharedPtr<FUICommandList>(), InitOptions.DefaultMenuExtender);
 
-		TSharedPtr<TOutlinerFolderTreeItem> ParentFolder;
 		auto SelectedItems = OutlinerTreeView->GetSelectedItems();
-		if (SelectedItems.Num() == 1 && SelectedItems[0]->Type == TOutlinerTreeItem::Folder)
+		auto SelectedFolders = GetSelectedFolders();
+
+		TSharedPtr<TOutlinerFolderTreeItem> ParentFolder;
+		if (SelectedFolders.Num() == 1)
 		{
-			ParentFolder = StaticCastSharedPtr<TOutlinerFolderTreeItem>(SelectedItems[0]);
+			ParentFolder = SelectedFolders[0];
 		}
 
 		const FSlateIcon NewFolderIcon(FEditorStyle::GetStyleSetName(), "SceneOutliner.NewFolderIcon");
@@ -1559,7 +1597,11 @@ namespace SceneOutliner
 		}
 		MenuBuilder.EndSection();
 
-		if (SelectedItems.Num() != 0)
+		if (SelectedItems.Num() == 0)
+		{
+			MenuBuilder.AddMenuEntry(LOCTEXT("CreateFolder", "Create Folder"), FText(), NewFolderIcon, FUIAction(FExecuteAction::CreateSP(this, &SSceneOutliner::CreateFolder)));
+		}
+		else
 		{
 			MenuBuilder.BeginSection("SelectionSection", LOCTEXT("OutlinerSelection", "Current Outliner Selection"));
 			{
@@ -1567,12 +1609,17 @@ namespace SceneOutliner
 					LOCTEXT("MoveActorsToFolder", "Move To Folder"),
 					LOCTEXT("MoveActorsToFolder_Tooltip", "Move selection to a folder"),
 					FNewMenuDelegate::CreateSP(this,  &SSceneOutliner::FillFoldersSubMenu));
+
+				// If we've only got folders selected, show the selection sub menu
+				if (SelectedFolders.Num() == SelectedItems.Num())
+				{
+					MenuBuilder.AddSubMenu(
+						LOCTEXT("SelectSubmenu", "Select"),
+						LOCTEXT("SelectSubmenu_Tooltip", "Select the contents of the current selection"),
+						FNewMenuDelegate::CreateSP(this, &SSceneOutliner::FillSelectionSubMenu));
+				}
 			}
 		}
-		else
-		{
-			MenuBuilder.AddMenuEntry(LOCTEXT("CreateFolder", "Create Folder"), FText(), NewFolderIcon, FUIAction(FExecuteAction::CreateSP(this, &SSceneOutliner::CreateFolder)));
-		}
 
 		return MenuBuilder.MakeWidget();
 	}
@@ -1612,6 +1659,66 @@ namespace SceneOutliner
 		MenuBuilder.EndSection();
 	}
 
+	void SSceneOutliner::FillSelectionSubMenu(FMenuBuilder& MenuBuilder) const
+	{
+		MenuBuilder.AddMenuEntry(
+			LOCTEXT( "AddChildrenToSelection", "Immediate Children" ),
+			LOCTEXT( "AddChildrenToSelection_ToolTip", "Select all immediate children of the selected folders" ),
+			FSlateIcon(),
+			FExecuteAction::CreateSP(this, &SSceneOutliner::SelectFoldersImmediateChildren));
+		MenuBuilder.AddMenuEntry(
+			LOCTEXT( "AddDescendantsToSelection", "All Descendants" ),
+			LOCTEXT( "AddDescendantsToSelection_ToolTip", "Select all descendants of the selected folders" ),
+			FSlateIcon(),
+			FExecuteAction::CreateSP(this, &SSceneOutliner::SelectFoldersDescendants));
+	}
+
+	void SSceneOutliner::SelectFoldersImmediateChildren() const
+	{
+		auto SelectedFolders = GetSelectedFolders();
+		if (SelectedFolders.Num())
+		{
+			// We'll batch selection changes instead by using BeginBatchSelectOperation()
+			GEditor->GetSelectedActors()->BeginBatchSelectOperation();
+
+			for (const auto& Folder : SelectedFolders)
+			{
+				for (const auto& WeakChild : Folder->ChildItems)
+				{
+					auto ChildItem = WeakChild.Pin();
+					if (ChildItem.IsValid() && ChildItem->Type == TOutlinerTreeItem::Actor)
+					{
+						if (AActor* Actor = StaticCastSharedPtr<TOutlinerActorTreeItem>(ChildItem)->Actor.Get())
+						{
+							GEditor->SelectActor(Actor, true, /*bNotify=*/false);
+						}
+					}
+				}
+			}
+
+			GEditor->GetSelectedActors()->EndBatchSelectOperation();
+			GEditor->NoteSelectionChange();
+		}
+	}
+
+	void SSceneOutliner::SelectFoldersDescendants() const
+	{
+		auto SelectedFolders = GetSelectedFolders();
+		if (SelectedFolders.Num())
+		{
+			// We'll batch selection changes instead by using BeginBatchSelectOperation()
+			GEditor->GetSelectedActors()->BeginBatchSelectOperation();
+
+			for (const auto& Folder : SelectedFolders)
+			{
+				SelectDescedantActors(Folder, true);
+			}
+
+			GEditor->GetSelectedActors()->EndBatchSelectOperation();
+			GEditor->NoteSelectionChange();
+		}
+	}
+
 	FName SSceneOutliner::MoveFolderTo(TSharedRef<TOutlinerFolderTreeItem> Folder, FName NewParent)
 	{
 		FName NewPath;
diff --git a/Engine/Source/Editor/SceneOutliner/Private/SSceneOutliner.h b/Engine/Source/Editor/SceneOutliner/Private/SSceneOutliner.h
--- a/Engine/Source/Editor/SceneOutliner/Private/SSceneOutliner.h
+++ b/Engine/Source/Editor/SceneOutliner/Private/SSceneOutliner.h
@@ -217,6 +217,13 @@ namespace SceneOutliner
 		/** Build a context menu for right-clicking an item in the tree */
 		TSharedPtr<SWidget> BuildDefaultContextMenu() const;
 		void FillFoldersSubMenu(FMenuBuilder& MenuBuilder) const;
+		void FillSelectionSubMenu(FMenuBuilder& MenuBuilder) const;
+
+		/** Select the immediate children of the currently selected folders */
+		void SelectFoldersImmediateChildren() const;
+
+		/** Called to select all the descendants of the currently selected folders */
+		void SelectFoldersDescendants() const;
 
 		/** Move the selected items to the specified parent */
 		void MoveSelectionTo(TSharedRef<TOutlinerTreeItem> NewParent);
@@ -333,6 +340,8 @@ namespace SceneOutliner
 		/** @return the foreground color for the view button */
 		FSlateColor GetViewButtonForegroundColor() const;
 
+		/** Get an array of selected folders in this scene outliner */
+		TArray<TSharedRef<TOutlinerFolderTreeItem>> GetSelectedFolders() const;
 
 		/** FILTERS */
 		/** @return whether we are displaying only selected Actors */
