#ttp 331033 - LIVE: EDITOR: SCENE OUTLINER: VISIBILITY: Unable to toggle visiblity of multiple actors by using the Toggle Visibility Eye Icon.
#proj Editor
#branch UE4
#summary The 'eye' visibility icon now acts on the current selection, if the row is selected

diff --git a/Engine/Source/Editor/SceneOutliner/Private/SSceneOutliner.cpp b/Engine/Source/Editor/SceneOutliner/Private/SSceneOutliner.cpp
--- a/Engine/Source/Editor/SceneOutliner/Private/SSceneOutliner.cpp
+++ b/Engine/Source/Editor/SceneOutliner/Private/SSceneOutliner.cpp
@@ -189,7 +189,8 @@ namespace SceneOutliner
 		if (InInitOptions.Mode == ESceneOutlinerMode::ActorBrowsing)
 		{
 			// Set up the gutter if we're in actor browsing mode
-			Gutter = MakeShareable(new FSceneOutlinerGutter());
+
+			Gutter = MakeShareable(new FSceneOutlinerGutter(FOnSetItemVisibility::CreateSP(this, &SSceneOutliner::OnSetItemVisibility)));
 			HeaderRowWidget->AddColumn(
 				Gutter->ConstructHeaderRowColumn()
 				.FixedWidth(16.f)
@@ -2547,6 +2548,24 @@ namespace SceneOutliner
 		}
 	}
 
+	void SSceneOutliner::OnSetItemVisibility(FOutlinerTreeItemRef Item, bool bIsVisible)
+	{
+		// We operate on all the selected items if the specified item is selected
+		if (OutlinerTreeView->IsItemSelected(Item))
+		{
+			for (auto& TreeItem : OutlinerTreeView->GetSelectedItems())
+			{
+				TreeItem->SetIsVisible(bIsVisible);
+			}
+		}
+		else
+		{
+			Item->SetIsVisible(bIsVisible);
+		}
+
+		GEditor->RedrawAllViewports();
+	}
+
 	void SSceneOutliner::OnFilterTextChanged( const FText& InFilterText )
 	{
 		SearchBoxFilter->SetRawFilterText( InFilterText );
diff --git a/Engine/Source/Editor/SceneOutliner/Private/SSceneOutliner.h b/Engine/Source/Editor/SceneOutliner/Private/SSceneOutliner.h
--- a/Engine/Source/Editor/SceneOutliner/Private/SSceneOutliner.h
+++ b/Engine/Source/Editor/SceneOutliner/Private/SSceneOutliner.h
@@ -136,6 +136,9 @@ namespace SceneOutliner
 		/** Check whether we should be showing folders or not in this scene outliner */
 		bool ShouldShowFolders() const;
 
+		/** Synchronize (and scroll into view) the current selection to the specified actor item, if the corresponding actor is selected */
+		void SynchronizeSelectionTo(TSharedRef<TOutlinerActorTreeItem> ActorItem);
+
 		/** Is actor class name text visible? */
 		EVisibility IsActorClassNameVisible() const;
 
@@ -217,6 +220,7 @@ namespace SceneOutliner
 		/** Build a context menu for right-clicking an item in the tree */
 		TSharedPtr<SWidget> BuildDefaultContextMenu() const;
 		void FillFoldersSubMenu(FMenuBuilder& MenuBuilder) const;
+		void AddMoveToFolderOutliner(FMenuBuilder& MenuBuilder) const;
 		void FillSelectionSubMenu(FMenuBuilder& MenuBuilder) const;
 
 		/** Select the immediate children of the currently selected folders */
@@ -288,6 +292,9 @@ namespace SceneOutliner
 		/** Called by the editable text control when a user presses enter or commits their text change */
 		void OnFilterTextCommitted( const FText& InFilterText, ETextCommit::Type CommitInfo );
 
+		/** Invoked by the scene outliner gutter when we need to set the specified item's visibility */
+		void OnSetItemVisibility(FOutlinerTreeItemRef Item, bool bIsVisible);
+
 		/**
 		 * Called by the filter button to get the image to display in the button
 		 *
diff --git a/Engine/Source/Editor/SceneOutliner/Private/SceneOutlinerGutter.cpp b/Engine/Source/Editor/SceneOutliner/Private/SceneOutlinerGutter.cpp
--- a/Engine/Source/Editor/SceneOutliner/Private/SceneOutlinerGutter.cpp
+++ b/Engine/Source/Editor/SceneOutliner/Private/SceneOutlinerGutter.cpp
@@ -43,12 +43,14 @@ class SVisibilityWidget : public SImage
 public:
 	SLATE_BEGIN_ARGS(SVisibilityWidget){}
 		SLATE_ARGUMENT(TWeakPtr<SceneOutliner::TOutlinerTreeItem>, TreeItem);
+		SLATE_EVENT(FOnSetItemVisibility, OnSetItemVisibility);
 	SLATE_END_ARGS()
 
 	/** Construct this widget */
 	void Construct(const FArguments& InArgs)
 	{
 		WeakTreeItem = InArgs._TreeItem;
+		OnSetItemVisibility = InArgs._OnSetItemVisibility;
 
 		SImage::Construct(
 			SImage::FArguments()
@@ -152,24 +154,26 @@ private:
 	{
 		if (IsVisible() != bVisible)
 		{
-			auto TreeItem = WeakTreeItem.Pin();
-			if (TreeItem.IsValid())
+			auto ClickedOnItem = WeakTreeItem.Pin();
+			if (ClickedOnItem.IsValid())
 			{
-				TreeItem->SetIsVisible(bVisible);
+				OnSetItemVisibility.ExecuteIfBound(ClickedOnItem.ToSharedRef(), bVisible);
 			}
-
-			GEditor->RedrawAllViewports();
 		}
 	}
 
 	/** Tree item that we relate to */
 	TWeakPtr<SceneOutliner::TOutlinerTreeItem> WeakTreeItem;
 
+	/** A delegate to execute when we need to set the visibility of an item */
+	FOnSetItemVisibility OnSetItemVisibility;
+
 	/** Scoped undo transaction */
 	TUniquePtr<FScopedTransaction> UndoTransaction;
 };
 
-FSceneOutlinerGutter::FSceneOutlinerGutter()
+FSceneOutlinerGutter::FSceneOutlinerGutter(FOnSetItemVisibility InOnSetItemVisibility)
+	: OnSetItemVisibility(InOnSetItemVisibility)
 {
 
 }
@@ -192,6 +196,7 @@ const TSharedRef<SWidget> FSceneOutlinerGutter::ConstructRowWidget(const TShared
 		.VAlign(VAlign_Center)
 		[
 			SNew(SVisibilityWidget)
+			.OnSetItemVisibility(OnSetItemVisibility)
 			.TreeItem(TreeItem)
 		];
 }
diff --git a/Engine/Source/Editor/SceneOutliner/Private/SceneOutlinerGutter.h b/Engine/Source/Editor/SceneOutliner/Private/SceneOutlinerGutter.h
--- a/Engine/Source/Editor/SceneOutliner/Private/SceneOutlinerGutter.h
+++ b/Engine/Source/Editor/SceneOutliner/Private/SceneOutlinerGutter.h
@@ -6,6 +6,8 @@
 #include "ISceneOutlinerColumn.h"
 #include "SceneOutlinerInitializationOptions.h"
 
+DECLARE_DELEGATE_TwoParams(FOnSetItemVisibility, TSharedRef<SceneOutliner::TOutlinerTreeItem>, bool)
+
 /**
  * A gutter for the SceneOutliner which is capable of displaying a variety of Actor details
  */
@@ -15,7 +17,7 @@ class FSceneOutlinerGutter : public ISceneOutlinerColumn
 public:
 
 	/**	Constructor */
-	FSceneOutlinerGutter();
+	FSceneOutlinerGutter(FOnSetItemVisibility InOnSetItemVisibility);
 
 	// -----------------------------------------
 	// ISceneOutlinerColumn Implementation
@@ -34,4 +36,8 @@ public:
 	virtual void SortItems(TArray<TSharedPtr<SceneOutliner::TOutlinerTreeItem>>& RootItems, const EColumnSortMode::Type SortMode) const OVERRIDE;
 	// -----------------------------------------
 
+private:
+
+	/** A delegate to execute when we need to set the visibility of an item */
+	FOnSetItemVisibility OnSetItemVisibility;
 };
\ No newline at end of file
