#ttp 330692 - LIVE: EDITOR: Scene Outliner opens all parent trees when undo button pressed.
#branch UE4
#proj Editor.SceneOutliner, Runtime.Slate
#summary Fixed issue where undo button caused scene outliner to expand all parent trees.
#add Added STreeView::RefreshItemExpansion.  Its purpose is to force the given item, during linearization, to linearize its children too, even if that item is not itself expanded.  The reason you may wish to do this is to force the tree to recognise an expanded subtree which is hidden underneath a collapsed item, for example during repopulation.
#add Made STreeView::IsItemExpanded const-correct.
#add Added some helper functions in the SSceneOutliner namespace for determining whether an actor or folder is a parent of another actor or folder, somewhere in its hierarchy.
#add Added two new TSets to SSceneOutliner for storing any parent items which have been collapsed, prior to repopulating the tree.  Upon adding parent items, the expansion state is set according to whether the item was found in the corresponding TSet, hence preserving its state.  Additionally, if the item has children which are also parents, and are in an expanded state, STreeView::RefreshItemExpansion is called, so that the expansion state of any children is preserved upon opening a collapsed parent.
#change Deferred setting the expansion state of folders and parent items until they have all been created, and until all parents have been determined.

diff --git a/Engine/Source/Editor/SceneOutliner/Private/SSceneOutliner.cpp b/Engine/Source/Editor/SceneOutliner/Private/SSceneOutliner.cpp
--- a/Engine/Source/Editor/SceneOutliner/Private/SSceneOutliner.cpp
+++ b/Engine/Source/Editor/SceneOutliner/Private/SSceneOutliner.cpp
@@ -640,21 +640,28 @@ namespace SceneOutliner
 
 		bRepresentingPlayWorld = RepresentingWorld->WorldType == EWorldType::PIE;
 
+		// Get a collection of items and folders which were formerly collapsed
+		// TODO: this will actually include non-parent items too; ideally only add item if it has children.
+		FParentsExpansionState ExpansionStateInfo;
+		GetParentsExpansionState(ExpansionStateInfo);
+
 		bool bMadeAnySignificantChanges = false;
 		if(bFullRefresh)
 		{
 			// Clear the selection here - UpdateActorMapWithWorld will reconstruct it.
 			OutlinerTreeView->ClearSelection();
 
-			EmptyTreeItems();
-			UpdateActorMapWithWorld( RepresentingWorld );
+			FParentActorsSet ParentActors;
+			UpdateActorMapWithWorld(RepresentingWorld, ParentActors);
+
+			SetParentsExpansionState(ParentActors, ExpansionStateInfo);
 
 			bMadeAnySignificantChanges = true;
 			bFullRefresh = false;
 		}
 		else
 		{
-			TSet< AActor* > ParentActors;
+			FParentActorsSet ParentActors;
 
 			// Add actors that have been added to the world
 			for(int32 ItemIdx = 0; ItemIdx < AddedItemsList.Num(); ++ItemIdx)
@@ -753,11 +760,13 @@ namespace SceneOutliner
 			{
 				for( auto ActorIt = ParentActors.CreateConstIterator(); ActorIt; ++ActorIt )
 				{
-					AActor* Actor = *ActorIt;
+					const AActor* Actor = *ActorIt;
 
 					AddFilteredParentActorToTree(Actor);
 				}
 			}
+
+			SetParentsExpansionState(ParentActors, ExpansionStateInfo);
 		}
 		
 		AddedItemsList.Empty();
@@ -792,7 +801,7 @@ namespace SceneOutliner
 		RootTreeItems.Empty();
 	}
 
-	void SSceneOutliner::UpdateActorMapWithWorld( UWorld* World )
+	void SSceneOutliner::UpdateActorMapWithWorld(UWorld* World, FParentActorsSet& ParentActors)
 	{
 		if ( World == NULL )
 		{
@@ -801,8 +810,6 @@ namespace SceneOutliner
 
 		EmptyTreeItems();
 
-		TSet< AActor* > ParentActors;
-
 		// We only synchronize selection when in actor browsing mode
 		const bool bSyncSelection = ( InitOptions.Mode == ESceneOutlinerMode::ActorBrowsing );
 		FOutlinerTreeItemPtr LastSelectedItem;
@@ -919,7 +926,7 @@ namespace SceneOutliner
 		}
 	}
 
-	bool SSceneOutliner::AddItemToTree(TSharedRef<TOutlinerTreeItem> InItem, TSet<AActor*>& ParentActors)
+	bool SSceneOutliner::AddItemToTree(TSharedRef<TOutlinerTreeItem> InItem, FParentActorsSet& ParentActors)
 	{
 		if (InItem->Type == TOutlinerTreeItem::Actor)
 		{
@@ -931,7 +938,7 @@ namespace SceneOutliner
 		}
 	}
 
-	bool SSceneOutliner::AddActorToTree(TSharedRef<TOutlinerActorTreeItem> InActorItem, TSet< AActor* >& ParentActors)
+	bool SSceneOutliner::AddActorToTree(TSharedRef<TOutlinerActorTreeItem> InActorItem, FParentActorsSet& ParentActors)
 	{
 		if (!InActorItem->Actor.IsValid())
 		{
@@ -1040,15 +1047,12 @@ namespace SceneOutliner
 		// Add the item to the map, making it available for display in the tree when its parent is looking for its children.
 		FolderToTreeItemMap.Add(InFolderItem->Path, InFolderItem);
 
-		// Folders are always expanded by default
-		OutlinerTreeView->SetItemExpansion(InFolderItem, true);
-
 		NewItemActions.ItemHasBeenCreated(OutlinerTreeView.ToSharedRef(), InFolderItem);
 
 		return true;
 	}
 
-	void SSceneOutliner::AddFilteredParentActorToTree(AActor* InActor)
+	void SSceneOutliner::AddFilteredParentActorToTree(const AActor* InActor)
 	{
 		/* Parents that are filtered out (or otherwise found and added to the list) will not appear in the actor map at this point
 		   and should be added as requested. If they have no parent, they will be added to the root of the tree, otherwise they will
@@ -1060,7 +1064,7 @@ namespace SceneOutliner
 			ActorToTreeItemMap.Add( InActor, NewItem );
 
 			// Setup our root-level item list
-			const auto ParentActor = InActor->GetAttachParentActor();
+			const auto* ParentActor = InActor->GetAttachParentActor();
 			if( ParentActor == NULL )
 			{
 				RootTreeItems.Add( NewItem );
@@ -1078,11 +1082,159 @@ namespace SceneOutliner
 
 			NewItem->Flags.IsFilteredOut = true;
 		}
+	}
+
+	// Helper function which returns whether a folder is a child of another.
+	// Note that a folder is defined to not be a child of itself.
+	bool DoesFolderContainFolder(FName ParentFolder, FName ChildFolder)
+	{
+		FString Path = ChildFolder.ToString();
+		FString ParentFolderAsString = ParentFolder.ToString();
+
+		do 
+		{
+			Path = FPaths::GetPath(Path);
+
+		} while (!Path.IsEmpty() && Path != ParentFolderAsString);
+
+		return !Path.IsEmpty();
+	}
+
+	// Helper function which returns whether an actor is a child of a folder.
+	bool DoesFolderContainActor(FName ParentFolder, const AActor* ChildActor)
+	{
+		if (ChildActor != nullptr)
+		{
+			const FName ActorFolder = ChildActor->GetFolderPath();
+			if (!ActorFolder.IsNone())
+			{
+				return (ActorFolder == ParentFolder || DoesFolderContainFolder(ParentFolder, ActorFolder));
+			}
+		}
+
+		return false;
+	}
+
+	// Helper function which returns whether an actor is a child of another.
+	// Note that an actor is defined to not be a child of itself.
+	bool IsActorParentOfActor(const AActor* ParentActor, const AActor* ChildActor)
+	{
+		const AActor* Actor = ChildActor;
+
+		if (Actor != nullptr)
+		{
+			do
+			{
+				Actor = Actor->GetAttachParentActor();
+
+			} while (Actor != nullptr && Actor != ParentActor);
+		}
+
+		return (Actor != nullptr);
+	}
+
+	bool SSceneOutliner::HasExpandedChildren(const AActor* ParentActor, const FParentActorsSet& ParentActors, const FParentsExpansionState& ExpansionStateInfo) const
+	{
+		// Iterate through all parent actors, determining whether any of them are within the subtree of the
+		// parent actor, and if so, whether they are expanded.
+		for (const auto* ParentToTest : ParentActors)
+		{
+			if (ParentToTest != ParentActor)
+			{
+				if (IsActorParentOfActor(ParentActor, ParentToTest))
+				{
+					if (!ExpansionStateInfo.CollapsedItems.Contains(ParentToTest))
+					{
+						return true;
+					}
+				}
+			}
+		}
 
+		return false;
+	}
 
-		// Make sure all parent actors are expanded by default
-		// @todo Outliner: Really we should be remembering expansion state and trying to restore that
-		OutlinerTreeView->SetItemExpansion( ActorToTreeItemMap.FindChecked( InActor ), true );
+	bool SSceneOutliner::HasExpandedChildren(FName ParentFolder, const FParentActorsSet& ParentActors, const FParentsExpansionState& ExpansionStateInfo) const
+	{
+		// Iterate through all parent actors, determining whether any of them are within the subtree of the
+		// parent actor, and if so, whether they are expanded.
+		for (const auto* ParentToTest : ParentActors)
+		{
+			if (DoesFolderContainActor(ParentFolder, ParentToTest))
+			{
+				if (!ExpansionStateInfo.CollapsedItems.Contains(ParentToTest))
+				{
+					return true;
+				}
+			}
+		}
+
+		// If there were no expanded parent actors in this folder subtree, do the same check, this time with other
+		// folders.
+		for (const auto& Item : FolderToTreeItemMap)
+		{
+			if (Item.Key != ParentFolder)
+			{
+				if (DoesFolderContainFolder(ParentFolder, Item.Key))
+				{
+					if (!ExpansionStateInfo.CollapsedFolders.Contains(Item.Key))
+					{
+						return true;
+					}
+				}
+			}
+		}
+
+		return false;
+	}
+
+	void SSceneOutliner::GetParentsExpansionState(FParentsExpansionState& ExpansionStateInfo) const
+	{
+		for (const auto& Item : ActorToTreeItemMap)
+		{
+			if (!OutlinerTreeView->IsItemExpanded(Item.Value) && Item.Key.IsValid())
+			{
+				ExpansionStateInfo.CollapsedItems.Add(Item.Key.Get());
+			}
+		}
+
+		for (const auto& Item : FolderToTreeItemMap)
+		{
+			if (!OutlinerTreeView->IsItemExpanded(Item.Value))
+			{
+				ExpansionStateInfo.CollapsedFolders.Add(Item.Key);
+			}
+		}
+	}
+
+	void SSceneOutliner::SetParentsExpansionState(const FParentActorsSet& ParentActors, const FParentsExpansionState& ExpansionStateInfo) const
+	{
+		for (const auto* Actor : ParentActors)
+		{
+			// Set parent actor expansion state according to previous state
+			const auto& Item = ActorToTreeItemMap.FindChecked(Actor);
+			const bool bShouldExpand = !ExpansionStateInfo.CollapsedItems.Contains(Actor);
+			OutlinerTreeView->SetItemExpansion(Item, bShouldExpand);
+
+			// If the actor was just collapsed, this will force it to notice any expanded children nonetheless
+			if (!bShouldExpand && HasExpandedChildren(Actor, ParentActors, ExpansionStateInfo))
+			{
+				OutlinerTreeView->RefreshItemExpansion(Item);
+			}
+		}
+
+		for (const auto& Item : FolderToTreeItemMap)
+		{
+			// Expand folder according to the previous state
+			const bool bShouldExpand = !ExpansionStateInfo.CollapsedFolders.Contains(Item.Key);
+			OutlinerTreeView->SetItemExpansion(Item.Value, bShouldExpand);
+
+			// If the folder was just collapsed, this will force it to notice any expanded children nonetheless
+			if (!bShouldExpand && HasExpandedChildren(Item.Key, ParentActors, ExpansionStateInfo))
+			{
+				OutlinerTreeView->RefreshItemExpansion(Item.Value);
+			}
+		}
 	}
 
 	void SSceneOutliner::PopulateSearchStrings( const TOutlinerTreeItem& TreeItem, OUT TArray< FString >& OutSearchStrings ) const
@@ -1794,7 +1946,6 @@ namespace SceneOutliner
 		if (ExistingParent)
 		{
 			(*ExistingParent)->ChildItems.Add(TreeItem);
-			OutlinerTreeView->SetItemExpansion(*ExistingParent, true);
 			return true;
 		}
 		else
@@ -1803,7 +1954,6 @@ namespace SceneOutliner
 			if (AddFolderToTree(NewFolderItem, bIgnoreSearchFilter))
 			{
 				NewFolderItem->ChildItems.Add(TreeItem);
-				OutlinerTreeView->SetItemExpansion(NewFolderItem, true);
 				return true;
 			}
 		}
diff --git a/Engine/Source/Editor/SceneOutliner/Private/SSceneOutliner.h b/Engine/Source/Editor/SceneOutliner/Private/SSceneOutliner.h
--- a/Engine/Source/Editor/SceneOutliner/Private/SSceneOutliner.h
+++ b/Engine/Source/Editor/SceneOutliner/Private/SSceneOutliner.h
@@ -15,6 +15,7 @@ namespace SceneOutliner
 
 	typedef TMap< TWeakObjectPtr< AActor >, TSharedRef<TOutlinerActorTreeItem> > FActorToTreeItemMap;
 	typedef TMap< FName, TSharedRef<TOutlinerFolderTreeItem> > FFolderToTreeItemMap;
+	typedef TSet< const AActor* > FParentActorsSet;
 
 	FText GetLabelForItem( const TSharedRef<TOutlinerTreeItem> TreeItem );
 
@@ -364,8 +365,20 @@ namespace SceneOutliner
 		void ToggleHideTemporaryActors();
 		/** Enables/Disables whether the HideTemporaryActorsFilter is applied */
 		void ApplyHideTemporaryActorsFilter(bool bHideTemporaryActors);
+
 	private:
 
+		/** Structure containing information relating to the expansion state of parent items in the tree */
+		struct FParentsExpansionState
+		{
+			/** The set of parent items which are not expanded */
+			TSet<const AActor*> CollapsedItems;
+
+			/** The set of folders which are not expanded */
+			TSet<FName> CollapsedFolders;
+		};
+
+
 		/** Init options, cached */
 		FSceneOutlinerInitializationOptions InitOptions;
 
@@ -421,28 +434,28 @@ namespace SceneOutliner
 		const FSlateBrush* MobilityStationaryBrush;
 		const FSlateBrush* MobilityMovableBrush;
 
-		/** Populates the specified ActorMap with the specified UWorld */
-		void UpdateActorMapWithWorld( UWorld* World );
+		/** Populates the actor map with the specified UWorld */
+		void UpdateActorMapWithWorld(UWorld* World, FParentActorsSet& ParentActors);
 		
 		/** 
 		 * Attempts to add an item to the tree. Will add any parent folders if required
 		 *
 		 * @param InItem			The item to add to the tree (either an actor or folder)
-		 * @param ParentActors		The list of parent actors found that may need to be added later.
+		 * @param ParentActors		A set of parent actors found that may need to be added to.
 		 *
 		 * @return					returns true if adding the actor to the tree was successful.
 		 */
-		bool AddItemToTree(TSharedRef<TOutlinerTreeItem> InItem, TSet< AActor* >& ParentActors);
+		bool AddItemToTree(TSharedRef<TOutlinerTreeItem> InItem, FParentActorsSet& ParentActors);
 		
 		/** 
 		 * Attempts to add an actor to the tree. Will add any parent folders if required
 		 *
 		 * @param InActorItem		The actor item to add to the tree.
-		 * @param ParentActors		A set of parent actors found that may need to be added later.
+		 * @param ParentActors		A set of parent actors found that may need to be added to.
 		 *
 		 * @return					returns true if adding the actor to the tree was successful.
 		 */
-		bool AddActorToTree(TSharedRef<TOutlinerActorTreeItem> InActorItem, TSet<AActor*>& ParentActors);
+		bool AddActorToTree(TSharedRef<TOutlinerActorTreeItem> InActorItem, FParentActorsSet& ParentActors);
 
 		/** Add a new folder to the tree, including any of its parents if possible */
 		bool AddFolderToTree(TSharedRef<TOutlinerFolderTreeItem> InFolderItem, bool bIgnoreSearchFilter = false);
@@ -454,9 +467,20 @@ namespace SceneOutliner
 		 * Attempts to add parent actors that were filtered out to the tree.
 		 *
 		 * @param InActor			The actor to add to the tree.
-		 * @param InActorMap		The mapping of actors to their item in the tree view.
 		 */
-		void AddFilteredParentActorToTree(AActor* InActor);
+		void AddFilteredParentActorToTree(const AActor* InActor);
+
+		/** Returns whether any of the parent items lie within the subtree of the supplied parent actor, and are in an expanded state */
+		bool HasExpandedChildren(const AActor* ParentActor, const FParentActorsSet& ParentActors, const FParentsExpansionState& ExpansionStateInfo) const;
+
+		/** Returns whether any of the parent items or folders lie within the subtree of the supplied folder, and are in an expanded state */
+		bool HasExpandedChildren(FName ParentFolder, const FParentActorsSet& ParentActors, const FParentsExpansionState& ExpansionStateInfo) const;
+
+		/** Gets the current expansion state of parent items */
+		void GetParentsExpansionState(FParentsExpansionState& ExpansionStateInfo) const;
+
+		/** Updates the expansion state of parent items after a repopulate, according to the previous state */
+		void SetParentsExpansionState(const FParentActorsSet& ParentActors, const FParentsExpansionState& ExpansionStateInfo) const;
 
 		/** Whether the scene outliner is currently displaying PlayWorld actors */
 		bool bRepresentingPlayWorld;
diff --git a/Engine/Source/Editor/SceneOutliner/Public/SceneOutlinerTreeItems.h b/Engine/Source/Editor/SceneOutliner/Public/SceneOutlinerTreeItems.h
--- a/Engine/Source/Editor/SceneOutliner/Public/SceneOutlinerTreeItems.h
+++ b/Engine/Source/Editor/SceneOutliner/Public/SceneOutlinerTreeItems.h
@@ -60,7 +60,7 @@ struct SCENEOUTLINER_API TOutlinerActorTreeItem : TOutlinerTreeItem
 	/** true if this item exists in both the current world and PIE. */
 	bool bExistsInCurrentWorldAndPIE;
 
-	TOutlinerActorTreeItem(AActor* InActor)
+	TOutlinerActorTreeItem(const AActor* InActor)
 		: TOutlinerTreeItem(TOutlinerTreeItem::Actor)
 		, bExistsInCurrentWorldAndPIE(false)
 		, Actor(InActor)
diff --git a/Engine/Source/Runtime/Slate/Public/TableView/STreeView.h b/Engine/Source/Runtime/Slate/Public/TableView/STreeView.h
--- a/Engine/Source/Runtime/Slate/Public/TableView/STreeView.h
+++ b/Engine/Source/Runtime/Slate/Public/TableView/STreeView.h
@@ -510,6 +510,26 @@ public:
 			Private_SetItemExpansion(InItem, InShouldExpandItem);
 		}
 
+		/**
+		 * Refreshes the tree view to take into account any expanded children which have previously been set.
+		 * This will be necessary if they were expanded while the parent itself was collapsed.
+		 *
+		 * @param InItem         The item whose expansion state to refresh
+		 */
+		void RefreshItemExpansion(ItemType TheItem)
+		{
+			const FSparseItemInfo* const SparseItemInfo = SparseItemInfos.Find(TheItem);
+
+			if (SparseItemInfo)
+			{
+				SparseItemInfos.Add(TheItem, FSparseItemInfo(SparseItemInfo->bIsExpanded, true));
+			}
+			else
+			{
+				SparseItemInfos.Add(TheItem, FSparseItemInfo(false, true));
+			}
+		}
+
 		/** Collapse all the items in the tree and expand InItem */
 		void SetSingleExpandedItem( const ItemType& InItem )
 		{
@@ -528,7 +548,7 @@ public:
 		 *
 		 * @return true if the item is expanded; false otherwise.
 		 */
-		bool IsItemExpanded( const ItemType& InItem )
+		bool IsItemExpanded( const ItemType& InItem ) const
 		{
 			return Private_IsItemExpanded( InItem );
 		}
