#ttp 331902 	[GitHub] 33 : Scene Outliner: added filter "Only in Current Level"
https://github.com/EpicGames/UnrealEngine/pull/33

diff --git a/Engine/Config/BaseEditorUserSettings.ini b/Engine/Config/BaseEditorUserSettings.ini
--- a/Engine/Config/BaseEditorUserSettings.ini
+++ b/Engine/Config/BaseEditorUserSettings.ini
@@ -56,6 +56,9 @@ bShowOnlySelectedActors=False
 ; True when the Scene Outliner is hiding temporary/run-time Actors
 bHideTemporaryActors=False
 
+; True when the Scene Outliner is showing only Actors that exist in the current level
+bShowOnlyActorsInCurrentLevel=False
+
 ; Property Matrix
 PropertyMatrix_NumberOfPasteOperationsBeforeWarning=20
 
diff --git a/Engine/Source/Editor/SceneOutliner/Private/SSceneOutliner.cpp b/Engine/Source/Editor/SceneOutliner/Private/SSceneOutliner.cpp
--- a/Engine/Source/Editor/SceneOutliner/Private/SSceneOutliner.cpp
+++ b/Engine/Source/Editor/SceneOutliner/Private/SSceneOutliner.cpp
@@ -169,6 +169,7 @@ namespace SceneOutliner
 		{
 			ApplyShowOnlySelectedFilter( IsShowingOnlySelected() );
 			ApplyHideTemporaryActorsFilter( IsHidingTemporaryActors() );
+			ApplyShowOnlyCurrentLevelFilter( IsShowingOnlyCurrentLevel() );
 		}
 
 
@@ -506,6 +507,19 @@ namespace SceneOutliner
 				NAME_None,
 				EUserInterfaceActionType::ToggleButton
 			);
+
+			MenuBuilder.AddMenuEntry(
+				LOCTEXT("ToggleShowOnlyCurrentLevel", "Only in Current Level"),
+				LOCTEXT("ToggleShowOnlyCurrentLevelToolTip", "When enabled, only shows Actors that are in the Current Level."),
+				FSlateIcon(),
+				FUIAction(
+				FExecuteAction::CreateSP( this, &SSceneOutliner::ToggleShowOnlyCurrentLevel ),
+				FCanExecuteAction(),
+				FIsActionChecked::CreateSP( this, &SSceneOutliner::IsShowingOnlyCurrentLevel )
+				),
+				NAME_None,
+				EUserInterfaceActionType::ToggleButton
+			);
 		}
 		MenuBuilder.EndSection();
 
@@ -514,6 +528,8 @@ namespace SceneOutliner
 
 
 	/** FILTERS */
+
+	// Show Only Selected
 	void SSceneOutliner::ToggleShowOnlySelected()
 	{
 		const bool bEnableFlag = !IsShowingOnlySelected();
@@ -543,6 +559,7 @@ namespace SceneOutliner
 		return (GEditor->AccessEditorUserSettings().bShowOnlySelectedActors);
 	}
 
+	// Hide Temporary Actors
 	void SSceneOutliner::ToggleHideTemporaryActors()
 	{
 		const bool bEnableFlag = !IsHidingTemporaryActors();
@@ -571,6 +588,37 @@ namespace SceneOutliner
 	{
 		return (GEditor->AccessEditorUserSettings().bHideTemporaryActors);
 	}
+
+	// Show Only Actors In Current Level
+	void SSceneOutliner::ToggleShowOnlyCurrentLevel()
+	{
+		const bool bEnableFlag = !IsShowingOnlyCurrentLevel();
+		GEditor->AccessEditorUserSettings().bShowOnlyActorsInCurrentLevel = bEnableFlag;
+		GEditor->AccessEditorUserSettings().PostEditChange();
+
+		ApplyShowOnlyCurrentLevelFilter(bEnableFlag);
+	}
+	void SSceneOutliner::ApplyShowOnlyCurrentLevelFilter(bool bShowOnlyActorsInCurrentLevel)
+	{
+		if ( !ShowOnlyActorsInCurrentLevelFilter.IsValid() )
+		{
+			ShowOnlyActorsInCurrentLevelFilter = SceneOutlinerFilters::CreateIsInCurrentLevelFilter();
+		}
+
+		if ( bShowOnlyActorsInCurrentLevel )
+		{			
+			CustomFilters->Add( ShowOnlyActorsInCurrentLevelFilter );
+		}
+		else
+		{
+			CustomFilters->Remove( ShowOnlyActorsInCurrentLevelFilter );
+		}
+	}
+	bool SSceneOutliner::IsShowingOnlyCurrentLevel() const
+	{
+		return (GEditor->AccessEditorUserSettings().bShowOnlyActorsInCurrentLevel);
+	}
+
 	/** END FILTERS */
 
 
diff --git a/Engine/Source/Editor/SceneOutliner/Private/SSceneOutliner.h b/Engine/Source/Editor/SceneOutliner/Private/SSceneOutliner.h
--- a/Engine/Source/Editor/SceneOutliner/Private/SSceneOutliner.h
+++ b/Engine/Source/Editor/SceneOutliner/Private/SSceneOutliner.h
@@ -365,6 +365,13 @@ namespace SceneOutliner
 		void ToggleHideTemporaryActors();
 		/** Enables/Disables whether the HideTemporaryActorsFilter is applied */
 		void ApplyHideTemporaryActorsFilter(bool bHideTemporaryActors);
+		
+		/** @return whether we are showing only Actors that are in the Current Level */
+		bool IsShowingOnlyCurrentLevel() const;
+		/** Toggles whether we are hiding Actors that aren't in the current level */
+		void ToggleShowOnlyCurrentLevel();
+		/** Enables/Disables whether the ShowOnlyActorsInCurrentLevelFilter is applied */
+		void ApplyShowOnlyCurrentLevelFilter(bool bShowOnlyActorsInCurrentLevel);
 
 	private:
 
@@ -422,6 +429,8 @@ namespace SceneOutliner
 		/** When applied, temporary and run-time actors are hidden */
 		TSharedPtr< TDelegateFilter< const AActor* const > > HideTemporaryActorsFilter;
 
+		/** When applied, only Actors that are in the current level are displayed */
+		TSharedPtr< TDelegateFilter< const AActor* const > > ShowOnlyActorsInCurrentLevelFilter;
 
 		/** The brush to use when in Editor mode */
 		const FSlateBrush* NoBorder;
diff --git a/Engine/Source/Editor/SceneOutliner/Private/SceneOutlinerFilters.h b/Engine/Source/Editor/SceneOutliner/Private/SceneOutlinerFilters.h
--- a/Engine/Source/Editor/SceneOutliner/Private/SceneOutlinerFilters.h
+++ b/Engine/Source/Editor/SceneOutliner/Private/SceneOutlinerFilters.h
@@ -17,6 +17,11 @@ namespace SceneOutlinerFilters
 		{
 			return (InActor != NULL && (InActor->GetWorld()->WorldType != EWorldType::PIE || GEditor->ObjectsThatExistInEditorWorld.Get(InActor)));
 		}
+	
+		static bool IsInCurrentLevel( const AActor* const InActor )
+		{
+			return (InActor != NULL && InActor->GetLevel() == InActor->GetWorld()->GetCurrentLevel());
+		}
 	};
 
 	TSharedPtr< TDelegateFilter< const AActor* const > > CreateSelectedActorFilter()
@@ -30,4 +35,10 @@ namespace SceneOutlinerFilters
 		return MakeShareable( new TDelegateFilter< const AActor* const >(
 			TDelegateFilter< const AActor* const >::FPredicate::CreateStatic( &FActorClassPredicate::ExistsInEditorWorld ) ) );
 	}
+
+	TSharedPtr< TDelegateFilter< const AActor* const > > CreateIsInCurrentLevelFilter()
+	{
+		return MakeShareable( new TDelegateFilter< const AActor* const >(
+			TDelegateFilter< const AActor* const >::FPredicate::CreateStatic( &FActorClassPredicate::IsInCurrentLevel ) ) );
+	}
 }
\ No newline at end of file
diff --git a/Engine/Source/Editor/UnrealEd/Classes/Editor/EditorUserSettings.h b/Engine/Source/Editor/UnrealEd/Classes/Editor/EditorUserSettings.h
--- a/Engine/Source/Editor/UnrealEd/Classes/Editor/EditorUserSettings.h
+++ b/Engine/Source/Editor/UnrealEd/Classes/Editor/EditorUserSettings.h
@@ -288,6 +288,10 @@ public:
 	UPROPERTY(config)
 	uint32 bHideTemporaryActors:1;
 
+	/** True when the Scene Outliner is showing only Actors that exist in the current level */
+	UPROPERTY(config)
+	uint32 bShowOnlyActorsInCurrentLevel:1;
+
 	UPROPERTY(config)
 	int32 MaterialQualityLevel;
 
