#ttp 331569 - EDITOR: Actors can be renamed with a space character which results in the actor taking on a default generic name
#branch UE4
#proj Editor.SceneOutliner
#summary Renaming an actor now ignores leading or trailing space when validating the name.
#add SSceneOutliner::OnVerifyItemLabelChanged now trims leading and trailing space from the label before validating it.

diff --git a/Engine/Source/Editor/SceneOutliner/Private/SSceneOutliner.cpp b/Engine/Source/Editor/SceneOutliner/Private/SSceneOutliner.cpp
--- a/Engine/Source/Editor/SceneOutliner/Private/SSceneOutliner.cpp
+++ b/Engine/Source/Editor/SceneOutliner/Private/SSceneOutliner.cpp
@@ -1679,13 +1679,15 @@ namespace SceneOutliner
 
 	bool SSceneOutliner::OnVerifyItemLabelChanged( const FText& InLabel, FText& OutErrorMessage, FOutlinerTreeItemPtr TreeItem )
 	{
-		if( InLabel.IsEmpty() )
+		FText TrimmedLabel = FText::TrimPrecedingAndTrailing(InLabel);
+
+		if (TrimmedLabel.IsEmpty())
 		{
 			OutErrorMessage = LOCTEXT( "RenameFailed_LeftBlank", "Names cannot be left blank" );
 			return false;
 		}
 
-		if(InLabel.ToString().Len() >= NAME_SIZE)
+		if (TrimmedLabel.ToString().Len() >= NAME_SIZE)
 		{
 			FFormatNamedArguments Arguments;
 			Arguments.Add( TEXT("CharCount"), NAME_SIZE );
@@ -1693,7 +1695,7 @@ namespace SceneOutliner
 			return false;
 		}
 
-		FString LabelString = InLabel.ToString();
+		FString LabelString = TrimmedLabel.ToString();
 		if (TreeItem->Type == TOutlinerTreeItem::Folder)
 		{
 			auto FolderTreeItem = StaticCastSharedPtr<TOutlinerFolderTreeItem>(TreeItem);
