#summary 	Fixed docking not working correctly when docking over the scene outliner.
#ttp 332669 	EDITOR: Undocking Details tab doesn't bring up dock guides over scene outliner
#detail 	The scene outliner was misbehaving itself, saying it had handled drag & drop operations when it really hadnt.

diff --git a/Engine/Source/Editor/SceneOutliner/Private/SOutlinerTreeView.cpp b/Engine/Source/Editor/SceneOutliner/Private/SOutlinerTreeView.cpp
--- a/Engine/Source/Editor/SceneOutliner/Private/SOutlinerTreeView.cpp
+++ b/Engine/Source/Editor/SceneOutliner/Private/SOutlinerTreeView.cpp
@@ -429,7 +429,7 @@ namespace SceneOutliner
 		auto SceneOutlinerPtr = SceneOutlinerWeak.Pin();
 		if (!SceneOutlinerPtr.IsValid())
 		{
-			return FReply::Handled();
+			return FReply::Unhandled();
 		}
 
 
@@ -450,16 +450,18 @@ namespace SceneOutliner
 				{
 					DecoratedDragOp->CurrentIconBrush = FEditorStyle::GetBrush(TEXT("Graph.ConnectorFeedback.Error"));
 				}
+				return FReply::Handled();
 			}
 			else if (DecoratedDragOp->IsOfType<FActorDragDropGraphEdOp>())
 			{
 				// Display actor specific messages
 				auto ActorDragOp = StaticCastSharedPtr<FActorDragDropGraphEdOp>(DecoratedDragOp);
 				ActorDragOp->SetToolTip(ValidationInfo.TooltipType, ValidationInfo.ValidationText);
+				return FReply::Handled();
 			}
 		}
 
-		return FReply::Handled();
+		return FReply::Unhandled();
 	}
 	
 	void SOutlinerTreeView::OnDragLeave(const FDragDropEvent& DragDropEvent)
@@ -490,14 +492,14 @@ namespace SceneOutliner
 		// Only allow dropping objects when in actor browsing mode
 		if (SceneOutlinerPtr->GetInitOptions().Mode != ESceneOutlinerMode::ActorBrowsing )
 		{
-			return FReply::Handled();
+			return FReply::Unhandled();
 		}
 
 		// Validate now to make sure we don't doing anything we shouldn't
 		const FDragValidationInfo ValidationInfo = ValidateRootDragDropEvent(SceneOutlinerPtr.ToSharedRef(), DragDropEvent);
 		if (!ValidationInfo.IsValid())
 		{
-			return FReply::Handled();
+			return FReply::Unhandled();
 		}
 
 		const FScopedTransaction Transaction(LOCTEXT("UndoAction_DropOnSceneOutliner", "Drop on Scene Outliner"));
@@ -547,7 +549,7 @@ namespace SceneOutliner
 		const FDragValidationInfo ValidationInfo = ValidateDragDropEvent(SceneOutlinerPtr.ToSharedRef(), DragDropEvent, Item.ToSharedRef());
 		if (!ValidationInfo.IsValid())
 		{
-			return FReply::Handled();
+			return FReply::Unhandled();
 		}
 
 		const FScopedTransaction Transaction(LOCTEXT("UndoAction_DropOnSceneOutliner", "Drop on Scene Outliner"));
