#ttp 330692 - EDITOR: Scene Outliner opens all parent trees when undo button pressed.
#branch UE4
#proj Editor.SceneOutliner, Runtime.Slate
#summary Rolled back change to STreeView, and made partial fix to undo opening scene outliner parent trees.
#remove Rolled back change to STreeView to remove old RefreshItemExpansion method.
#remove Removed methods in SSceneOutliner used to determine whether a given parent has expanded children.

diff --git a/Engine/Source/Editor/SceneOutliner/Private/SSceneOutliner.cpp b/Engine/Source/Editor/SceneOutliner/Private/SSceneOutliner.cpp
--- a/Engine/Source/Editor/SceneOutliner/Private/SSceneOutliner.cpp
+++ b/Engine/Source/Editor/SceneOutliner/Private/SSceneOutliner.cpp
@@ -1146,110 +1146,6 @@ namespace SceneOutliner
 		}
 	}
 
-	// Helper function which returns whether a folder is a child of another.
-	// Note that a folder is defined to not be a child of itself.
-	bool DoesFolderContainFolder(FName ParentFolder, FName ChildFolder)
-	{
-		FString Path = ChildFolder.ToString();
-		FString ParentFolderAsString = ParentFolder.ToString();
-
-		do 
-		{
-			Path = FPaths::GetPath(Path);
-
-		} while (!Path.IsEmpty() && Path != ParentFolderAsString);
-
-		return !Path.IsEmpty();
-	}
-
-	// Helper function which returns whether an actor is a child of a folder.
-	bool DoesFolderContainActor(FName ParentFolder, const AActor* ChildActor)
-	{
-		if (ChildActor != nullptr)
-		{
-			const FName ActorFolder = ChildActor->GetFolderPath();
-			if (!ActorFolder.IsNone())
-			{
-				return (ActorFolder == ParentFolder || DoesFolderContainFolder(ParentFolder, ActorFolder));
-			}
-		}
-
-		return false;
-	}
-
-	// Helper function which returns whether an actor is a child of another.
-	// Note that an actor is defined to not be a child of itself.
-	bool IsActorParentOfActor(const AActor* ParentActor, const AActor* ChildActor)
-	{
-		const AActor* Actor = ChildActor;
-
-		if (Actor != nullptr)
-		{
-			do
-			{
-				Actor = Actor->GetAttachParentActor();
-
-			} while (Actor != nullptr && Actor != ParentActor);
-		}
-
-		return (Actor != nullptr);
-	}
-
-	bool SSceneOutliner::HasExpandedChildren(const AActor* ParentActor, const FParentActorsSet& ParentActors, const FParentsExpansionState& ExpansionStateInfo) const
-	{
-		// Iterate through all parent actors, determining whether any of them are within the subtree of the
-		// parent actor, and if so, whether they are expanded.
-		for (const auto* ParentToTest : ParentActors)
-		{
-			if (ParentToTest != ParentActor)
-			{
-				if (IsActorParentOfActor(ParentActor, ParentToTest))
-				{
-					if (!ExpansionStateInfo.CollapsedItems.Contains(ParentToTest))
-					{
-						return true;
-					}
-				}
-			}
-		}
-
-		return false;
-	}
-
-	bool SSceneOutliner::HasExpandedChildren(FName ParentFolder, const FParentActorsSet& ParentActors, const FParentsExpansionState& ExpansionStateInfo) const
-	{
-		// Iterate through all parent actors, determining whether any of them are within the subtree of the
-		// parent actor, and if so, whether they are expanded.
-		for (const auto* ParentToTest : ParentActors)
-		{
-			if (DoesFolderContainActor(ParentFolder, ParentToTest))
-			{
-				if (!ExpansionStateInfo.CollapsedItems.Contains(ParentToTest))
-				{
-					return true;
-				}
-			}
-		}
-
-		// If there were no expanded parent actors in this folder subtree, do the same check, this time with other
-		// folders.
-		for (const auto& Item : FolderToTreeItemMap)
-		{
-			if (Item.Key != ParentFolder)
-			{
-				if (DoesFolderContainFolder(ParentFolder, Item.Key))
-				{
-					if (!ExpansionStateInfo.CollapsedFolders.Contains(Item.Key))
-					{
-						return true;
-					}
-				}
-			}
-		}
-
-		return false;
-	}
-
 	void SSceneOutliner::GetParentsExpansionState(FParentsExpansionState& ExpansionStateInfo) const
 	{
 		for (const auto& Item : ActorToTreeItemMap)
@@ -1277,12 +1173,6 @@ namespace SceneOutliner
 			const auto& Item = ActorToTreeItemMap.FindChecked(Actor);
 			const bool bShouldExpand = !ExpansionStateInfo.CollapsedItems.Contains(Actor);
 			OutlinerTreeView->SetItemExpansion(Item, bShouldExpand);
-
-			// If the actor was just collapsed, this will force it to notice any expanded children nonetheless
-			if (!bShouldExpand && HasExpandedChildren(Actor, ParentActors, ExpansionStateInfo))
-			{
-				OutlinerTreeView->RefreshItemExpansion(Item);
-			}
 		}
 
 		for (const auto& Item : FolderToTreeItemMap)
@@ -1290,12 +1180,6 @@ namespace SceneOutliner
 			// Expand folder according to the previous state
 			const bool bShouldExpand = !ExpansionStateInfo.CollapsedFolders.Contains(Item.Key);
 			OutlinerTreeView->SetItemExpansion(Item.Value, bShouldExpand);
-
-			// If the folder was just collapsed, this will force it to notice any expanded children nonetheless
-			if (!bShouldExpand && HasExpandedChildren(Item.Key, ParentActors, ExpansionStateInfo))
-			{
-				OutlinerTreeView->RefreshItemExpansion(Item.Value);
-			}
 		}
 	}
 
diff --git a/Engine/Source/Editor/SceneOutliner/Private/SSceneOutliner.h b/Engine/Source/Editor/SceneOutliner/Private/SSceneOutliner.h
--- a/Engine/Source/Editor/SceneOutliner/Private/SSceneOutliner.h
+++ b/Engine/Source/Editor/SceneOutliner/Private/SSceneOutliner.h
@@ -480,12 +480,6 @@ namespace SceneOutliner
 		 */
 		void AddFilteredParentActorToTree(const AActor* InActor);
 
-		/** Returns whether any of the parent items lie within the subtree of the supplied parent actor, and are in an expanded state */
-		bool HasExpandedChildren(const AActor* ParentActor, const FParentActorsSet& ParentActors, const FParentsExpansionState& ExpansionStateInfo) const;
-
-		/** Returns whether any of the parent items or folders lie within the subtree of the supplied folder, and are in an expanded state */
-		bool HasExpandedChildren(FName ParentFolder, const FParentActorsSet& ParentActors, const FParentsExpansionState& ExpansionStateInfo) const;
-
 		/** Gets the current expansion state of parent items */
 		void GetParentsExpansionState(FParentsExpansionState& ExpansionStateInfo) const;
 
diff --git a/Engine/Source/Runtime/Slate/Public/TableView/STreeView.h b/Engine/Source/Runtime/Slate/Public/TableView/STreeView.h
--- a/Engine/Source/Runtime/Slate/Public/TableView/STreeView.h
+++ b/Engine/Source/Runtime/Slate/Public/TableView/STreeView.h
@@ -510,26 +510,6 @@ public:
 			Private_SetItemExpansion(InItem, InShouldExpandItem);
 		}
 
-		/**
-		 * Refreshes the tree view to take into account any expanded children which have previously been set.
-		 * This will be necessary if they were expanded while the parent itself was collapsed.
-		 *
-		 * @param InItem         The item whose expansion state to refresh
-		 */
-		void RefreshItemExpansion(ItemType TheItem)
-		{
-			const FSparseItemInfo* const SparseItemInfo = SparseItemInfos.Find(TheItem);
-
-			if (SparseItemInfo)
-			{
-				SparseItemInfos.Add(TheItem, FSparseItemInfo(SparseItemInfo->bIsExpanded, true));
-			}
-			else
-			{
-				SparseItemInfos.Add(TheItem, FSparseItemInfo(false, true));
-			}
-		}
-
 		/** Collapse all the items in the tree and expand InItem */
 		void SetSingleExpandedItem( const ItemType& InItem )
 		{
