Fixed scene outliner context menu showing an empty folder picker for choosing a folder to move actors into (it's now hidden if there aren't any existing folders)

diff --git a/Engine/Source/Editor/SceneOutliner/Private/SSceneOutliner.cpp b/Engine/Source/Editor/SceneOutliner/Private/SSceneOutliner.cpp
--- a/Engine/Source/Editor/SceneOutliner/Private/SSceneOutliner.cpp
+++ b/Engine/Source/Editor/SceneOutliner/Private/SSceneOutliner.cpp
@@ -1727,24 +1727,6 @@ namespace SceneOutliner
 
 	void SSceneOutliner::FillFoldersSubMenu(FMenuBuilder& MenuBuilder) const
 	{
-		// Build up the menu
-		FSceneOutlinerInitializationOptions InitOptions;
-		InitOptions.bShowHeaderRow = false;
-		InitOptions.bFocusSearchBoxWhenOpened = true;
-		InitOptions.bOnlyShowFolders = true;
-
-		// Actor selector to allow the user to choose a folder
-		FSceneOutlinerModule& SceneOutlinerModule = FModuleManager::LoadModuleChecked<FSceneOutlinerModule>( "SceneOutliner" );
-		TSharedRef< SWidget > MiniSceneOutliner =
-			SNew( SVerticalBox )
-			+SVerticalBox::Slot()
-			.MaxHeight(400.0f)
-			[
-				SNew( SceneOutliner::SSceneOutliner, InitOptions )
-				.IsEnabled( FSlateApplication::Get().GetNormalExecutionAttribute() )
-				.OnItemPickedDelegate(FOnSceneOutlinerItemPicked::CreateSP(this, &SSceneOutliner::MoveSelectionTo))
-			];
-
 		MenuBuilder.AddMenuEntry(LOCTEXT( "CreateNew", "Create New" ), LOCTEXT( "MoveToRoot_ToolTip", "Move to a new folder" ),
 			FSlateIcon(FEditorStyle::GetStyleSetName(), "SceneOutliner.NewFolderIcon"), FExecuteAction::CreateSP(this, &SSceneOutliner::CreateFolder));
 
@@ -1754,10 +1736,31 @@ namespace SceneOutliner
 				FSlateIcon(FEditorStyle::GetStyleSetName(), "SceneOutliner.MoveToRoot"), FExecuteAction::CreateSP(this, &SSceneOutliner::MoveSelectionTo, FName()));
 		}
 
-		// Existing folders - it would be good to filter out the existing selected folders from this mini-outliner, but that's currently not possible
-		MenuBuilder.BeginSection(FName(), LOCTEXT("ExistingFolders", "Existing:"));
-		MenuBuilder.AddWidget(MiniSceneOutliner, FText::GetEmpty(), false);
-		MenuBuilder.EndSection();
+		// Add a mini scene outliner for choosing an existing folder
+		if (RepresentingWorld && FActorFolders::Get().GetFoldersForWorld(*RepresentingWorld).Num())
+		{
+			FSceneOutlinerInitializationOptions InitOptions;
+			InitOptions.bShowHeaderRow = false;
+			InitOptions.bFocusSearchBoxWhenOpened = true;
+			InitOptions.bOnlyShowFolders = true;
+
+			// Actor selector to allow the user to choose a folder
+			FSceneOutlinerModule& SceneOutlinerModule = FModuleManager::LoadModuleChecked<FSceneOutlinerModule>("SceneOutliner");
+			TSharedRef< SWidget > MiniSceneOutliner =
+				SNew(SVerticalBox)
+				+ SVerticalBox::Slot()
+				.MaxHeight(400.0f)
+				[
+					SNew(SceneOutliner::SSceneOutliner, InitOptions)
+					.IsEnabled(FSlateApplication::Get().GetNormalExecutionAttribute())
+					.OnItemPickedDelegate(FOnSceneOutlinerItemPicked::CreateSP(this, &SSceneOutliner::MoveSelectionTo))
+				];
+
+			// Existing folders - it would be good to filter out the existing selected folders from this mini-outliner, but that's currently not possible
+			MenuBuilder.BeginSection(FName(), LOCTEXT("ExistingFolders", "Existing:"));
+			MenuBuilder.AddWidget(MiniSceneOutliner, FText::GetEmpty(), false);
+			MenuBuilder.EndSection();
+		}
 	}
 
 	void SSceneOutliner::FillSelectionSubMenu(FMenuBuilder& MenuBuilder) const
