Simplified the message shown when attempting to move a scene outliner folder to a location that would result in a duplicate
This addresses TTP 331615

diff --git a/Engine/Source/Editor/SceneOutliner/Private/SOutlinerTreeView.cpp b/Engine/Source/Editor/SceneOutliner/Private/SOutlinerTreeView.cpp
--- a/Engine/Source/Editor/SceneOutliner/Private/SOutlinerTreeView.cpp
+++ b/Engine/Source/Editor/SceneOutliner/Private/SOutlinerTreeView.cpp
@@ -114,18 +114,10 @@ namespace SceneOutliner
 					// The folder already exists
 					FFormatNamedArguments Args;
 					Args.Add(TEXT("DragName"), FText::FromString(LeafName));
-					if (DstFolderPath.IsEmpty())
-					{
+					
+					return FDragValidationInfo(FActorDragDropGraphEdOp::ToolTip_IncompatibleGeneric,
+						FText::Format(LOCTEXT("FolderAlreadyExists", "A folder called \"{DragName}\" already exists at this level"), Args));
 
-						return FDragValidationInfo(FActorDragDropGraphEdOp::ToolTip_IncompatibleGeneric,
-							FText::Format(LOCTEXT("FolderAlreadyExistsRoot", "\"{DragName}\" already exists in root"), Args));
-					}
-					else
-					{
-						Args.Add(TEXT("DropPath"), FText::FromString(DstFolderPath));
-						return FDragValidationInfo(FActorDragDropGraphEdOp::ToolTip_IncompatibleGeneric,
-							FText::Format(LOCTEXT("FolderAlreadyExists", "\"{DragName}\" already exists in \"{DropPath}\""), Args));
-					}
 				}
 				else if (DragFolderPath == DstFolderPath || DstFolderPath.StartsWith(DragFolderPath + "/"))
 				{
