Added folder filtering support to the scene outliner

This allows the filtering of any actors that are part of the specified folder.

This addresses TTP#331604 - EDITOR: Scene Outliner: Folders will appear in the drop down when you try to move them, and TTP#331640 - EDITOR: Scene Outliner: In the 'Move to Folder" sub menu the "Create New" tooltip is incorrect

Also fixed not being able to move attached actors to a folder

diff --git a/Engine/Source/Editor/Layers/Private/SLayerBrowser.h b/Engine/Source/Editor/Layers/Private/SLayerBrowser.h
--- a/Engine/Source/Editor/Layers/Private/SLayerBrowser.h
+++ b/Engine/Source/Editor/Layers/Private/SLayerBrowser.h
@@ -164,7 +164,7 @@ public:
 			InitOptions.CustomDelete = FCustomSceneOutlinerDeleteDelegate::CreateSP( this, &SLayerBrowser::RemoveActorsFromSelectedLayer );
 			InitOptions.CustomColumnFactory = FCreateSceneOutlinerColumnDelegate::CreateSP( this, &SLayerBrowser::CreateCustomLayerColumn );
 
-			InitOptions.ActorFilters->Add( SelectedLayersFilter );
+			InitOptions.Filters->Add( SelectedLayersFilter );
 		}
 
 		SAssignNew( LayerContentsSection, SBorder )
diff --git a/Engine/Source/Editor/LevelEditor/Private/LevelEditorContextMenu.cpp b/Engine/Source/Editor/LevelEditor/Private/LevelEditorContextMenu.cpp
--- a/Engine/Source/Editor/LevelEditor/Private/LevelEditorContextMenu.cpp
+++ b/Engine/Source/Editor/LevelEditor/Private/LevelEditorContextMenu.cpp
@@ -710,9 +710,6 @@ void FLevelEditorContextMenuImpl::FillActorMenu( FMenuBuilder& MenuBuilder )
 		InitOptions.bShowHeaderRow = false;
 		InitOptions.bFocusSearchBoxWhenOpened = true;
 
-		// Only display Actors that we can attach too
-		InitOptions.ActorFilters = MakeShareable( new TFilterCollection< const AActor* const >() );
-
 		struct Local
 		{
 			static bool IsAttachableActor( const AActor* const ParentActor )
@@ -737,7 +734,8 @@ void FLevelEditorContextMenuImpl::FillActorMenu( FMenuBuilder& MenuBuilder )
 			}
 		};
 
-		InitOptions.ActorFilters->Add( MakeShareable( new TDelegateFilter< const AActor* const >( TDelegateFilter< const AActor* const >::FPredicate::CreateStatic( &Local::IsAttachableActor ) ) ) );
+		// Only display Actors that we can attach too
+		InitOptions.Filters->AddFilterPredicate( SceneOutliner::FActorFilterPredicate::CreateStatic( &Local::IsAttachableActor ) );
 	}		
 
 	// Actor selector to allow the user to choose a parent actor
diff --git a/Engine/Source/Editor/LevelEditor/Private/LevelEditorToolBar.cpp b/Engine/Source/Editor/LevelEditor/Private/LevelEditorToolBar.cpp
--- a/Engine/Source/Editor/LevelEditor/Private/LevelEditorToolBar.cpp
+++ b/Engine/Source/Editor/LevelEditor/Private/LevelEditorToolBar.cpp
@@ -1003,9 +1003,6 @@ TSharedRef< SWidget > FLevelEditorToolBar::GenerateMatineeMenuContent( TSharedRe
 		// @todo: Might be useful to have this sometimes, actually.  Ideally the user could summon it.
 		InitOptions.bShowHeaderRow = false;
 
-		// Only display Matinee actors
-		InitOptions.ActorFilters = MakeShareable( new TFilterCollection< const AActor* const >() );
-
 		struct Local
 		{
 			static bool IsMatineeActor( const AActor* const Actor )
@@ -1014,7 +1011,8 @@ TSharedRef< SWidget > FLevelEditorToolBar::GenerateMatineeMenuContent( TSharedRe
 			}
 		};
 
-		InitOptions.ActorFilters->Add( MakeShareable( new TDelegateFilter< const AActor* const >( TDelegateFilter< const AActor* const >::FPredicate::CreateStatic( &Local::IsMatineeActor ) ) ) );
+		// Only display Matinee actors
+		InitOptions.Filters->AddFilterPredicate( SceneOutliner::FActorFilterPredicate::CreateStatic( &Local::IsMatineeActor ) );
 	}
 
 	// actor selector to allow the user to choose a Matinee actor
diff --git a/Engine/Source/Editor/LevelEditor/Private/SLevelViewportToolBar.cpp b/Engine/Source/Editor/LevelEditor/Private/SLevelViewportToolBar.cpp
--- a/Engine/Source/Editor/LevelEditor/Private/SLevelViewportToolBar.cpp
+++ b/Engine/Source/Editor/LevelEditor/Private/SLevelViewportToolBar.cpp
@@ -326,20 +326,17 @@ static void OnGenerateActorLockingMenuSection( TWeakPtr<SLevelViewport> Viewport
 	InitOptions.bShowHeaderRow = false;
 	InitOptions.CustomColumnFixedWidth = SLevelViewport::GetActorLockSceneOutlinerColumnWidth();
 	InitOptions.CustomColumnFactory = FCreateSceneOutlinerColumnDelegate::CreateSP( Viewport.Pin().ToSharedRef(), &SLevelViewport::CreateActorLockSceneOutlinerColumn );
-	InitOptions.ActorFilters = MakeShareable( new TFilterCollection< const AActor* const >() );
 
 	// Predicate for selecting actors by class
 	struct FActorClassPredicate
 	{
-		static bool IsLockableActor( const AActor* const InActor )
+		static bool IsLockableActor( const AActor* InActor )
 		{			
 			return InActor != NULL && InActor->IsA(ActorClass::StaticClass()) && !InActor->IsPendingKill();
 		}
 	};
 
-	InitOptions.ActorFilters->Add(
-		MakeShareable(new TDelegateFilter< const AActor* const >(
-			TDelegateFilter< const AActor* const >::FPredicate::CreateStatic( &FActorClassPredicate::IsLockableActor ) ) ) );
+	InitOptions.Filters->AddFilterPredicate( SceneOutliner::FActorFilterPredicate::CreateStatic(&FActorClassPredicate::IsLockableActor) );
 
 	// Create an outliner with the options from above. It sits in a box with a max height limit to stop it getting too tall when lots of actors are added.
 	FSceneOutlinerModule& SceneOutlinerModule = FModuleManager::LoadModuleChecked<FSceneOutlinerModule>("SceneOutliner");
diff --git a/Engine/Source/Editor/PropertyEditor/Private/Presentation/PropertyEditor/PropertyEditor.cpp b/Engine/Source/Editor/PropertyEditor/Private/Presentation/PropertyEditor/PropertyEditor.cpp
--- a/Engine/Source/Editor/PropertyEditor/Private/Presentation/PropertyEditor/PropertyEditor.cpp
+++ b/Engine/Source/Editor/PropertyEditor/Private/Presentation/PropertyEditor/PropertyEditor.cpp
@@ -345,7 +345,7 @@ void FPropertyEditor::OnActorSelected( AActor* InActor )
 	OnAssetSelected(InActor);
 }
 
-void FPropertyEditor::OnGetActorFiltersForSceneOutliner( TSharedPtr<TFilterCollection<const AActor* const> >& OutFilters )
+void FPropertyEditor::OnGetActorFiltersForSceneOutliner( TSharedPtr<SceneOutliner::FOutlinerFilters>& OutFilters )
 {
 	struct Local
 	{
@@ -363,7 +363,6 @@ void FPropertyEditor::OnGetActorFiltersForSceneOutliner( TSharedPtr<TFilterColle
 		}
 	};
 
-	OutFilters = MakeShareable( new TFilterCollection< const AActor* const >() );
 	OutFilters->Add( MakeShareable( new TDelegateFilter< const AActor* const >( TDelegateFilter< const AActor* const >::FPredicate::CreateStatic( &Local::IsFilteredActor, AsShared() ) ) ) );
 }
 
diff --git a/Engine/Source/Editor/PropertyEditor/Private/Presentation/PropertyEditor/PropertyEditor.h b/Engine/Source/Editor/PropertyEditor/Private/Presentation/PropertyEditor/PropertyEditor.h
--- a/Engine/Source/Editor/PropertyEditor/Private/Presentation/PropertyEditor/PropertyEditor.h
+++ b/Engine/Source/Editor/PropertyEditor/Private/Presentation/PropertyEditor/PropertyEditor.h
@@ -1,8 +1,7 @@
 // Copyright 1998-2014 Epic Games, Inc. All Rights Reserved.
 #pragma once
 
-#include "IFilter.h"
-#include "FilterCollection.h"
+#include "Editor/SceneOutliner/Public/SceneOutlinerFilters.h"
 
 class FPropertyNode;
 class IPropertyUtilities;
@@ -89,7 +88,7 @@ public:
 	void OnGetClassesForAssetPicker( TArray<const UClass*>& OutClasses );
 	void OnAssetSelected( const FAssetData& AssetData );
 	void OnActorSelected( AActor* InActor );
-	void OnGetActorFiltersForSceneOutliner( TSharedPtr<TFilterCollection<const AActor* const> >& OutFilters );
+	void OnGetActorFiltersForSceneOutliner( TSharedPtr<SceneOutliner::FOutlinerFilters>& OutFilters );
 
 	/**	In an ideal world we wouldn't expose these */
 	TSharedRef< FPropertyNode > GetPropertyNode() const;
diff --git a/Engine/Source/Editor/PropertyEditor/Private/PropertyCustomizationHelpers.cpp b/Engine/Source/Editor/PropertyEditor/Private/PropertyCustomizationHelpers.cpp
--- a/Engine/Source/Editor/PropertyEditor/Private/PropertyCustomizationHelpers.cpp
+++ b/Engine/Source/Editor/PropertyEditor/Private/PropertyCustomizationHelpers.cpp
@@ -207,7 +207,7 @@ namespace PropertyCustomizationHelpers
 			.OnActorSelected( OnActorSelectedFromPicker );
 	}
 
-	TSharedRef<SWidget> MakeActorPickerWithMenu( AActor* const InitialActor, const bool AllowClear, const TSharedPtr< TFilterCollection<const AActor* const> >& ActorFilters, FOnActorSelected OnSet, FSimpleDelegate OnClose, FSimpleDelegate OnUseSelected )
+	TSharedRef<SWidget> MakeActorPickerWithMenu( AActor* const InitialActor, const bool AllowClear, const TSharedPtr< SceneOutliner::FOutlinerFilters >& ActorFilters, FOnActorSelected OnSet, FSimpleDelegate OnClose, FSimpleDelegate OnUseSelected )
 	{
 		return 
 			SNew( SPropertyMenuActorPicker )
diff --git a/Engine/Source/Editor/PropertyEditor/Private/UserInterface/PropertyEditor/SPropertyEditorAsset.cpp b/Engine/Source/Editor/PropertyEditor/Private/UserInterface/PropertyEditor/SPropertyEditorAsset.cpp
--- a/Engine/Source/Editor/PropertyEditor/Private/UserInterface/PropertyEditor/SPropertyEditorAsset.cpp
+++ b/Engine/Source/Editor/PropertyEditor/Private/UserInterface/PropertyEditor/SPropertyEditorAsset.cpp
@@ -271,7 +271,7 @@ TSharedRef<SWidget> SPropertyEditorAsset::OnGetMenuContent()
 
 	if(bIsActor)
 	{
-		TSharedPtr<TFilterCollection<const AActor* const> > ActorFilters = MakeShareable( new TFilterCollection< const AActor* const >() );
+		TSharedPtr<SceneOutliner::FOutlinerFilters> ActorFilters = MakeShareable( new SceneOutliner::FOutlinerFilters );
 		ActorFilters->Add( MakeShareable( new TDelegateFilter< const AActor* const >( TDelegateFilter< const AActor* const >::FPredicate::CreateSP( this, &SPropertyEditorAsset::IsFilteredActor ) ) ) );
 
 		return PropertyCustomizationHelpers::MakeActorPickerWithMenu(Cast<AActor>(Value.Object), bAllowClear, ActorFilters, FOnActorSelected::CreateSP( this, &SPropertyEditorAsset::OnActorSelected), FSimpleDelegate::CreateSP( this, &SPropertyEditorAsset::CloseComboButton ), FSimpleDelegate::CreateSP( this, &SPropertyEditorAsset::OnUse ) );
diff --git a/Engine/Source/Editor/PropertyEditor/Private/UserInterface/PropertyEditor/SPropertyMenuActorPicker.cpp b/Engine/Source/Editor/PropertyEditor/Private/UserInterface/PropertyEditor/SPropertyMenuActorPicker.cpp
--- a/Engine/Source/Editor/PropertyEditor/Private/UserInterface/PropertyEditor/SPropertyMenuActorPicker.cpp
+++ b/Engine/Source/Editor/PropertyEditor/Private/UserInterface/PropertyEditor/SPropertyMenuActorPicker.cpp
@@ -73,7 +73,7 @@ void SPropertyMenuActorPicker::Construct( const FArguments& InArgs )
 
 		FSceneOutlinerInitializationOptions InitOptions;
 		InitOptions.Mode = ESceneOutlinerMode::ActorPicker;
-		InitOptions.ActorFilters = ActorFilters;
+		InitOptions.Filters = ActorFilters;
 		InitOptions.bFocusSearchBoxWhenOpened = true;
 		
 		MenuContent =
diff --git a/Engine/Source/Editor/PropertyEditor/Private/UserInterface/PropertyEditor/SPropertyMenuActorPicker.h b/Engine/Source/Editor/PropertyEditor/Private/UserInterface/PropertyEditor/SPropertyMenuActorPicker.h
--- a/Engine/Source/Editor/PropertyEditor/Private/UserInterface/PropertyEditor/SPropertyMenuActorPicker.h
+++ b/Engine/Source/Editor/PropertyEditor/Private/UserInterface/PropertyEditor/SPropertyMenuActorPicker.h
@@ -12,7 +12,7 @@ public:
 	{}
 		SLATE_ARGUMENT( AActor*, InitialActor )
 		SLATE_ARGUMENT( bool, AllowClear )
-		SLATE_ARGUMENT( TSharedPtr< TFilterCollection<const AActor* const> >, ActorFilters )
+		SLATE_ARGUMENT( TSharedPtr< SceneOutliner::FOutlinerFilters >, ActorFilters )
 		SLATE_EVENT( FOnActorSelected, OnSet )
 		SLATE_EVENT( FSimpleDelegate, OnClose )
 		SLATE_EVENT( FSimpleDelegate, OnUseSelected )
@@ -75,7 +75,7 @@ private:
 	/** Whether the asset can be 'None' in this case */
 	bool bAllowClear;
 
-	TSharedPtr< TFilterCollection<const AActor* const> > ActorFilters;
+	TSharedPtr< SceneOutliner::FOutlinerFilters > ActorFilters;
 
 	/** Delegate to call when our object value should be set */
 	FOnActorSelected OnSet;
diff --git a/Engine/Source/Editor/PropertyEditor/Private/UserInterface/PropertyEditor/SPropertySceneOutliner.cpp b/Engine/Source/Editor/PropertyEditor/Private/UserInterface/PropertyEditor/SPropertySceneOutliner.cpp
--- a/Engine/Source/Editor/PropertyEditor/Private/UserInterface/PropertyEditor/SPropertySceneOutliner.cpp
+++ b/Engine/Source/Editor/PropertyEditor/Private/UserInterface/PropertyEditor/SPropertySceneOutliner.cpp
@@ -51,14 +51,11 @@ FReply SPropertySceneOutliner::OnClicked()
 
 TSharedRef<SWidget> SPropertySceneOutliner::OnGenerateSceneOutliner()
 {
-	TSharedPtr<TFilterCollection<const AActor* const> > ActorFilters = NULL;
-	OnGetActorFilters.ExecuteIfBound( ActorFilters );
-
 	FSceneOutlinerModule& SceneOutlinerModule = FModuleManager::Get().LoadModuleChecked<FSceneOutlinerModule>(TEXT("SceneOutliner"));
 
 	FSceneOutlinerInitializationOptions InitOptions;
 	InitOptions.Mode = ESceneOutlinerMode::ActorPicker;
-	InitOptions.ActorFilters = ActorFilters;
+	OnGetActorFilters.ExecuteIfBound( InitOptions.Filters );
 
 	TSharedRef<SWidget> MenuContent = 
 		SNew(SBox)
diff --git a/Engine/Source/Editor/PropertyEditor/Public/PropertyCustomizationHelpers.h b/Engine/Source/Editor/PropertyEditor/Public/PropertyCustomizationHelpers.h
--- a/Engine/Source/Editor/PropertyEditor/Public/PropertyCustomizationHelpers.h
+++ b/Engine/Source/Editor/PropertyEditor/Public/PropertyCustomizationHelpers.h
@@ -6,11 +6,13 @@
 #include "FilterCollection.h"
 #include "SResetToDefaultMenu.h"
 
+namespace SceneOutliner { struct FOutlinerFilters; }
+
 DECLARE_DELEGATE_OneParam(FOnAssetSelected, const class FAssetData& /*AssetData*/);
 DECLARE_DELEGATE_OneParam( FOnGetAllowedClasses, TArray<const UClass*>& );
 DECLARE_DELEGATE_RetVal_OneParam(bool, FOnShouldFilterAsset, const class FAssetData& /*AssetData*/);
 DECLARE_DELEGATE_OneParam( FOnActorSelected, AActor* );
-DECLARE_DELEGATE_OneParam( FOnGetActorFilters, TSharedPtr<TFilterCollection<const AActor* const> >& );
+DECLARE_DELEGATE_OneParam( FOnGetActorFilters, TSharedPtr<SceneOutliner::FOutlinerFilters>& );
 
 namespace PropertyCustomizationHelpers
 {
@@ -25,7 +27,7 @@ namespace PropertyCustomizationHelpers
 	PROPERTYEDITOR_API TSharedRef<SWidget> MakeAssetPickerAnchorButton( FOnGetAllowedClasses OnGetAllowedClasses, FOnAssetSelected OnAssetSelectedFromPicker );
 	PROPERTYEDITOR_API TSharedRef<SWidget> MakeAssetPickerWithMenu( UObject* const InitialObject, const bool AllowClear, const TArray<const UClass*>* const AllowedClasses, FOnShouldFilterAsset OnShouldFilterAsset, FOnAssetSelected OnSet, FSimpleDelegate OnClose );
 	PROPERTYEDITOR_API TSharedRef<SWidget> MakeActorPickerAnchorButton( FOnGetActorFilters OnGetActorFilters, FOnActorSelected OnActorSelectedFromPicker );
-	PROPERTYEDITOR_API TSharedRef<SWidget> MakeActorPickerWithMenu( AActor* const InitialActor, const bool AllowClear, const TSharedPtr< TFilterCollection<const AActor* const> >& ActorFilters, FOnActorSelected OnSet, FSimpleDelegate OnClose, FSimpleDelegate OnUseSelected );
+	PROPERTYEDITOR_API TSharedRef<SWidget> MakeActorPickerWithMenu( AActor* const InitialActor, const bool AllowClear, const TSharedPtr< SceneOutliner::FOutlinerFilters >& ActorFilters, FOnActorSelected OnSet, FSimpleDelegate OnClose, FSimpleDelegate OnUseSelected );
 	PROPERTYEDITOR_API TSharedRef<SWidget> MakeInteractiveActorPicker( FOnGetAllowedClasses OnClearClicked, FOnActorSelected OnActorSelectedFromPicker );
 }
 
diff --git a/Engine/Source/Editor/SceneOutliner/Private/SOutlinerTreeView.cpp b/Engine/Source/Editor/SceneOutliner/Private/SOutlinerTreeView.cpp
--- a/Engine/Source/Editor/SceneOutliner/Private/SOutlinerTreeView.cpp
+++ b/Engine/Source/Editor/SceneOutliner/Private/SOutlinerTreeView.cpp
@@ -114,10 +114,8 @@ namespace SceneOutliner
 					// The folder already exists
 					FFormatNamedArguments Args;
 					Args.Add(TEXT("DragName"), FText::FromString(LeafName));
-					
 					return FDragValidationInfo(FActorDragDropGraphEdOp::ToolTip_IncompatibleGeneric,
-						FText::Format(LOCTEXT("FolderAlreadyExists", "A folder called \"{DragName}\" already exists at this level"), Args));
-
+						FText::Format(LOCTEXT("FolderAlreadyExistsRoot", "A folder called \"{DragName}\" already exists at this level"), Args));
 				}
 				else if (DragFolderPath == DstFolderPath || DstFolderPath.StartsWith(DragFolderPath + "/"))
 				{
@@ -290,7 +288,7 @@ namespace SceneOutliner
 				for (const auto& WeakActor : *DraggedActors)
 				{
 					AActor* DragActor = WeakActor.Get();
-					if (DragActor && DragActor->GetFolderPath() == DropFolder->Path)
+					if (DragActor && !DragActor->GetAttachParentActor() && DragActor->GetFolderPath() == DropFolder->Path)
 					{
 						Args.Add(TEXT("SourceName"), FText::FromString(DragActor->GetActorLabel()));
 						const FText Text = FText::Format(LOCTEXT("MoveToFolderAlreadyAssigned", "{SourceName} is already assigned to {DestName}"), Args);
@@ -376,7 +374,7 @@ namespace SceneOutliner
 		}
 	}
 
-	bool SOutlinerTreeView::ValidateMoveSelectionTo(FName NewParent)
+	bool SOutlinerTreeView::ValidateMoveSelectionTo(FName NewParent, FText* OutValidationText)
 	{
 		auto SceneOutlinerPtr = SceneOutlinerWeak.Pin();
 		if (!SceneOutlinerPtr.IsValid())
@@ -413,6 +411,10 @@ namespace SceneOutliner
 			}
 		}
 
+		if (OutValidationText)
+		{
+			*OutValidationText = ValidationInfo.ValidationText;
+		}
 		return ValidationInfo.IsValid();
 	}
 
diff --git a/Engine/Source/Editor/SceneOutliner/Private/SOutlinerTreeView.h b/Engine/Source/Editor/SceneOutliner/Private/SOutlinerTreeView.h
--- a/Engine/Source/Editor/SceneOutliner/Private/SOutlinerTreeView.h
+++ b/Engine/Source/Editor/SceneOutliner/Private/SOutlinerTreeView.h
@@ -24,7 +24,7 @@ namespace SceneOutliner
 		void FlashHighlightOnItem( FOutlinerTreeItemPtr FlashHighlightOnItem );
 
 		/** Check that it is valid to move the current selection to the specified folder path */
-		bool ValidateMoveSelectionTo(FName NewParent);
+		bool ValidateMoveSelectionTo(FName NewParent, FText* OutValidationText = nullptr);
 
 	protected:
 
diff --git a/Engine/Source/Editor/SceneOutliner/Private/SSceneOutliner.cpp b/Engine/Source/Editor/SceneOutliner/Private/SSceneOutliner.cpp
--- a/Engine/Source/Editor/SceneOutliner/Private/SSceneOutliner.cpp
+++ b/Engine/Source/Editor/SceneOutliner/Private/SSceneOutliner.cpp
@@ -13,6 +13,7 @@
 
 #include "LevelUtils.h"
 #include "ActorEditorUtils.h"
+#include "SceneOutlinerFilterHelpers.h"
 #include "SceneOutlinerFilters.h"
 
 #include "EditorActorFolders.h"
@@ -166,17 +167,10 @@ namespace SceneOutliner
 
 		//Setup the FilterCollection
 		//We use the filter collection provide, otherwise we create our own
-		if (InInitOptions.ActorFilters.IsValid())
-		{
-			CustomFilters = InInitOptions.ActorFilters;
-		}
-		else
-		{
-			CustomFilters = MakeShareable( new ActorFilterCollection() );
-		}
+		Filters = InInitOptions.Filters.IsValid() ? InInitOptions.Filters : MakeShareable(new FOutlinerFilters);
 	
 		SearchBoxFilter->OnChanged().AddSP( this, &SSceneOutliner::FullRefresh );
-		CustomFilters->OnChanged().AddSP( this, &SSceneOutliner::FullRefresh );
+		Filters->OnChanged().AddSP( this, &SSceneOutliner::FullRefresh );
 
 		//Apply custom filters based on global preferences
 		if (InInitOptions.Mode == ESceneOutlinerMode::ActorBrowsing)
@@ -458,7 +452,7 @@ namespace SceneOutliner
 		GEditor->UnregisterForUndo( this );
 
 		SearchBoxFilter->OnChanged().RemoveAll( this );
-		CustomFilters->OnChanged().RemoveAll( this );
+		Filters->OnChanged().RemoveAll( this );
 		
 		FWorldDelegates::LevelAddedToWorld.RemoveAll( this );
 		FWorldDelegates::LevelRemovedFromWorld.RemoveAll( this );
@@ -561,11 +555,11 @@ namespace SceneOutliner
 
 		if ( bShowOnlySelected )
 		{			
-			CustomFilters->Add( SelectedActorFilter );
+			Filters->Add( SelectedActorFilter );
 		}
 		else
 		{
-			CustomFilters->Remove( SelectedActorFilter );
+			Filters->Remove( SelectedActorFilter );
 		}
 	}
 	bool SSceneOutliner::IsShowingOnlySelected() const
@@ -582,6 +576,7 @@ namespace SceneOutliner
 
 		ApplyHideTemporaryActorsFilter(bEnableFlag);
 	}
+
 	void SSceneOutliner::ApplyHideTemporaryActorsFilter(bool bHideTemporaryActors)
 	{
 		if ( !HideTemporaryActorsFilter.IsValid() )
@@ -589,13 +584,14 @@ namespace SceneOutliner
 			HideTemporaryActorsFilter = SceneOutlinerFilters::CreateHideTemporaryActorsFilter();
 		}
 
+		sizeof(TSharedPtr<TDelegateFilter<const AActor*const>>);
 		if ( bHideTemporaryActors )
 		{			
-			CustomFilters->Add( HideTemporaryActorsFilter );
+			Filters->Add( HideTemporaryActorsFilter );
 		}
 		else
 		{
-			CustomFilters->Remove( HideTemporaryActorsFilter );
+			Filters->Remove( HideTemporaryActorsFilter );
 		}
 	}
 	bool SSceneOutliner::IsHidingTemporaryActors() const
@@ -621,11 +617,11 @@ namespace SceneOutliner
 
 		if ( bShowOnlyActorsInCurrentLevel )
 		{			
-			CustomFilters->Add( ShowOnlyActorsInCurrentLevelFilter );
+			Filters->Add( ShowOnlyActorsInCurrentLevelFilter );
 		}
 		else
 		{
-			CustomFilters->Remove( ShowOnlyActorsInCurrentLevelFilter );
+			Filters->Remove( ShowOnlyActorsInCurrentLevelFilter );
 		}
 	}
 	bool SSceneOutliner::IsShowingOnlyCurrentLevel() const
@@ -1008,7 +1004,7 @@ namespace SceneOutliner
 		}
 
 		AActor* Actor = InActorItem->Actor.Get();
-		if (!IsActorDisplayable(Actor) || !CustomFilters->PassesAllFilters(Actor))
+		if (!IsActorDisplayable(Actor) || !Filters->PassesAllFilters(Actor))
 		{
 			return false;
 		}
@@ -1076,7 +1072,7 @@ namespace SceneOutliner
 
 	bool SSceneOutliner::AddFolderToTree(TSharedRef<TOutlinerFolderTreeItem> InFolderItem, bool bIgnoreSearchFilter)
 	{
-		if (!InitOptions.bShowParentTree)
+		if (!InitOptions.bShowParentTree || !Filters->PassesAllFilters(InFolderItem->Path))
 		{
 			return false;
 		}
@@ -1727,7 +1723,7 @@ namespace SceneOutliner
 
 	void SSceneOutliner::FillFoldersSubMenu(FMenuBuilder& MenuBuilder) const
 	{
-		MenuBuilder.AddMenuEntry(LOCTEXT( "CreateNew", "Create New" ), LOCTEXT( "MoveToRoot_ToolTip", "Move to a new folder" ),
+		MenuBuilder.AddMenuEntry(LOCTEXT( "CreateNew", "Create New" ), LOCTEXT( "CreateNew_ToolTip", "Move to a new folder" ),
 			FSlateIcon(FEditorStyle::GetStyleSetName(), "SceneOutliner.NewFolderIcon"), FExecuteAction::CreateSP(this, &SSceneOutliner::CreateFolder));
 
 		if (OutlinerTreeView->ValidateMoveSelectionTo(FName()))
@@ -1736,13 +1732,77 @@ namespace SceneOutliner
 				FSlateIcon(FEditorStyle::GetStyleSetName(), "SceneOutliner.MoveToRoot"), FExecuteAction::CreateSP(this, &SSceneOutliner::MoveSelectionTo, FName()));
 		}
 
+		AddMoveToFolderOutliner(MenuBuilder);
+	}
+
+	void SSceneOutliner::AddMoveToFolderOutliner(FMenuBuilder& MenuBuilder) const
+	{
+		// We don't show this if there is no world (unusual) or if there aren't any folders in the world
+		const int32 NumFoldersInWorld = FActorFolders::Get().GetFoldersForWorld(*RepresentingWorld).Num();
+		if (!RepresentingWorld || !NumFoldersInWorld)
+		{
+			return;
+		}
+
 		// Add a mini scene outliner for choosing an existing folder
-		if (RepresentingWorld && FActorFolders::Get().GetFoldersForWorld(*RepresentingWorld).Num())
+		FSceneOutlinerInitializationOptions InitOptions;
+		InitOptions.bShowHeaderRow = false;
+		InitOptions.bFocusSearchBoxWhenOpened = true;
+		InitOptions.bOnlyShowFolders = true;
+			
 		{
-			FSceneOutlinerInitializationOptions InitOptions;
-			InitOptions.bShowHeaderRow = false;
-			InitOptions.bFocusSearchBoxWhenOpened = true;
-			InitOptions.bOnlyShowFolders = true;
+			// Don't show any folders that are a child of any of the selected folders
+
+			// We use a pointer here to save copying the whole array for every invocation of the filter delegate
+			TSharedRef<TSet<FName>> ExcludedParents(new TSet<FName>());
+			for (const auto& Item : OutlinerTreeView->GetSelectedItems())
+			{
+				if (Item->Type == TOutlinerTreeItem::Folder)
+				{
+					// Can't move the folder into its current parent (it already exists)
+					const auto Folder = StaticCastSharedPtr<TOutlinerFolderTreeItem>(Item);
+					const FName ParentPath = Folder->GetParentPath();
+					ExcludedParents->Add(ParentPath.IsNone() ? Folder->Path : ParentPath);
+				}
+				else if (AActor* Actor = StaticCastSharedPtr<TOutlinerActorTreeItem>(Item)->Actor.Get())
+				{
+					// We exclude actor parent folders if they don't have any sub folders
+					const FName& Folder = Actor->GetFolderPath();
+					if (ExcludedParents->Contains(Folder))
+					{
+						continue;
+					}
+
+					auto HasSubFolders = [](TWeakPtr<TOutlinerTreeItem> WeakItem){
+						auto Item = WeakItem.Pin();
+						return Item.IsValid() && Item->Type == TOutlinerTreeItem::Folder;
+					};
+
+					auto* ParentFolder = FolderToTreeItemMap.Find(Folder);
+					if (ParentFolder && !(*ParentFolder)->ChildItems.ContainsByPredicate(HasSubFolders))
+					{
+						ExcludedParents->Add(Folder);
+					}
+				}
+			}
+
+			if (ExcludedParents->Num())
+			{
+				// Add a filter if necessary
+				auto FilterOutChildFolders = [](FName Path, TSharedRef<TSet<FName>> ExcludedParents){
+					for (const auto& Parent : *ExcludedParents)
+					{
+						if (Path == Parent || FActorFolders::PathIsChildOf(Path.ToString(), Parent.ToString()))
+						{
+							return false;
+						}
+					}
+
+					return true;
+				};
+
+				InitOptions.Filters->AddFilterPredicate(SceneOutliner::FFolderFilterPredicate::CreateStatic(FilterOutChildFolders, ExcludedParents));
+			}
 
 			// Actor selector to allow the user to choose a folder
 			FSceneOutlinerModule& SceneOutlinerModule = FModuleManager::LoadModuleChecked<FSceneOutlinerModule>("SceneOutliner");
@@ -1756,7 +1816,6 @@ namespace SceneOutliner
 					.OnItemPickedDelegate(FOnSceneOutlinerItemPicked::CreateSP(this, &SSceneOutliner::MoveSelectionTo))
 				];
 
-			// Existing folders - it would be good to filter out the existing selected folders from this mini-outliner, but that's currently not possible
 			MenuBuilder.BeginSection(FName(), LOCTEXT("ExistingFolders", "Existing:"));
 			MenuBuilder.AddWidget(MiniSceneOutliner, FText::GetEmpty(), false);
 			MenuBuilder.EndSection();
@@ -1855,8 +1914,15 @@ namespace SceneOutliner
 	{
 		FSlateApplication::Get().DismissAllMenus();
 
-		if (!OutlinerTreeView->ValidateMoveSelectionTo(NewParent))
+		FText ValidationText;
+		if (!OutlinerTreeView->ValidateMoveSelectionTo(NewParent, &ValidationText))
 		{
+			FNotificationInfo Info(ValidationText);
+			Info.ExpireDuration = 3.0f;
+			Info.bUseLargeFont = false;
+			Info.bFireAndForget = true;
+			Info.bUseSuccessFailIcons = true;
+			FSlateNotificationManager::Get().AddNotification(Info)->SetCompletionState(SNotificationItem::CS_Fail);
 			return;
 		}
 
@@ -2207,7 +2273,7 @@ namespace SceneOutliner
 						// Only fire the callback if this actor actually passes our required filters.  It's possible that
 						// some actors displayed are parent actors which don't actually match the requested criteria.
 						if( IsActorDisplayable( FirstSelectedActor ) && 
-							CustomFilters->PassesAllFilters( FirstSelectedActor ) &&
+							Filters->PassesAllFilters( FirstSelectedActor ) &&
 							SearchBoxFilter->PassesFilter( *FirstItem ) )
 						{
 							// Signal that an actor was selected
@@ -2694,7 +2760,7 @@ namespace SceneOutliner
 						// pending kill), in which case we don't want to try to select it
 						if( IsActorDisplayable( Actor ) )
 						{
-							if( CustomFilters->PassesAllFilters( Actor ) && SearchBoxFilter->PassesFilter( *ActorIt.Value() ) )
+							if( Filters->PassesAllFilters( Actor ) && SearchBoxFilter->PassesFilter( *ActorIt.Value() ) )
 							{
 								ActorsToSelect.Add( Actor );
 							}
diff --git a/Engine/Source/Editor/SceneOutliner/Private/SSceneOutliner.h b/Engine/Source/Editor/SceneOutliner/Private/SSceneOutliner.h
--- a/Engine/Source/Editor/SceneOutliner/Private/SSceneOutliner.h
+++ b/Engine/Source/Editor/SceneOutliner/Private/SSceneOutliner.h
@@ -11,7 +11,6 @@ namespace SceneOutliner
 	typedef TArray< FOutlinerTreeItemPtr > FOutlinerData;
 
 	typedef TTextFilter< const TOutlinerTreeItem& > TreeItemTextFilter;
-	typedef TFilterCollection< const AActor* const > ActorFilterCollection;
 
 	typedef TMap< TWeakObjectPtr< AActor >, TSharedRef<TOutlinerActorTreeItem> > FActorToTreeItemMap;
 	typedef TMap< FName, TSharedRef<TOutlinerFolderTreeItem> > FFolderToTreeItemMap;
@@ -517,9 +516,9 @@ namespace SceneOutliner
 		/* Widget containing the filtering text box */
 		TSharedPtr< SSearchBox > FilterTextBoxWidget;
 
-		/** A collection of filters used to filter the displayed actors in the scene outliner */
-		TSharedPtr< ActorFilterCollection > CustomFilters;
-
+		/** A collection of filters used to filter the displayed actors and folders in the scene outliner */
+		TSharedPtr< FOutlinerFilters > Filters;
+		
 		/** The TextFilter attached to the SearchBox widget of the Scene Outliner */
 		TSharedPtr< TreeItemTextFilter > SearchBoxFilter;
 
diff --git a/Engine/Source/Editor/SceneOutliner/Private/SceneOutlinerFilterHelpers.h b/Engine/Source/Editor/SceneOutliner/Private/SceneOutlinerFilterHelpers.h
new file mode 100644
--- /dev/null
+++ b/Engine/Source/Editor/SceneOutliner/Private/SceneOutlinerFilterHelpers.h
@@ -0,0 +1,44 @@
+// Copyright 1998-2014 Epic Games, Inc. All Rights Reserved.
+#pragma once
+
+#include "DelegateFilter.h"
+#include "SharedPointer.h"
+
+namespace SceneOutlinerFilters
+{
+	struct FActorClassPredicate
+	{
+		static bool IsSelectedActor( const AActor* const InActor )
+		{			
+			return ( InActor != NULL && InActor->IsSelected() );
+		}
+
+		static bool ExistsInEditorWorld( const AActor* const InActor )
+		{
+			return (InActor != NULL && (InActor->GetWorld()->WorldType != EWorldType::PIE || GEditor->ObjectsThatExistInEditorWorld.Get(InActor)));
+		}
+	
+		static bool IsInCurrentLevel( const AActor* const InActor )
+		{
+			return (InActor != NULL && InActor->GetLevel() == InActor->GetWorld()->GetCurrentLevel());
+		}
+	};
+
+	TSharedPtr< TDelegateFilter< const AActor* const > > CreateSelectedActorFilter()
+	{
+		return MakeShareable( new TDelegateFilter< const AActor* const >(
+			TDelegateFilter< const AActor* const >::FPredicate::CreateStatic( &FActorClassPredicate::IsSelectedActor ) ) );
+	}
+
+	TSharedPtr< TDelegateFilter< const AActor* const > > CreateHideTemporaryActorsFilter()
+	{
+		return MakeShareable( new TDelegateFilter< const AActor* const >(
+			TDelegateFilter< const AActor* const >::FPredicate::CreateStatic( &FActorClassPredicate::ExistsInEditorWorld ) ) );
+	}
+
+	TSharedPtr< TDelegateFilter< const AActor* const > > CreateIsInCurrentLevelFilter()
+	{
+		return MakeShareable( new TDelegateFilter< const AActor* const >(
+			TDelegateFilter< const AActor* const >::FPredicate::CreateStatic( &FActorClassPredicate::IsInCurrentLevel ) ) );
+	}
+}
\ No newline at end of file
diff --git a/Engine/Source/Editor/SceneOutliner/Private/SceneOutlinerFilters.h b/Engine/Source/Editor/SceneOutliner/Private/SceneOutlinerFilters.h
deleted file mode 100644
diff --git a/Engine/Source/Editor/SceneOutliner/Public/SceneOutlinerFilters.h b/Engine/Source/Editor/SceneOutliner/Public/SceneOutlinerFilters.h
new file mode 100644
--- /dev/null
+++ b/Engine/Source/Editor/SceneOutliner/Public/SceneOutlinerFilters.h
@@ -0,0 +1,150 @@
+// Copyright 1998-2014 Epic Games, Inc. All Rights Reserved.
+
+#pragma once
+
+#include "IFilter.h"
+#include "FilterCollection.h"
+#include "DelegateFilter.h"
+
+namespace SceneOutliner
+{
+	/** Helper typedefs to aid in creation of filter predicates */
+	typedef TDelegateFilter< const AActor* const >::FPredicate FActorFilterPredicate;
+	typedef TDelegateFilter< FName >::FPredicate FFolderFilterPredicate;
+
+	/** Proxy type that is passed to the filters to check whether something should be shown or not */
+	struct FOutlinerFilterProxy
+	{
+		/** The actor to filter, null for folders */
+		const AActor* Actor;
+
+		/** The folder to filter */
+		FName FolderPath;
+
+		/** Constructor for filtering the specified actor */
+		FOutlinerFilterProxy(const AActor* InActor)
+			: Actor(InActor)
+			, FolderPath(Actor->GetFolderPath())
+		{}
+
+		/** Constructor for filtering the specified folder path */
+		FOutlinerFilterProxy(FName InFolder)
+			: FolderPath(InFolder), Actor(nullptr)
+		{}
+	};
+
+	/** Scene outliner filter class. This class abstracts the filtering of both actors and folders and allows for filtering on both types */
+	struct FOutlinerFilters : public TFilterCollection<const FOutlinerFilterProxy&>
+	{
+		/** Check if the specified object matches the specified folder filter */
+		static bool ProxyPassesFilter(const FOutlinerFilterProxy& Proxy, TSharedPtr<IFilter<FName>> Filter)
+		{
+			return Filter->PassesFilter(Proxy.FolderPath);
+		}
+
+		/** Check if the specified object matches the specified actor filter */
+		static bool ProxyPassesFilter(const FOutlinerFilterProxy& Proxy, TSharedPtr<IFilter<const AActor* const>> Filter)
+		{
+			return !Proxy.Actor || Filter->PassesFilter(Proxy.Actor);
+		}
+
+		/** Base class for proxy delegates */
+		struct ProxyBase : IFilter<const FOutlinerFilterProxy&>
+		{
+			/** Virtual functions to test whether the wrapped filter matches the specified filter. Used for removal */
+			virtual bool MatchesWrappedFilter(TSharedPtr<IFilter<FName>>) const					{ return false; }
+			virtual bool MatchesWrappedFilter(TSharedPtr<IFilter<const AActor* const>>) const	{ return false; }
+		};
+
+		/** Derived proxy wrapper. This proxy executes the wrapped delegate with the correct type.
+			Currently supported TFilterTypes are FName and const AActor* const.
+		 */
+		template<class TFilterType>
+		struct ProxyFilter : ProxyBase
+		{
+			typedef TSharedPtr<IFilter<TFilterType>> TFilterPtr;
+
+			/** Constructor */
+			ProxyFilter(TFilterPtr InFilter) : Filter(InFilter) {}
+
+			/** Returns whether the specified Item passes the Filter's restrictions */
+			virtual bool PassesFilter(const FOutlinerFilterProxy& InItem) const OVERRIDE
+			{
+				return ProxyPassesFilter(InItem, Filter);
+			}
+
+			/** Get the on changed event - does nothing for this wrapper */
+			virtual FChangedEvent& OnChanged() OVERRIDE
+			{
+				return ChangedEvent;
+			}
+
+			/** Override the relevant function that matches our type */
+			virtual bool MatchesWrappedFilter(TFilterPtr WrappedFilter) const OVERRIDE
+			{
+				return WrappedFilter == Filter;
+			}
+
+			TFilterPtr Filter;
+			FChangedEvent ChangedEvent;
+		};
+
+		/** Meta template used to extract the filter type from a pointer or raw ptr proxy (to support MakeShareable) */
+		template <class> struct TGetFilterType;
+		template <class X, template<class, ESPMode::Type> class TPtr, ESPMode::Type E> struct TGetFilterType<TPtr<X, E>> { typedef typename X::ItemType Type; };
+		template <class X, template<class> class TPtr> struct TGetFilterType<TPtr<X>> { typedef typename X::ItemType Type; };
+
+		/** Add a custom filter to this collection.
+			Example usage:
+				struct FMyFilterType : public IFilter<const AActor* const> {...}
+				Filters.Add(MakeSareable(new FMyFilterType));
+		*/
+		template<class T>
+		void Add(T InFilter)
+		{
+			typedef ProxyFilter<typename TGetFilterType<T>::Type> TProxyType;
+
+			TSharedPtr<TProxyType> NewProxy(new TProxyType(InFilter));
+			NewProxy->Filter->OnChanged().AddSP(this, &FOutlinerFilters::OnChildFilterChanged);
+			TFilterCollection::Add(NewProxy);
+		}
+		
+		/** Add an actor filter predicate to this filter collection */
+		void AddFilterPredicate(SceneOutliner::FActorFilterPredicate Predicate)
+		{
+			Add(MakeShareable(new TDelegateFilter<const AActor* const>(Predicate)));
+		}
+
+		/** Add a folder filter predicate to this filter collection */
+		void AddFilterPredicate(SceneOutliner::FFolderFilterPredicate Predicate)
+		{
+			Add(MakeShareable(new TDelegateFilter<FName>(Predicate)));
+		}
+
+		/** Remove the specified wrapped filter from this collection */
+		template<class T>
+		void Remove(TSharedPtr<T> InFilter)
+		{
+			TSharedPtr<IFilter<typename T::ItemType>> InBaseFilter = InFilter;
+
+			InFilter->OnChanged().RemoveAll(this);
+
+			TArray<TSharedPtr<ProxyBase>> FiltersToRemove;
+			for (const auto& Child : ChildFilters)
+			{
+				const TSharedPtr<ProxyBase> ProxyChild = StaticCastSharedPtr<ProxyBase>(Child);
+				if (ProxyChild->MatchesWrappedFilter(InBaseFilter))
+				{
+					FiltersToRemove.Add(ProxyChild);
+				}
+			}
+			if (FiltersToRemove.Num())
+			{
+				for (const auto& Filter : FiltersToRemove)
+				{
+					TFilterCollection::Remove(Filter);
+				}
+			}
+		}
+	};
+}
\ No newline at end of file
diff --git a/Engine/Source/Editor/SceneOutliner/Public/SceneOutlinerInitializationOptions.h b/Engine/Source/Editor/SceneOutliner/Public/SceneOutlinerInitializationOptions.h
--- a/Engine/Source/Editor/SceneOutliner/Public/SceneOutlinerInitializationOptions.h
+++ b/Engine/Source/Editor/SceneOutliner/Public/SceneOutlinerInitializationOptions.h
@@ -3,8 +3,7 @@
 
 #include "ModuleInterface.h"
 
-#include "IFilter.h"
-#include "FilterCollection.h"
+#include "SceneOutlinerFilters.h"
 #include "ISceneOutlinerColumn.h"
 #include "ISceneOutliner.h"
 
@@ -78,8 +77,8 @@ public:
 	/** True if we should draw the header row above the tree view */
 	bool bShowHeaderRow;
 
-	/** Optional collection of filters to use when filtering actors in the Scene Outliner */
-	TSharedPtr< TFilterCollection< const AActor* const > > ActorFilters;
+	/** Optional collection of filters to use when filtering in the Scene Outliner */
+	TSharedPtr< SceneOutliner::FOutlinerFilters > Filters;
 
 	/** Whether the Scene Outliner should display parenta actors in a Tree */
 	bool bShowParentTree;
@@ -117,7 +116,7 @@ public:
 	FSceneOutlinerInitializationOptions()
 		: Mode( ESceneOutlinerMode::ActorPicker )
 		, bShowHeaderRow( true )
-		, ActorFilters( new TFilterCollection< const AActor* const >() )
+		, Filters( new SceneOutliner::FOutlinerFilters )
 		, bShowParentTree( true )
 		, bShowSearchBox( true )
 		, CustomColumnFixedWidth( 0.0f )
diff --git a/Engine/Source/Runtime/Core/Public/Misc/FilterCollection.h b/Engine/Source/Runtime/Core/Public/Misc/FilterCollection.h
--- a/Engine/Source/Runtime/Core/Public/Misc/FilterCollection.h
+++ b/Engine/Source/Runtime/Core/Public/Misc/FilterCollection.h
@@ -115,7 +115,7 @@ public:
 	FChangedEvent& OnChanged() { return ChangedEvent; }
 
 
-private:
+protected:
 
 	/** 
 	 *	Called when a child Filter restrictions change and broadcasts the FilterChanged delegate
@@ -126,9 +126,6 @@ private:
 		ChangedEvent.Broadcast();
 	}
 
-
-private:
-
 	/** The array of child filters */
 	TArray< TSharedPtr< IFilter< ItemType > > > ChildFilters;
 
diff --git a/Engine/Source/Runtime/Core/Public/Misc/IFilter.h b/Engine/Source/Runtime/Core/Public/Misc/IFilter.h
--- a/Engine/Source/Runtime/Core/Public/Misc/IFilter.h
+++ b/Engine/Source/Runtime/Core/Public/Misc/IFilter.h
@@ -8,6 +8,8 @@ template< typename ItemType >
 class IFilter
 {
 public:
+	typedef ItemType ItemType;
+
 	virtual ~IFilter(){ }
 
 	/** Returns whether the specified Item passes the Filter's restrictions */
diff --git a/Engine/Source/Runtime/Engine/Private/ActorEditor.cpp b/Engine/Source/Runtime/Engine/Private/ActorEditor.cpp
--- a/Engine/Source/Runtime/Engine/Private/ActorEditor.cpp
+++ b/Engine/Source/Runtime/Engine/Private/ActorEditor.cpp
@@ -490,7 +490,11 @@ const FName& AActor::GetFolderPath() const
 
 void AActor::SetFolderPath(const FName& NewFolderPath)
 {
-	if (NewFolderPath == FolderPath)
+	// Detach the actor if it is attached
+	USceneComponent* RootComp = GetRootComponent();
+	const bool bIsAttached = RootComp  && RootComp->AttachParent;
+
+	if (NewFolderPath == FolderPath && !bIsAttached)
 	{
 		return;
 	}
@@ -501,8 +505,7 @@ void AActor::SetFolderPath(const FName& NewFolderPath)
 	FolderPath = NewFolderPath;
 	
 	// Detach the actor if it is attached
-	USceneComponent* RootComp = GetRootComponent();
-	if (RootComp  && RootComp->AttachParent)
+	if (bIsAttached)
 	{
 		AActor* OldParentActor = RootComp->AttachParent->GetOwner();
 		OldParentActor->Modify();
