Fixed scene outliner setting engine selection states based off internal tree selection (non-ui interaction) events

This addresses TTP#333723 - LIVE: EDITOR: Actors being rotated with alt+rotation widget fail to keep rotating if in a collapsed scene outliner tree

diff --git a/Engine/Source/Editor/SceneOutliner/Private/SSceneOutliner.cpp b/Engine/Source/Editor/SceneOutliner/Private/SSceneOutliner.cpp
--- a/Engine/Source/Editor/SceneOutliner/Private/SSceneOutliner.cpp
+++ b/Engine/Source/Editor/SceneOutliner/Private/SSceneOutliner.cpp
@@ -2254,8 +2254,13 @@ namespace SceneOutliner
 	}
 
 
-	void SSceneOutliner::OnOutlinerTreeSelectionChanged( FOutlinerTreeItemPtr TreeItem, ESelectInfo::Type /*SelectInfo*/ )
+	void SSceneOutliner::OnOutlinerTreeSelectionChanged( FOutlinerTreeItemPtr TreeItem, ESelectInfo::Type SelectInfo )
 	{
+		if (SelectInfo == ESelectInfo::Direct)
+		{
+			return;
+		}
+
 		if( InitOptions.Mode == ESceneOutlinerMode::ActorPicker )
 		{
 			// In actor picking mode, we check to see if we have a selected actor, and if so, fire
