Clarified the message shown when attempting to rename a scene outliner folder to a location that would result in a duplicate
This addresses TTP 331615

diff --git a/Engine/Source/Editor/SceneOutliner/Private/SSceneOutliner.cpp b/Engine/Source/Editor/SceneOutliner/Private/SSceneOutliner.cpp
--- a/Engine/Source/Editor/SceneOutliner/Private/SSceneOutliner.cpp
+++ b/Engine/Source/Editor/SceneOutliner/Private/SSceneOutliner.cpp
@@ -1618,7 +1618,7 @@ namespace SceneOutliner
 
 			if (FolderToTreeItemMap.Find(NewPath))
 			{
-				OutErrorMessage = LOCTEXT("RenameFailed_AlreadyExists", "This folder already exists");
+				OutErrorMessage = LOCTEXT("RenameFailed_AlreadyExists", "A folder with this name already exists at this level");
 				return false;
 			}
 		}
