Fixed drag and drop regression on the scene outliner caused by unhandled messages bubbling up to the tree

This addresses the following TTPs:

334060 - EDITOR: Scene Outliner: If an actor is inside a folder and you click/drag the actor, regardless of where you release the mouse the actor will always be moved from the folder.
334033 - EDITOR: Scene Outliner: Dragging an actor to the folder it is already inside should display a tooltip about moving the actor to the root

diff --git a/Engine/Source/Editor/SceneOutliner/Private/SOutlinerTreeView.cpp b/Engine/Source/Editor/SceneOutliner/Private/SOutlinerTreeView.cpp
--- a/Engine/Source/Editor/SceneOutliner/Private/SOutlinerTreeView.cpp
+++ b/Engine/Source/Editor/SceneOutliner/Private/SOutlinerTreeView.cpp
@@ -76,6 +76,13 @@ namespace SceneOutliner
 		}
 	}
 
+	/** Check to see if the specified drag drop operation is relevant to the scene outliner */
+	bool IsDropOperationApplicable(TSharedPtr<FDragDropOperation> Operation)
+	{
+		// We're only interested in actor drags and folder drags
+		return Operation.IsValid() && (Operation->IsOfType<FActorDragDropGraphEdOp>() || Operation->IsOfType<FFolderDragDropOp>());
+	}
+
 	/** Compute validation information for dropping folder(s) onto another folder */
 	static FDragValidationInfo ValidateDropOnFolder(TSharedRef<SSceneOutliner> SceneOutlinerRef, const FolderArray& DragFolders, FName DropTargetPath)
 	{
@@ -477,14 +484,13 @@ namespace SceneOutliner
 	FReply SOutlinerTreeView::OnDrop(const FGeometry& MyGeometry, const FDragDropEvent& DragDropEvent)
 	{
 		auto SceneOutlinerPtr = SceneOutlinerWeak.Pin();
-		if (!SceneOutlinerPtr.IsValid() ||  !SceneOutlinerPtr->GetInitOptions().bShowParentTree)
-		{
-			// We only allowing dropping to perform attachments when displaying the actors in their parent/child hierarchy
-			return FReply::Unhandled();
-		}
 
-		// Only allow dropping objects when in actor browsing mode
-		if (SceneOutlinerPtr->GetInitOptions().Mode != ESceneOutlinerMode::ActorBrowsing )
+		// Don't handle this if we're not showing a hierarchy, not in browsing mode, or the drop operation is not applicable
+		if (!SceneOutlinerPtr.IsValid() ||
+			!SceneOutlinerPtr->GetInitOptions().bShowParentTree ||
+			SceneOutlinerPtr->GetInitOptions().Mode != ESceneOutlinerMode::ActorBrowsing ||
+			!IsDropOperationApplicable(DragDropEvent.GetOperation())
+			)
 		{
 			return FReply::Unhandled();
 		}
@@ -493,7 +499,7 @@ namespace SceneOutliner
 		const FDragValidationInfo ValidationInfo = ValidateRootDragDropEvent(SceneOutlinerPtr.ToSharedRef(), DragDropEvent);
 		if (!ValidationInfo.IsValid())
 		{
-			return FReply::Unhandled();
+			return FReply::Handled();
 		}
 
 		const FScopedTransaction Transaction(LOCTEXT("UndoAction_DropOnSceneOutliner", "Drop on Scene Outliner"));
@@ -532,10 +538,14 @@ namespace SceneOutliner
 	FReply SSceneOutlinerTreeRow::OnDrop( const FGeometry& MyGeometry, const FDragDropEvent& DragDropEvent )
 	{
 		auto SceneOutlinerPtr = SceneOutlinerWeak.Pin();
-		if (!SceneOutlinerPtr.IsValid() ||  !SceneOutlinerPtr->GetInitOptions().bShowParentTree)
+
+		// Don't handle this if we're not showing a hierarchy, not in browsing mode, or the drop operation is not applicable
+		if (!SceneOutlinerPtr.IsValid() ||
+			!SceneOutlinerPtr->GetInitOptions().bShowParentTree ||
+			SceneOutlinerPtr->GetInitOptions().Mode != ESceneOutlinerMode::ActorBrowsing ||
+			!IsDropOperationApplicable(DragDropEvent.GetOperation())
+			)
 		{
-			// We only allowing dropping to perform attachments when displaying the actors
-			// in their parent/child hierarchy
 			return FReply::Unhandled();
 		}
 
@@ -543,7 +553,8 @@ namespace SceneOutliner
 		const FDragValidationInfo ValidationInfo = ValidateDragDropEvent(SceneOutlinerPtr.ToSharedRef(), DragDropEvent, Item.ToSharedRef());
 		if (!ValidationInfo.IsValid())
 		{
-			return FReply::Unhandled();
+			// Return handled here to stop anything else trying to handle it - the operation is invalid as far as we're concerned
+			return FReply::Handled();
 		}
 
 		const FScopedTransaction Transaction(LOCTEXT("UndoAction_DropOnSceneOutliner", "Drop on Scene Outliner"));
