SceneOutliner: Moved scene outliner settings into their own settings class

diff --git a/Engine/Config/BaseEditorUserSettings.ini b/Engine/Config/BaseEditorUserSettings.ini
--- a/Engine/Config/BaseEditorUserSettings.ini
+++ b/Engine/Config/BaseEditorUserSettings.ini
@@ -50,15 +50,6 @@ bShowFrameRateAndMemory=False
 ; True if we should automatically reimport animset assets when a change to source content is detected
 bAutoReimportAnimSets=False
 
-; Scene Outliner Settings
-bShowOnlySelectedActors=False
-
-; True when the Scene Outliner is hiding temporary/run-time Actors
-bHideTemporaryActors=False
-
-; True when the Scene Outliner is showing only Actors that exist in the current level
-bShowOnlyActorsInCurrentLevel=False
-
 ; Property Matrix
 PropertyMatrix_NumberOfPasteOperationsBeforeWarning=20
 
@@ -93,6 +84,15 @@ bShowFriendlyNames=True
 bExpandConfigurationMenus=False
 
 
+[/Script/SceneOutliner.SceneOutlinerSettings]
+; True when the Scene Outliner is hiding temporary/run-time Actors
+bHideTemporaryActors=False
+; True when the Scene Outliner is showing only Actors that exist in the current level
+bShowOnlyActorsInCurrentLevel=False
+; Scene Outliner Settings
+bShowOnlySelectedActors=False
+
+
 [/Script/UnrealEd.ContentBrowserSettings]
 ; The number of objects to load at once in the Content Browser before displaying a warning about loading many assets
 NumObjectsToLoadBeforeWarning=20
diff --git a/Engine/Source/Editor/SceneOutliner/Private/SSceneOutliner.cpp b/Engine/Source/Editor/SceneOutliner/Private/SSceneOutliner.cpp
--- a/Engine/Source/Editor/SceneOutliner/Private/SSceneOutliner.cpp
+++ b/Engine/Source/Editor/SceneOutliner/Private/SSceneOutliner.cpp
@@ -541,8 +541,10 @@ namespace SceneOutliner
 	void SSceneOutliner::ToggleShowOnlySelected()
 	{
 		const bool bEnableFlag = !IsShowingOnlySelected();
-		GEditor->AccessEditorUserSettings().bShowOnlySelectedActors = bEnableFlag;
-		GEditor->AccessEditorUserSettings().PostEditChange();
+
+		USceneOutlinerSettings* Settings = GetMutableDefault<USceneOutlinerSettings>();
+		Settings->bShowOnlySelectedActors = bEnableFlag;
+		Settings->PostEditChange();
 
 		ApplyShowOnlySelectedFilter(bEnableFlag);
 	}
@@ -564,15 +566,17 @@ namespace SceneOutliner
 	}
 	bool SSceneOutliner::IsShowingOnlySelected() const
 	{
-		return (GEditor->AccessEditorUserSettings().bShowOnlySelectedActors);
+		return GetDefault<USceneOutlinerSettings>()->bShowOnlySelectedActors;
 	}
 
 	// Hide Temporary Actors
 	void SSceneOutliner::ToggleHideTemporaryActors()
 	{
 		const bool bEnableFlag = !IsHidingTemporaryActors();
-		GEditor->AccessEditorUserSettings().bHideTemporaryActors = bEnableFlag;
-		GEditor->AccessEditorUserSettings().PostEditChange();
+
+		USceneOutlinerSettings* Settings = GetMutableDefault<USceneOutlinerSettings>();
+		Settings->bHideTemporaryActors = bEnableFlag;
+		Settings->PostEditChange();
 
 		ApplyHideTemporaryActorsFilter(bEnableFlag);
 	}
@@ -596,15 +600,17 @@ namespace SceneOutliner
 	}
 	bool SSceneOutliner::IsHidingTemporaryActors() const
 	{
-		return (GEditor->AccessEditorUserSettings().bHideTemporaryActors);
+		return GetDefault<USceneOutlinerSettings>()->bHideTemporaryActors;
 	}
 
 	// Show Only Actors In Current Level
 	void SSceneOutliner::ToggleShowOnlyCurrentLevel()
 	{
 		const bool bEnableFlag = !IsShowingOnlyCurrentLevel();
-		GEditor->AccessEditorUserSettings().bShowOnlyActorsInCurrentLevel = bEnableFlag;
-		GEditor->AccessEditorUserSettings().PostEditChange();
+
+		USceneOutlinerSettings* Settings = GetMutableDefault<USceneOutlinerSettings>();
+		Settings->bShowOnlyActorsInCurrentLevel = bEnableFlag;
+		Settings->PostEditChange();
 
 		ApplyShowOnlyCurrentLevelFilter(bEnableFlag);
 	}
@@ -626,7 +632,7 @@ namespace SceneOutliner
 	}
 	bool SSceneOutliner::IsShowingOnlyCurrentLevel() const
 	{
-		return (GEditor->AccessEditorUserSettings().bShowOnlyActorsInCurrentLevel);
+		return GetDefault<USceneOutlinerSettings>()->bShowOnlyActorsInCurrentLevel;
 	}
 
 	/** END FILTERS */
diff --git a/Engine/Source/Editor/SceneOutliner/Private/SceneOutlinerModule.cpp b/Engine/Source/Editor/SceneOutliner/Private/SceneOutlinerModule.cpp
--- a/Engine/Source/Editor/SceneOutliner/Private/SceneOutlinerModule.cpp
+++ b/Engine/Source/Editor/SceneOutliner/Private/SceneOutlinerModule.cpp
@@ -1,20 +1,16 @@
 // Copyright 1998-2014 Epic Games, Inc. All Rights Reserved.
 
+/*=============================================================================
+	SceneOutlinerModule.cpp: Implements the FSceneOutlinerModule class.
+=============================================================================*/
 
 #include "SceneOutlinerPrivatePCH.h"
 #include "SceneOutlinerTreeItems.h"
 #include "ModuleManager.h"
 
-IMPLEMENT_MODULE( FSceneOutlinerModule, SceneOutliner );
 
-void FSceneOutlinerModule::StartupModule()
-{
-}
-
-
-void FSceneOutlinerModule::ShutdownModule()
-{
-}
+/* FSceneOutlinerModule interface
+ *****************************************************************************/
 
 namespace SceneOutliner
 {
@@ -31,6 +27,7 @@ namespace SceneOutliner
 	}
 }
 
+
 TSharedRef< ISceneOutliner > FSceneOutlinerModule::CreateSceneOutliner( const FSceneOutlinerInitializationOptions& InitOptions, const FOnActorPicked& OnActorPickedDelegate ) const
 {
 	auto OnItemPicked = FOnSceneOutlinerItemPicked::CreateStatic( &SceneOutliner::OnSceneOutlinerItemClicked, OnActorPickedDelegate );
@@ -43,3 +40,14 @@ TSharedRef< ISceneOutliner > FSceneOutlinerModule::CreateSceneOutliner( const FS
 		.IsEnabled( FSlateApplication::Get().GetNormalExecutionAttribute() )
 		.OnItemPickedDelegate( OnItemPickedDelegate );
 }
+
+
+/* Class constructors
+ *****************************************************************************/
+
+USceneOutlinerSettings::USceneOutlinerSettings( const class FPostConstructInitializeProperties& PCIP )
+	: Super(PCIP)
+{ }
+
+
+IMPLEMENT_MODULE(FSceneOutlinerModule, SceneOutliner);
diff --git a/Engine/Source/Editor/SceneOutliner/Private/SceneOutlinerPrivatePCH.h b/Engine/Source/Editor/SceneOutliner/Private/SceneOutlinerPrivatePCH.h
--- a/Engine/Source/Editor/SceneOutliner/Private/SceneOutlinerPrivatePCH.h
+++ b/Engine/Source/Editor/SceneOutliner/Private/SceneOutlinerPrivatePCH.h
@@ -8,6 +8,7 @@
 #include "FilterCollection.h"
 #include "ScopedTransaction.h"
 
+#include "SceneOutlinerSettings.h"
 #include "SceneOutlinerActorInfoColumn.h"
 #include "SceneOutlinerGutter.h"
 
diff --git a/Engine/Source/Editor/SceneOutliner/Private/SceneOutlinerSettings.h b/Engine/Source/Editor/SceneOutliner/Private/SceneOutlinerSettings.h
new file mode 100644
--- /dev/null
+++ b/Engine/Source/Editor/SceneOutliner/Private/SceneOutlinerSettings.h
@@ -0,0 +1,32 @@
+// Copyright 1998-2014 Epic Games, Inc. All Rights Reserved.
+
+/*=============================================================================
+	SceneOutlinerSettings.h: Declares the USceneOutlinerSettings class.
+=============================================================================*/
+
+#pragma once
+
+#include "SceneOutlinerSettings.generated.h"
+
+
+/**
+ * Implements the settings for the Scene Outliner.
+ */
+UCLASS(config=EditorUserSettings)
+class USceneOutlinerSettings
+	: public UObject
+{
+	GENERATED_UCLASS_BODY()
+
+	/** True when the Scene Outliner is hiding temporary/run-time Actors */
+	UPROPERTY(config)
+	uint32 bHideTemporaryActors:1;
+
+	/** True when the Scene Outliner is showing only Actors that exist in the current level */
+	UPROPERTY(config)
+	uint32 bShowOnlyActorsInCurrentLevel:1;
+
+	/** True when the Scene Outliner is only displaying selected Actors */
+	UPROPERTY(config)
+	uint32 bShowOnlySelectedActors:1;
+};
diff --git a/Engine/Source/Editor/SceneOutliner/Public/SceneOutlinerModule.h b/Engine/Source/Editor/SceneOutliner/Public/SceneOutlinerModule.h
--- a/Engine/Source/Editor/SceneOutliner/Public/SceneOutlinerModule.h
+++ b/Engine/Source/Editor/SceneOutliner/Public/SceneOutlinerModule.h
@@ -1,4 +1,9 @@
 // Copyright 1998-2014 Epic Games, Inc. All Rights Reserved.
+
+/*=============================================================================
+	SceneOutlinerModule.h: Declares the FSceneOutlinerModule class.
+=============================================================================*/
+
 #pragma once
 
 #include "ModuleInterface.h"
@@ -6,24 +11,15 @@
 #include "ISceneOutliner.h"
 #include "SceneOutlinerInitializationOptions.h"
 
+
 /**
- * Scene Outliner module
+ * Implements the Scene Outliner module.
  */
-class FSceneOutlinerModule : public IModuleInterface
+class FSceneOutlinerModule
+	: public IModuleInterface
 {
-
 public:
 
-	/**
-	 * Called right after the module DLL has been loaded and the module object has been created
-	 */
-	virtual void StartupModule();
-
-	/**
-	 * Called before the module is unloaded, right before the module object is destroyed.
-	 */
-	virtual void ShutdownModule();
-
 	/**
 	 * Creates a scene outliner widget
 	 *
@@ -46,10 +42,16 @@ public:
 	 *
 	 * @return	New scene outliner widget
 	 */
-
 	virtual TSharedRef< class ISceneOutliner > CreateSceneOutliner(
 		const FSceneOutlinerInitializationOptions& InitOptions,
 		const FOnSceneOutlinerItemPicked& OnItemPickedDelegate ) const;
-};
 
+public:
+
+	// Begin IModuleInterface interface
 
+	virtual void StartupModule() { }
+	virtual void ShutdownModule() { }
+
+	// End IModuleInterface interface
+};
diff --git a/Engine/Source/Editor/UnrealEd/Classes/Editor/EditorUserSettings.h b/Engine/Source/Editor/UnrealEd/Classes/Editor/EditorUserSettings.h
--- a/Engine/Source/Editor/UnrealEd/Classes/Editor/EditorUserSettings.h
+++ b/Engine/Source/Editor/UnrealEd/Classes/Editor/EditorUserSettings.h
@@ -277,21 +277,6 @@ class UEditorUserSettings : public UObject
 
 public:
 
-	/**
-	 * Scene Outliner settings
-	 */
-	/** True when the Scene Outliner is only displaying selected Actors */
-	UPROPERTY(config)
-	uint32 bShowOnlySelectedActors:1;
-
-	/** True when the Scene Outliner is hiding temporary/run-time Actors */
-	UPROPERTY(config)
-	uint32 bHideTemporaryActors:1;
-
-	/** True when the Scene Outliner is showing only Actors that exist in the current level */
-	UPROPERTY(config)
-	uint32 bShowOnlyActorsInCurrentLevel:1;
-
 	UPROPERTY(config)
 	int32 MaterialQualityLevel;
 
