Fixed context menu "Rename" option in the scene outliner not working

I'm not sure what caused this regression, but forcing the item to be scrolled into view, and deferring the rename operation to the next tick has fixed it.

diff --git a/Engine/Source/Editor/SceneOutliner/Private/SSceneOutliner.cpp b/Engine/Source/Editor/SceneOutliner/Private/SSceneOutliner.cpp
--- a/Engine/Source/Editor/SceneOutliner/Private/SSceneOutliner.cpp
+++ b/Engine/Source/Editor/SceneOutliner/Private/SSceneOutliner.cpp
@@ -2097,7 +2097,9 @@ namespace SceneOutliner
 
 	void SSceneOutliner::RenameFolder(TSharedRef<TOutlinerFolderTreeItem> Folder)
 	{
-		Folder->RenameRequestEvent.ExecuteIfBound();
+		// Ensure that the item we want to rename is visible in the tree
+		Folder->Flags.RenameWhenInView = true;
+		OutlinerTreeView->RequestScrollIntoView(Folder);
 	}
 
 	void SSceneOutliner::DeleteFolder(TSharedRef<TOutlinerFolderTreeItem> Folder)
