Folder management context menu options are no longer available on Scene Outliners that don't show parent hierarchies

This addresses TTP#335116 - Scene Outliner: Folders: Can be made from the Layers See Contents tab but no option to name and not displayed

diff --git a/Engine/Source/Editor/SceneOutliner/Private/SSceneOutliner.cpp b/Engine/Source/Editor/SceneOutliner/Private/SSceneOutliner.cpp
--- a/Engine/Source/Editor/SceneOutliner/Private/SSceneOutliner.cpp
+++ b/Engine/Source/Editor/SceneOutliner/Private/SSceneOutliner.cpp
@@ -1689,6 +1689,7 @@ namespace SceneOutliner
 		}
 
 		const FSlateIcon NewFolderIcon(FEditorStyle::GetStyleSetName(), "SceneOutliner.NewFolderIcon");
+		// We always create a section here, even if there is no parent so that clients can still extend the menu
 		MenuBuilder.BeginSection("FolderSection", LOCTEXT("ThisFolderSection", "This Folder"));
 		{
 			if (ParentFolder.IsValid())
@@ -1700,26 +1701,30 @@ namespace SceneOutliner
 		}
 		MenuBuilder.EndSection();
 
-		if (SelectedItems.Num() == 0)
+		// Don't add any of these menu items if we're not showing the parent tree
+		if (InitOptions.bShowParentTree)
 		{
-			MenuBuilder.AddMenuEntry(LOCTEXT("CreateFolder", "Create Folder"), FText(), NewFolderIcon, FUIAction(FExecuteAction::CreateSP(this, &SSceneOutliner::CreateFolder)));
-		}
-		else
-		{
-			MenuBuilder.BeginSection("SelectionSection", LOCTEXT("OutlinerSelection", "Current Outliner Selection"));
+			if (SelectedItems.Num() == 0)
 			{
-				MenuBuilder.AddSubMenu(
-					LOCTEXT("MoveActorsToFolder", "Move To Folder"),
-					LOCTEXT("MoveActorsToFolder_Tooltip", "Move selection to a folder"),
-					FNewMenuDelegate::CreateSP(this,  &SSceneOutliner::FillFoldersSubMenu));
-
-				// If we've only got folders selected, show the selection sub menu
-				if (SelectedFolders.Num() == SelectedItems.Num())
+				MenuBuilder.AddMenuEntry(LOCTEXT("CreateFolder", "Create Folder"), FText(), NewFolderIcon, FUIAction(FExecuteAction::CreateSP(this, &SSceneOutliner::CreateFolder)));
+			}
+			else
+			{
+				MenuBuilder.BeginSection("SelectionSection", LOCTEXT("OutlinerSelection", "Current Outliner Selection"));
 				{
 					MenuBuilder.AddSubMenu(
-						LOCTEXT("SelectSubmenu", "Select"),
-						LOCTEXT("SelectSubmenu_Tooltip", "Select the contents of the current selection"),
-						FNewMenuDelegate::CreateSP(this, &SSceneOutliner::FillSelectionSubMenu));
+						LOCTEXT("MoveActorsToFolder", "Move To Folder"),
+						LOCTEXT("MoveActorsToFolder_Tooltip", "Move selection to a folder"),
+						FNewMenuDelegate::CreateSP(this,  &SSceneOutliner::FillFoldersSubMenu));
+
+					// If we've only got folders selected, show the selection sub menu
+					if (SelectedFolders.Num() == SelectedItems.Num())
+					{
+						MenuBuilder.AddSubMenu(
+							LOCTEXT("SelectSubmenu", "Select"),
+							LOCTEXT("SelectSubmenu_Tooltip", "Select the contents of the current selection"),
+							FNewMenuDelegate::CreateSP(this, &SSceneOutliner::FillSelectionSubMenu));
+					}
 				}
 			}
 		}
