Fixed a shutdown crash where it was possible for a singleton instance to disappear before a widget was destroyed

This addresses TTP#335622 - Session/Unreal Frontend: CRASH: Occurs after terminating an editor session.

Submitted for Andrew.Rodham.

diff --git a/Engine/Source/Editor/SceneOutliner/Private/SSceneOutliner.cpp b/Engine/Source/Editor/SceneOutliner/Private/SSceneOutliner.cpp
--- a/Engine/Source/Editor/SceneOutliner/Private/SSceneOutliner.cpp
+++ b/Engine/Source/Editor/SceneOutliner/Private/SSceneOutliner.cpp
@@ -462,9 +462,12 @@ namespace SceneOutliner
 
 		FCoreDelegates::OnActorLabelChanged.RemoveAll(this);
 
-		auto& Folders = FActorFolders::Get();
-		Folders.OnFolderCreate.RemoveAll(this);
-		Folders.OnFolderDelete.RemoveAll(this);
+		if (FActorFolders::IsAvailable())
+		{
+			auto& Folders = FActorFolders::Get();
+			Folders.OnFolderCreate.RemoveAll(this);
+			Folders.OnFolderDelete.RemoveAll(this);
+		}
 	}
 
 	TArray<FOutlinerTreeItemPtr> SSceneOutliner::GetSelectedItems()
diff --git a/Engine/Source/Editor/UnrealEd/Public/EditorActorFolders.h b/Engine/Source/Editor/UnrealEd/Public/EditorActorFolders.h
--- a/Engine/Source/Editor/UnrealEd/Public/EditorActorFolders.h
+++ b/Engine/Source/Editor/UnrealEd/Public/EditorActorFolders.h
@@ -45,7 +45,10 @@ struct UNREALED_API FActorFolders : public FGCObject
 	virtual void AddReferencedObjects(FReferenceCollector& Collector) OVERRIDE;
 	// End FGCObject Interface
 
-	/** Singleton access */
+	/** Check whether the singleton is valid */
+	static bool IsAvailable() { return Singleton != nullptr; }
+
+	/** Singleton access - only valid if IsAvailable() */
 	static FActorFolders& Get();
 
 	/** Initialize the singleton instance - called on Editor Startup */
