Notification now correctly informs the user if the specified folder is not available
#TTP 335344 - EDITOR: Moving actor to existing folder from layers tab results in a blank error notification appearing

#branch UE4

#proj Editor.SceneOutliner

#add Added text to the error notification if the asset is requested to move to a folder that does not exist. The original repro steps for this issue are not available to the user anymore however to ensure that the notification displays the correct message should a separate issue arise I have updated the text.

diff --git a/Engine/Source/Editor/SceneOutliner/Private/SOutlinerTreeView.cpp b/Engine/Source/Editor/SceneOutliner/Private/SOutlinerTreeView.cpp
--- a/Engine/Source/Editor/SceneOutliner/Private/SOutlinerTreeView.cpp
+++ b/Engine/Source/Editor/SceneOutliner/Private/SOutlinerTreeView.cpp
@@ -416,6 +416,10 @@ namespace SceneOutliner
 			{
 				ValidationInfo = ValidateDropOnItem(SceneOutlinerPtr.ToSharedRef(), Folder.ToSharedRef(), &Actors, &Folders);
 			}
+			else
+			{
+				ValidationInfo.ValidationText = FText::Format(LOCTEXT("UnableToFindFolderFailure", "Unable to find folder: {0}"), FText::FromName(NewParent));
+			}
 		}
 
 		if (OutValidationText)
