Fixed editor freeze while deleting big number of actors in Scene Ouliner

diff --git a/Engine/Source/Editor/SceneOutliner/Private/SSceneOutliner.cpp b/Engine/Source/Editor/SceneOutliner/Private/SSceneOutliner.cpp
--- a/Engine/Source/Editor/SceneOutliner/Private/SSceneOutliner.cpp
+++ b/Engine/Source/Editor/SceneOutliner/Private/SSceneOutliner.cpp
@@ -3030,7 +3030,7 @@ namespace SceneOutliner
 							auto* Actor = StaticCastSharedPtr<TOutlinerActorTreeItem>(Item)->Actor.Get();
 							if (Actor)
 							{
-								GEditor->SelectActor(Actor, true, true);
+								GEditor->SelectActor(Actor, true, false);
 							}
 						}
 					}
