Double clicking visibility icon in Scene Outliner no longer focuses on the selected item.
#ttp 337215 - EDITOR: SCENE OUTLINER: Double clicking visibility toggle zooms to selected actor
#branch UE4
#proj Editor.SceneOutliner
#add Added SVisibilityWidget::OnMouseButtonDoubleClick whose main purpose is simply to stop the double click event from propagating up to the parent widget (the SOutlinerTreeView) which was causing the unwanted behaviour.  It also toggles the visibility, so as to behave like two single clicks, as per the request.

diff --git a/Engine/Source/Editor/SceneOutliner/Private/SceneOutlinerGutter.cpp b/Engine/Source/Editor/SceneOutliner/Private/SceneOutlinerGutter.cpp
--- a/Engine/Source/Editor/SceneOutliner/Private/SceneOutlinerGutter.cpp
+++ b/Engine/Source/Editor/SceneOutliner/Private/SceneOutlinerGutter.cpp
@@ -66,6 +66,17 @@ public:
 
 private:
 
+	virtual FReply OnMouseButtonDoubleClick(const FGeometry& InMyGeometry, const FPointerEvent& InMouseEvent)
+	{
+		if (InMouseEvent.GetEffectingButton() == EKeys::LeftMouseButton)
+		{
+			UndoTransaction.Reset(new FScopedTransaction(LOCTEXT("SetActorVisibility", "Set Actor Visibility")));
+			SetIsVisible(!IsVisible());
+		}
+
+		return FReply::Handled();
+	}
+
 	/** Start a new drag/drop operation for this widget */
 	virtual FReply OnDragDetected(const FGeometry& MyGeometry, const FPointerEvent& MouseEvent) override
 	{
