Refreshing scene outliner when current level has changed while "Show Only in Current Level" filter is enabled

diff --git a/Engine/Source/Editor/SceneOutliner/Private/SSceneOutliner.cpp b/Engine/Source/Editor/SceneOutliner/Private/SSceneOutliner.cpp
--- a/Engine/Source/Editor/SceneOutliner/Private/SSceneOutliner.cpp
+++ b/Engine/Source/Editor/SceneOutliner/Private/SSceneOutliner.cpp
@@ -418,6 +418,7 @@ namespace SceneOutliner
 		// Register to find out when actors are added or removed
 		// @todo outliner: Might not catch some cases (see: CALLBACK_ActorPropertiesChange, CALLBACK_LayerChange, CALLBACK_LevelDirtied, CALLBACK_OnActorMoved, CALLBACK_UpdateLevelsForAllActors)
 		FEditorDelegates::MapChange.AddSP( this, &SSceneOutliner::OnMapChange );
+		FEditorDelegates::NewCurrentLevel.AddSP( this, &SSceneOutliner::OnNewCurrentLevel );
 		GEngine->OnLevelActorListChanged().AddSP( this, &SSceneOutliner::FullRefresh );
 		FWorldDelegates::LevelAddedToWorld.AddSP( this, &SSceneOutliner::OnLevelAdded );
 		FWorldDelegates::LevelRemovedFromWorld.AddSP( this, &SSceneOutliner::OnLevelRemoved );
@@ -451,6 +452,7 @@ namespace SceneOutliner
 			USelection::SelectObjectEvent.RemoveAll(this);
 		}
 		FEditorDelegates::MapChange.RemoveAll( this );
+		FEditorDelegates::NewCurrentLevel.RemoveAll( this );
 		GEngine->OnLevelActorListChanged().RemoveAll( this );
 		GEditor->UnregisterForUndo( this );
 
@@ -2729,6 +2731,14 @@ namespace SceneOutliner
 		FullRefresh();
 	}
 
+	void SSceneOutliner::OnNewCurrentLevel()
+	{
+		if (IsShowingOnlyCurrentLevel())
+		{
+			FullRefresh();
+		}
+	}
+
 	void SSceneOutliner::PostUndo(bool bSuccess)
 	{
 		// Refresh our tree in case any changes have been made to the scene that might effect our actor list
diff --git a/Engine/Source/Editor/SceneOutliner/Private/SSceneOutliner.h b/Engine/Source/Editor/SceneOutliner/Private/SSceneOutliner.h
--- a/Engine/Source/Editor/SceneOutliner/Private/SSceneOutliner.h
+++ b/Engine/Source/Editor/SceneOutliner/Private/SSceneOutliner.h
@@ -218,6 +218,9 @@ namespace SceneOutliner
 		/** Called when the map has changed*/
 		void OnMapChange(uint32 MapFlags);
 
+		/** Called when the current level has changed */
+		void OnNewCurrentLevel();
+
 		/** Open a context menu for this scene outliner */
 		TSharedPtr<SWidget> OnOpenContextMenu() const;
 
