Only show the actor class headers in the scene outliner if no other custom column is selected

diff --git a/Engine/Source/Editor/SceneOutliner/Private/SSceneOutliner.cpp b/Engine/Source/Editor/SceneOutliner/Private/SSceneOutliner.cpp
--- a/Engine/Source/Editor/SceneOutliner/Private/SSceneOutliner.cpp
+++ b/Engine/Source/Editor/SceneOutliner/Private/SSceneOutliner.cpp
@@ -187,7 +187,7 @@ namespace SceneOutliner
 		}
 		else
 		{
-			CustomColumn = MakeShareable( new FSceneOutlinerActorInfoColumn( SharedThis( this ), ECustomColumnMode::None ) );
+			CustomColumn = MakeShareable( new FSceneOutlinerActorInfoColumn( SharedThis( this ), ECustomColumnMode::ActorClass ) );
 		}
 
 		TSharedRef< SHeaderRow > HeaderRowWidget =
diff --git a/Engine/Source/Editor/SceneOutliner/Private/SceneOutlinerActorInfoColumn.cpp b/Engine/Source/Editor/SceneOutliner/Private/SceneOutlinerActorInfoColumn.cpp
--- a/Engine/Source/Editor/SceneOutliner/Private/SceneOutlinerActorInfoColumn.cpp
+++ b/Engine/Source/Editor/SceneOutliner/Private/SceneOutlinerActorInfoColumn.cpp
@@ -76,9 +76,10 @@ TSharedRef< SWidget > FSceneOutlinerActorInfoColumn::ConstructClassHyperlink( co
 
 	if (SourceLink.IsValid())
 	{
+		SourceLink->SetVisibility( TAttribute<EVisibility>::Create( TAttribute<EVisibility>::FGetter::CreateSP( this, &FSceneOutlinerActorInfoColumn::GetColumnDataVisibility, true ) ) );
+
 		return SourceLink.ToSharedRef();
 	}
-
 	return SNew( STextBlock )
 		.Text( ObjectClass->GetName() )
 		.HighlightText( SceneOutlinerWeak.Pin().Get(), &ISceneOutliner::GetFilterHighlightText )
