Don't show C++ header link in SceneOutliner

diff --git a/Engine/Source/Editor/SceneOutliner/Private/SceneOutlinerActorInfoColumn.cpp b/Engine/Source/Editor/SceneOutliner/Private/SceneOutlinerActorInfoColumn.cpp
--- a/Engine/Source/Editor/SceneOutliner/Private/SceneOutlinerActorInfoColumn.cpp
+++ b/Engine/Source/Editor/SceneOutliner/Private/SceneOutlinerActorInfoColumn.cpp
@@ -72,7 +72,12 @@ TSharedRef< SWidget > FSceneOutlinerActorInfoColumn::ConstructClassHyperlink( co
 
 	UClass* ObjectClass = Actor->GetClass();
 
-	TSharedPtr<SWidget> SourceLink = FEditorClassUtils::GetSourceLink(ObjectClass, InActor);
+	TSharedPtr<SWidget> SourceLink; 
+	// Only create source link in SO for Blueprints (not native classes, too noisy having header link in SO)
+	if(UBlueprint::GetBlueprintFromClass(ObjectClass) != NULL)
+	{
+		SourceLink = FEditorClassUtils::GetSourceLink(ObjectClass, InActor);
+	}
 
 	if (SourceLink.IsValid())
 	{
