Reworded tooltips for the mobility of actors in the Scene Outliner
#ttp 345816 - NUX: MVP: Unify static/stationary/movable icon in BP editor and SceneOutliner, and explain in tooltip
#branch UE4
#proj Editor.SceneOutliner

diff --git a/Engine/Source/Editor/SceneOutliner/Private/SSceneOutliner.cpp b/Engine/Source/Editor/SceneOutliner/Private/SSceneOutliner.cpp
--- a/Engine/Source/Editor/SceneOutliner/Private/SSceneOutliner.cpp
+++ b/Engine/Source/Editor/SceneOutliner/Private/SSceneOutliner.cpp
@@ -1339,15 +1339,15 @@ namespace SceneOutliner
 
 						if( RootComponent->Mobility == EComponentMobility::Static )
 						{
-							ToolTipText = FText::Format( LOCTEXT( "ComponentMobility_Static", "Static {ActorClassName}" ), Args );
+							ToolTipText = FText::Format( LOCTEXT( "ComponentMobility_Static", "{ActorClassName} with static mobility" ), Args );
 						}
 						else if( RootComponent->Mobility == EComponentMobility::Stationary )
 						{
-							ToolTipText = FText::Format( LOCTEXT( "ComponentMobility_Stationary", "Stationary {ActorClassName}" ), Args );
+							ToolTipText = FText::Format( LOCTEXT( "ComponentMobility_Stationary", "{ActorClassName} with stationary mobility" ), Args );
 						}
 						else if( RootComponent->Mobility == EComponentMobility::Movable )
 						{
-							ToolTipText = FText::Format( LOCTEXT( "ComponentMobility_Movable", "Movable {ActorClassName}" ), Args );
+							ToolTipText = FText::Format( LOCTEXT( "ComponentMobility_Movable", "{ActorClassName} with movable mobility" ), Args );
 						}
 					}
 				}
