[Pull Request # 490] Fix scene outliner Actor Counter
The scene outliner actor count could go out of sync due to a superfluous actor count variable. Removing this in favor of using the actor map works more reliably.

Credit to github user Nojan for finding and fixing this issue.

#github

diff --git a/Engine/Source/Editor/SceneOutliner/Private/SSceneOutliner.cpp b/Engine/Source/Editor/SceneOutliner/Private/SSceneOutliner.cpp
--- a/Engine/Source/Editor/SceneOutliner/Private/SSceneOutliner.cpp
+++ b/Engine/Source/Editor/SceneOutliner/Private/SSceneOutliner.cpp
@@ -139,7 +139,6 @@ namespace SceneOutliner
 		bFullRefresh = true;
 		bNeedsRefresh = true;
 		bIsReentrant = false;
-		TotalActorCount = 0;
 		FilteredActorCount = 0;
 		SortOutlinerTimer = 0.0f;
 		bPendingFocusNextFrame = InInitOptions.bFocusSearchBoxWhenOpened;
@@ -862,7 +861,6 @@ namespace SceneOutliner
 
 	void SSceneOutliner::EmptyTreeItems()
 	{
-		TotalActorCount = 0;
 		FilteredActorCount = 0;
 
 		ActorToTreeItemMap.Reset();
@@ -971,8 +969,6 @@ namespace SceneOutliner
 			}
 			ActorToTreeItemMap.Remove(ActorItemPtr->Actor);
 
-			--TotalActorCount;
-
 			if(!InItem->Flags.IsFilteredOut)
 			{
 				--FilteredActorCount;
@@ -1021,8 +1017,6 @@ namespace SceneOutliner
 			return false;
 		}
 
-		++TotalActorCount;
-
 		// Apply text filter
 		if (!SearchBoxFilter->PassesFilter(*InActorItem))
 		{
@@ -2891,6 +2885,7 @@ namespace SceneOutliner
 	FString SSceneOutliner::GetFilterStatusText() const
 	{
 		const int32 SelectedActorCount = OutlinerTreeView->GetNumItemsSelected();
+		const int32 TotalActorCount = ActorToTreeItemMap.Num();
 
 		if ( !IsFilterActive() )
 		{
@@ -2938,6 +2933,7 @@ namespace SceneOutliner
 
 	bool SSceneOutliner::IsFilterActive() const
 	{
+		const int32 TotalActorCount = ActorToTreeItemMap.Num();
 		return FilterTextBoxWidget->GetText().ToString().Len() > 0 && TotalActorCount != FilteredActorCount;
 	}
 
diff --git a/Engine/Source/Editor/SceneOutliner/Private/SSceneOutliner.h b/Engine/Source/Editor/SceneOutliner/Private/SSceneOutliner.h
--- a/Engine/Source/Editor/SceneOutliner/Private/SSceneOutliner.h
+++ b/Engine/Source/Editor/SceneOutliner/Private/SSceneOutliner.h
@@ -491,9 +491,6 @@ namespace SceneOutliner
 		/** Whether the scene outliner is currently displaying PlayWorld actors */
 		bool bRepresentingPlayWorld;
 
-		/** Total number of displayable actors we've seen, before applying a search filter */
-		int32 TotalActorCount;
-
 		/** Number of actors that passed the search filter */
 		int32 FilteredActorCount;
 
