Added Worlds and Level Blueprints to the scene outliner

This change includes a large rework of the scene outliner to make it more expandable and flexible with what data it supports.
Tree items now support visitation to allow for external logic on custom columns etc without having to type check and cast.
Sorting has been overhauled to be more flexible and efficient. It now supports correct lexical comparison for numbers.

diff --git a/Engine/Source/Editor/Layers/Private/ActorsAssignedToSpecificLayersFilter.h b/Engine/Source/Editor/Layers/Private/ActorsAssignedToSpecificLayersFilter.h
--- a/Engine/Source/Editor/Layers/Private/ActorsAssignedToSpecificLayersFilter.h
+++ b/Engine/Source/Editor/Layers/Private/ActorsAssignedToSpecificLayersFilter.h
@@ -1,13 +1,16 @@
 // Copyright 1998-2014 Epic Games, Inc. All Rights Reserved.
 #pragma once
 
+#include "SceneOutlinerFilters.h"
 
-class FActorsAssignedToSpecificLayersFilter : public IFilter< const AActor* const >, public TSharedFromThis< FActorsAssignedToSpecificLayersFilter >
+class FActorsAssignedToSpecificLayersFilter : public SceneOutliner::FOutlinerFilter, public TSharedFromThis< FActorsAssignedToSpecificLayersFilter >
 {
 public:
 
-	DECLARE_DERIVED_EVENT(FActorsAssignedToSpecificLayersFilter, IFilter< const AActor* const >::FChangedEvent, FChangedEvent);
-	virtual FChangedEvent& OnChanged() override { return ChangedEvent; }
+	/** By default, any types not handled by this filter will fail */
+	FActorsAssignedToSpecificLayersFilter()
+		: FOutlinerFilter(SceneOutliner::EDefaultFilterBehaviour::Fail)
+	{}
 
 	/** 
 	 * Returns whether the specified Item passes the Filter's text restrictions 
@@ -15,7 +18,7 @@ public:
 	 *	@param	InItem	The Item to check 
 	 *	@return			Whether the specified Item passed the filter
 	 */
-	virtual bool PassesFilter( const AActor* const InActor ) const override
+	virtual bool PassesFilter( const AActor* InActor ) const override
 	{
 		if( LayerNames.Num() == 0 )
 		{
@@ -83,7 +86,4 @@ private:
 
 	/**	The list of layer names which actors need to belong to */
 	TArray< FName > LayerNames;
-
-	/**	The event that broadcasts whenever a change occurs to the filter */
-	FChangedEvent ChangedEvent;
 };
\ No newline at end of file
diff --git a/Engine/Source/Editor/Layers/Private/SLayerBrowser.cpp b/Engine/Source/Editor/Layers/Private/SLayerBrowser.cpp
--- a/Engine/Source/Editor/Layers/Private/SLayerBrowser.cpp
+++ b/Engine/Source/Editor/Layers/Private/SLayerBrowser.cpp
@@ -5,6 +5,7 @@
 #include "SLayerBrowser.h"
 #include "SDockTab.h"
 #include "SSearchBox.h"
+#include "Editor/SceneOutliner/Public/SceneOutliner.h"
 
 #define LOCTEXT_NAMESPACE "LayerBrowser"
 
@@ -118,7 +119,9 @@ void SLayerBrowser::Construct(const FArguments& InArgs)
 	//////////////////////////////////////////////////////////////////////////
 	//	Layer Contents Section
 	FSceneOutlinerModule& SceneOutlinerModule = FModuleManager::LoadModuleChecked< FSceneOutlinerModule >("SceneOutliner");
-	FSceneOutlinerInitializationOptions InitOptions;
+	
+	using namespace SceneOutliner;
+	FInitializationOptions InitOptions;
 	{
 		InitOptions.Mode = ESceneOutlinerMode::ActorBrowsing;
 
@@ -126,7 +129,14 @@ void SLayerBrowser::Construct(const FArguments& InArgs)
 		InitOptions.bShowHeaderRow = false;
 		InitOptions.bShowParentTree = false;
 		InitOptions.CustomDelete = FCustomSceneOutlinerDeleteDelegate::CreateSP(this, &SLayerBrowser::RemoveActorsFromSelectedLayer);
-		InitOptions.CustomColumnFactory = FCreateSceneOutlinerColumnDelegate::CreateSP(this, &SLayerBrowser::CreateCustomLayerColumn);
+
+		// Outliner Gutter
+		InitOptions.ColumnMap.Add(FBuiltInColumnTypes::Gutter(), FColumnInfo(EColumnVisibility::Visible, 0) );
+		// Actor Label
+		InitOptions.ColumnMap.Add(FBuiltInColumnTypes::Label(), FColumnInfo(EColumnVisibility::Visible, 10) );
+		// Layer Contents
+		InitOptions.ColumnMap.Add(FSceneOutlinerLayerContentsColumn::GetID(), FColumnInfo(SceneOutliner::EColumnVisibility::Visible, 20,
+			FCreateSceneOutlinerColumn::CreateSP( this, &SLayerBrowser::CreateCustomLayerColumn )) );
 
 		InitOptions.Filters->Add(SelectedLayersFilter);
 	}
diff --git a/Engine/Source/Editor/Layers/Private/SLayerBrowser.h b/Engine/Source/Editor/Layers/Private/SLayerBrowser.h
--- a/Engine/Source/Editor/Layers/Private/SLayerBrowser.h
+++ b/Engine/Source/Editor/Layers/Private/SLayerBrowser.h
@@ -176,7 +176,7 @@ private:
 		SelectedLayerViewModel->RemoveActors( Actors );
 	}
 
-	TSharedRef< ISceneOutlinerColumn > CreateCustomLayerColumn( const TWeakPtr< ISceneOutliner >& SceneOutliner ) const
+	TSharedRef< ISceneOutlinerColumn > CreateCustomLayerColumn( ISceneOutliner& SceneOutliner ) const
 	{
 		return MakeShareable( new FSceneOutlinerLayerContentsColumn( SelectedLayerViewModel.ToSharedRef() ) );
 	}
diff --git a/Engine/Source/Editor/Layers/Private/SceneOutlinerLayerContentsColumn.cpp b/Engine/Source/Editor/Layers/Private/SceneOutlinerLayerContentsColumn.cpp
--- a/Engine/Source/Editor/Layers/Private/SceneOutlinerLayerContentsColumn.cpp
+++ b/Engine/Source/Editor/Layers/Private/SceneOutlinerLayerContentsColumn.cpp
@@ -2,7 +2,7 @@
 
 
 #include "LayersPrivatePCH.h"
-#include "SceneOutlinerTreeItems.h"
+#include "SceneOutliner.h"
 
 #define LOCTEXT_NAMESPACE "SceneOutlinerLayerContentsColumn"
 
@@ -13,59 +13,41 @@ FSceneOutlinerLayerContentsColumn::FSceneOutlinerLayerContentsColumn( const TSha
 
 }
 
+FName FSceneOutlinerLayerContentsColumn::GetID()
+{
+	return FName( "LayerContents" );
+}
 
 FName FSceneOutlinerLayerContentsColumn::GetColumnID()
 {
-	return FName( "LayerContents" );
+	return GetID();
 }
 
 
 SHeaderRow::FColumn::FArguments FSceneOutlinerLayerContentsColumn::ConstructHeaderRowColumn()
 {
 	return SHeaderRow::Column( GetColumnID() )
-		.FillWidth( 48.0f )
+		.FillWidth( 2.f )
 		[
 			SNew( SSpacer )
 		];
 }
 
-const TSharedRef< SWidget > FSceneOutlinerLayerContentsColumn::ConstructRowWidget( const TSharedRef<SceneOutliner::TOutlinerTreeItem> TreeItem )
-{
-	if (TreeItem->Type == SceneOutliner::TOutlinerTreeItem::Actor)
-	{
-		auto ActorTreeItem = StaticCastSharedRef<SceneOutliner::TOutlinerActorTreeItem>(TreeItem);
-		return 	SNew( SButton )
-			.HAlign( HAlign_Center )
-			.VAlign( VAlign_Center )
-			.ButtonStyle( FEditorStyle::Get(), "LayerBrowserButton" )
-			.ContentPadding( 0 )
-			.OnClicked( this, &FSceneOutlinerLayerContentsColumn::OnRemoveFromLayerClicked, ConstCastSharedRef<SceneOutliner::TOutlinerActorTreeItem>(ActorTreeItem)->Actor )
-			.ToolTipText( LOCTEXT("RemoveFromLayerButtonText", "Remove from Layer").ToString() )
-			[
-				SNew( SImage )
-				.Image( FEditorStyle::GetBrush( TEXT( "LayerBrowser.Actor.RemoveFromLayer" ) ) )
-			]
-		;
-	}
-	else
-	{
-		return SNullWidget::NullWidget;
-	}
-}
-
-
-void FSceneOutlinerLayerContentsColumn::PopulateActorSearchStrings( const AActor* const InActor, OUT TArray< FString >& OutSearchStrings ) const
+TSharedRef<SWidget> FSceneOutlinerLayerContentsColumn::ConstructRowWidget(const TWeakObjectPtr< AActor >& Actor )
 {
-	
-}
-
-
-bool FSceneOutlinerLayerContentsColumn::ProvidesSearchStrings()
-{
-	return false;
+	return SNew( SButton )
+		.HAlign( HAlign_Center )
+		.VAlign( VAlign_Center )
+		.ButtonStyle( FEditorStyle::Get(), "LayerBrowserButton" )
+		.ContentPadding( 0 )
+		.OnClicked( this, &FSceneOutlinerLayerContentsColumn::OnRemoveFromLayerClicked, Actor )
+		.ToolTipText( LOCTEXT("RemoveFromLayerButtonText", "Remove from Layer").ToString() )
+		[
+			SNew( SImage )
+			.Image( FEditorStyle::GetBrush( TEXT( "LayerBrowser.Actor.RemoveFromLayer" ) ) )
+		];
 }
 
-
 FReply FSceneOutlinerLayerContentsColumn::OnRemoveFromLayerClicked( const TWeakObjectPtr< AActor > Actor )
 {
 	ViewModel->RemoveActor( Actor );
diff --git a/Engine/Source/Editor/Layers/Private/SceneOutlinerLayerContentsColumn.h b/Engine/Source/Editor/Layers/Private/SceneOutlinerLayerContentsColumn.h
--- a/Engine/Source/Editor/Layers/Private/SceneOutlinerLayerContentsColumn.h
+++ b/Engine/Source/Editor/Layers/Private/SceneOutlinerLayerContentsColumn.h
@@ -19,6 +19,8 @@ public:
 	FSceneOutlinerLayerContentsColumn( const TSharedRef< class FLayerViewModel >& InViewModel );
 
 	virtual ~FSceneOutlinerLayerContentsColumn() {}
+	
+	static FName GetID();
 
 	//////////////////////////////////////////////////////////////////////////
 	// Begin ISceneOutlinerColumn Implementation
@@ -27,15 +29,12 @@ public:
 
 	virtual SHeaderRow::FColumn::FArguments ConstructHeaderRowColumn() override;
 
-	virtual const TSharedRef< SWidget > ConstructRowWidget( const TSharedRef<SceneOutliner::TOutlinerTreeItem> TreeItem ) override;
-
-	virtual bool ProvidesSearchStrings();
-
-	virtual void PopulateActorSearchStrings( const AActor* const InActor, OUT TArray< FString >& OutSearchStrings ) const override;
-	
-	virtual bool SupportsSorting() const override { return false; }
-
-	virtual void SortItems(TArray<TSharedPtr<SceneOutliner::TOutlinerTreeItem>>& RootItems, const EColumnSortMode::Type SortMode) const override {}
+	virtual const TSharedRef< SWidget > ConstructRowWidget( SceneOutliner::FTreeItemRef TreeItem, const STableRow<SceneOutliner::FTreeItemPtr>& Row ) override
+	{
+		FColumnGenerator Generator(*this);
+		TreeItem->Visit(Generator);
+		return Generator.Widget.ToSharedRef();
+	}
 
 	// End ISceneOutlinerColumn Implementation
 	//////////////////////////////////////////////////////////////////////////
@@ -44,7 +43,19 @@ private:
 
 	FReply OnRemoveFromLayerClicked( const TWeakObjectPtr< AActor > Actor );
 
-
+	TSharedRef<SWidget> ConstructRowWidget(const TWeakObjectPtr< AActor >& Actor );
+	
+	struct FColumnGenerator : SceneOutliner::FColumnGenerator
+	{
+		FSceneOutlinerLayerContentsColumn& Column;
+		FColumnGenerator(FSceneOutlinerLayerContentsColumn& InColumn) : Column(InColumn) {}
+
+		virtual TSharedRef<SWidget> GenerateWidget(SceneOutliner::FActorTreeItem& ActorItem) const override
+		{
+			return Column.ConstructRowWidget(ActorItem.Actor);
+		}
+	};
+	friend FColumnGenerator;
 private:
 
 	/**	 */
diff --git a/Engine/Source/Editor/LevelEditor/Private/LevelEditorActions.cpp b/Engine/Source/Editor/LevelEditor/Private/LevelEditorActions.cpp
--- a/Engine/Source/Editor/LevelEditor/Private/LevelEditorActions.cpp
+++ b/Engine/Source/Editor/LevelEditor/Private/LevelEditorActions.cpp
@@ -16,7 +16,6 @@
 #include "Editor/PropertyEditor/Public/PropertyEditorModule.h"
 #include "Editor/ContentBrowser/Public/ContentBrowserModule.h"
 #include "Editor/MainFrame/Public/MainFrame.h"
-#include "Editor/SceneOutliner/Public/SceneOutlinerModule.h"
 #include "Editor/UnrealEd/Private/GeomFitUtils.h"
 #include "Editor/UnrealEd/Public/BSPOps.h"
 #include "Editor/LevelEditor/Public/DlgDeltaTransform.h"
diff --git a/Engine/Source/Editor/LevelEditor/Private/LevelEditorContextMenu.cpp b/Engine/Source/Editor/LevelEditor/Private/LevelEditorContextMenu.cpp
--- a/Engine/Source/Editor/LevelEditor/Private/LevelEditorContextMenu.cpp
+++ b/Engine/Source/Editor/LevelEditor/Private/LevelEditorContextMenu.cpp
@@ -13,7 +13,7 @@
 #include "Toolkits/AssetEditorManager.h"
 #include "SLevelEditor.h"
 #include "SLevelViewport.h"
-#include "Editor/SceneOutliner/Public/SceneOutlinerModule.h"
+#include "Editor/SceneOutliner/Public/SceneOutliner.h"
 #include "DelegateFilter.h"
 #include "AssetData.h"
 #include "DebuggerCommands.h"
@@ -743,7 +743,7 @@ void FLevelEditorContextMenuImpl::FillActorMenu( FMenuBuilder& MenuBuilder )
 		}
 	};
 
-	FSceneOutlinerInitializationOptions InitOptions;
+	SceneOutliner::FInitializationOptions InitOptions;
 	{
 		InitOptions.Mode = ESceneOutlinerMode::ActorPicker;			
 		InitOptions.bShowHeaderRow = false;
diff --git a/Engine/Source/Editor/LevelEditor/Private/LevelEditorToolBar.cpp b/Engine/Source/Editor/LevelEditor/Private/LevelEditorToolBar.cpp
--- a/Engine/Source/Editor/LevelEditor/Private/LevelEditorToolBar.cpp
+++ b/Engine/Source/Editor/LevelEditor/Private/LevelEditorToolBar.cpp
@@ -6,7 +6,7 @@
 #include "LevelEditorActions.h"
 #include "Editor/UnrealEd/Public/SourceCodeNavigation.h"
 #include "Editor/UnrealEd/Public/Kismet2/DebuggerCommands.h"
-#include "Editor/SceneOutliner/Public/SceneOutlinerModule.h"
+#include "Editor/SceneOutliner/Public/SceneOutliner.h"
 #include "DelegateFilter.h"
 #include "EditorViewportCommands.h"
 #include "SScalabilitySettings.h"
@@ -1873,7 +1873,7 @@ TSharedRef< SWidget > FLevelEditorToolBar::GenerateMatineeMenuContent( TSharedRe
 	FMenuBuilder MenuBuilder( bShouldCloseWindowAfterMenuSelection, InCommandList );
 
 	// We can't build a list of Matinees while the current World is a PIE world.
-	FSceneOutlinerInitializationOptions InitOptions;
+	SceneOutliner::FInitializationOptions InitOptions;
 	{
 		InitOptions.Mode = ESceneOutlinerMode::ActorPicker;
 
diff --git a/Engine/Source/Editor/LevelEditor/Private/SLevelEditor.cpp b/Engine/Source/Editor/LevelEditor/Private/SLevelEditor.cpp
--- a/Engine/Source/Editor/LevelEditor/Private/SLevelEditor.cpp
+++ b/Engine/Source/Editor/LevelEditor/Private/SLevelEditor.cpp
@@ -18,7 +18,7 @@
 #include "SLevelEditorModeContent.h"
 #include "SLevelEditorBuildAndSubmit.h"
 #include "Editor/UnrealEd/Public/Kismet2/DebuggerCommands.h"
-#include "Editor/SceneOutliner/Public/SceneOutlinerModule.h"
+#include "Editor/SceneOutliner/Public/SceneOutliner.h"
 #include "Editor/Layers/Public/LayersModule.h"
 #include "Editor/Levels/Public/LevelsModule.h"
 #include "Editor/WorldBrowser/Public/WorldBrowserModule.h"
@@ -607,13 +607,13 @@ TSharedRef<SDockTab> SLevelEditor::SpawnLevelEditorTab( const FSpawnTabArgs& Arg
 	}
 	else if( TabIdentifier == TEXT("LevelEditorSceneOutliner") )
 	{
-		FSceneOutlinerInitializationOptions InitOptions;
+		SceneOutliner::FInitializationOptions InitOptions;
 		InitOptions.Mode = ESceneOutlinerMode::ActorBrowsing;
 		{
 			TWeakPtr<SLevelEditor> WeakLevelEditor = SharedThis(this);
 			InitOptions.DefaultMenuExtender = MakeShareable(new FExtender);
 			InitOptions.DefaultMenuExtender->AddMenuExtension(
-				"FolderSection", EExtensionHook::Before, GetLevelEditorActions(),
+				"MainSection", EExtensionHook::Before, GetLevelEditorActions(),
 				FMenuExtensionDelegate::CreateStatic([](FMenuBuilder& MenuBuilder, TWeakPtr<SLevelEditor> InWeakLevelEditor){
 					// Only extend the menu if we have actors selected
 					if (GEditor->GetSelectedActors()->Num())
diff --git a/Engine/Source/Editor/LevelEditor/Private/SLevelViewport.cpp b/Engine/Source/Editor/LevelEditor/Private/SLevelViewport.cpp
--- a/Engine/Source/Editor/LevelEditor/Private/SLevelViewport.cpp
+++ b/Engine/Source/Editor/LevelEditor/Private/SLevelViewport.cpp
@@ -19,14 +19,12 @@
 #include "Editor/UnrealEd/Public/DragAndDrop/AssetDragDropOp.h"
 #include "Editor/UnrealEd/Public/DragAndDrop/ExportTextDragDropOp.h"
 #include "Editor/UnrealEd/Public/DragAndDrop/BrushBuilderDragDropOp.h"
-#include "Editor/SceneOutliner/Public/SceneOutlinerModule.h"
-#include "Editor/SceneOutliner/Public/ISceneOutlinerColumn.h"
+#include "Editor/SceneOutliner/Public/SceneOutliner.h"
 #include "ScopedTransaction.h"
 #include "LevelUtils.h"
 #include "HighresScreenshotUI.h"
 #include "SCaptureRegionWidget.h"
 #include "ISettingsModule.h"
-#include "SceneOutlinerTreeItems.h"
 #include "BufferVisualizationData.h"
 #include "EditorViewportCommands.h"
 #include "Runtime/Engine/Classes/Engine/UserInterfaceSettings.h"
@@ -2168,7 +2166,7 @@ float SLevelViewport::GetActorLockSceneOutlinerColumnWidth()
 	return 18.0f;	// 16.0f for the icons and 2.0f padding
 }
 
-TSharedRef< ISceneOutlinerColumn > SLevelViewport::CreateActorLockSceneOutlinerColumn( const TWeakPtr< ISceneOutliner >& SceneOutliner ) const
+TSharedRef< ISceneOutlinerColumn > SLevelViewport::CreateActorLockSceneOutlinerColumn( ISceneOutliner& SceneOutliner ) const
 {
 	/**
 	 * A custom column for the SceneOutliner which shows whether an actor is locked to a viewport
@@ -2199,25 +2197,46 @@ TSharedRef< ISceneOutlinerColumn > SLevelViewport::CreateActorLockSceneOutlinerC
 		virtual SHeaderRow::FColumn::FArguments ConstructHeaderRowColumn() override
 		{
 			return SHeaderRow::Column( GetColumnID() )
+				.FixedWidth(SLevelViewport::GetActorLockSceneOutlinerColumnWidth())
 				[
 					SNew( SSpacer )
 				];
 		}
 
-		virtual const TSharedRef< SWidget > ConstructRowWidget( const TSharedRef<SceneOutliner::TOutlinerTreeItem> TreeItem ) override
+		virtual const TSharedRef< SWidget > ConstructRowWidget( SceneOutliner::FTreeItemRef TreeItem, const STableRow<SceneOutliner::FTreeItemPtr>& InRow ) override
 		{
-			if (TreeItem->Type == SceneOutliner::TOutlinerTreeItem::Actor)
+			struct FConstructWidget : SceneOutliner::FColumnGenerator
 			{
-				bool bLocked = Viewport->IsActorLocked(StaticCastSharedRef<SceneOutliner::TOutlinerActorTreeItem>(TreeItem)->Actor);
+				const SLevelViewport* Viewport;
+				FConstructWidget(const SLevelViewport* InViewport) : Viewport(InViewport) {}
 
-				return SNew(SBox)
-					.WidthOverride(SLevelViewport::GetActorLockSceneOutlinerColumnWidth())
-					.Padding(FMargin(2.0f, 0.0f, 0.0f, 0.0f))
-					[
-						SNew(SImage)
-						.Image(FEditorStyle::GetBrush(bLocked ? "PropertyWindow.Locked" : "PropertyWindow.Unlocked"))
-						.ColorAndOpacity(bLocked ? FLinearColor::White : FLinearColor(1.0f, 1.0f, 1.0f, 0.5f))
-					];
+				virtual TSharedRef<SWidget> GenerateWidget(SceneOutliner::FActorTreeItem& ActorItem) const override
+				{
+					AActor* Actor = ActorItem.Actor.Get();
+					if (!Actor)
+					{
+						return SNullWidget::NullWidget;
+					}
+
+					const bool bLocked = Viewport->IsActorLocked(Actor);
+
+					return SNew(SBox)
+						.WidthOverride(SLevelViewport::GetActorLockSceneOutlinerColumnWidth())
+						.Padding(FMargin(2.0f, 0.0f, 0.0f, 0.0f))
+						[
+							SNew(SImage)
+							.Image(FEditorStyle::GetBrush(bLocked ? "PropertyWindow.Locked" : "PropertyWindow.Unlocked"))
+							.ColorAndOpacity(bLocked ? FLinearColor::White : FLinearColor(1.0f, 1.0f, 1.0f, 0.5f))
+						];	
+				}
+			};
+
+			FConstructWidget Visitor(Viewport);
+			TreeItem->Visit(Visitor);
+
+			if (Visitor.Widget.IsValid())
+			{
+				return Visitor.Widget.ToSharedRef();	
 			}
 			else
 			{
@@ -2225,14 +2244,6 @@ TSharedRef< ISceneOutlinerColumn > SLevelViewport::CreateActorLockSceneOutlinerC
 			}
 		}
 
-		virtual bool ProvidesSearchStrings() { return false; }
-
-		virtual void PopulateActorSearchStrings( const AActor* const InActor, OUT TArray< FString >& OutSearchStrings ) const override {}
-	
-		virtual bool SupportsSorting() const override { return false; }
-
-		virtual void SortItems(TArray<TSharedPtr<SceneOutliner::TOutlinerTreeItem>>& RootItems, const EColumnSortMode::Type SortMode) const override {}
-
 		// End ISceneOutlinerColumn Implementation
 		//////////////////////////////////////////////////////////////////////////
 
diff --git a/Engine/Source/Editor/LevelEditor/Private/SLevelViewportToolBar.cpp b/Engine/Source/Editor/LevelEditor/Private/SLevelViewportToolBar.cpp
--- a/Engine/Source/Editor/LevelEditor/Private/SLevelViewportToolBar.cpp
+++ b/Engine/Source/Editor/LevelEditor/Private/SLevelViewportToolBar.cpp
@@ -9,7 +9,7 @@
 #include "LevelViewportActions.h"
 #include "LevelEditorActions.h"
 #include "Layers/ILayers.h"
-#include "Editor/SceneOutliner/Public/SceneOutlinerModule.h"
+#include "Editor/SceneOutliner/Public/SceneOutliner.h"
 #include "DelegateFilter.h"
 #include "Editor/SceneOutliner/Public/ISceneOutlinerColumn.h"
 #include "DeviceProfileServices.h"
@@ -396,22 +396,20 @@ static void OnGenerateActorLockingMenuSection( TWeakPtr<SLevelViewport> Viewport
 
 	// Creates a scene outliner to fill the menu. The outliner lists only actors of class T and its children
 
-	FSceneOutlinerInitializationOptions InitOptions;
+	using namespace SceneOutliner;
+
+	SceneOutliner::FInitializationOptions InitOptions;
 	InitOptions.Mode = ESceneOutlinerMode::ActorPicker;			
 	InitOptions.bShowHeaderRow = false;
-	InitOptions.CustomColumnFixedWidth = SLevelViewport::GetActorLockSceneOutlinerColumnWidth();
-	InitOptions.CustomColumnFactory = FCreateSceneOutlinerColumnDelegate::CreateSP( Viewport.Pin().ToSharedRef(), &SLevelViewport::CreateActorLockSceneOutlinerColumn );
 
-	// Predicate for selecting actors by class
-	struct FActorClassPredicate
-	{
-		static bool IsLockableActor( const AActor* InActor )
-		{			
-			return InActor != NULL && InActor->IsA(ActorClass::StaticClass()) && !InActor->IsPendingKill();
-		}
-	};
+	InitOptions.ColumnMap.Add(FBuiltInColumnTypes::Label(), FColumnInfo(EColumnVisibility::Visible, 0) );
+	InitOptions.ColumnMap.Add("LockedToViewport", FColumnInfo(EColumnVisibility::Visible, 10,
+		FCreateSceneOutlinerColumn::CreateSP( Viewport.Pin().ToSharedRef(), &SLevelViewport::CreateActorLockSceneOutlinerColumn )) );
 
-	InitOptions.Filters->AddFilterPredicate( SceneOutliner::FActorFilterPredicate::CreateStatic(&FActorClassPredicate::IsLockableActor) );
+	/** Only show actors that are of the templated type */
+	InitOptions.Filters->AddFilterPredicate( FActorFilterPredicate::CreateStatic([](const AActor* InActor){
+		return InActor && InActor->IsA(ActorClass::StaticClass()) && !InActor->IsPendingKill();
+	}) );
 
 	// Create an outliner with the options from above. It sits in a box with a max height limit to stop it getting too tall when lots of actors are added.
 	FSceneOutlinerModule& SceneOutlinerModule = FModuleManager::LoadModuleChecked<FSceneOutlinerModule>("SceneOutliner");
diff --git a/Engine/Source/Editor/LevelEditor/Public/SLevelViewport.h b/Engine/Source/Editor/LevelEditor/Public/SLevelViewport.h
--- a/Engine/Source/Editor/LevelEditor/Public/SLevelViewport.h
+++ b/Engine/Source/Editor/LevelEditor/Public/SLevelViewport.h
@@ -210,7 +210,7 @@ public:
 	/**
 	 * @return a new custom column for a scene outliner that indicates whether each actor is locked to this viewport
 	 */
-	TSharedRef< class ISceneOutlinerColumn > CreateActorLockSceneOutlinerColumn( const TWeakPtr< class ISceneOutliner >& SceneOutliner ) const;
+	TSharedRef< class ISceneOutlinerColumn > CreateActorLockSceneOutlinerColumn( class ISceneOutliner& SceneOutliner ) const;
 
 	/** Called when Preview Selected Cameras preference is changed.*/
 	void OnPreviewSelectedCamerasChange();
diff --git a/Engine/Source/Editor/PropertyEditor/Private/Presentation/PropertyEditor/PropertyEditor.cpp b/Engine/Source/Editor/PropertyEditor/Private/Presentation/PropertyEditor/PropertyEditor.cpp
--- a/Engine/Source/Editor/PropertyEditor/Private/Presentation/PropertyEditor/PropertyEditor.cpp
+++ b/Engine/Source/Editor/PropertyEditor/Private/Presentation/PropertyEditor/PropertyEditor.cpp
@@ -370,7 +370,7 @@ void FPropertyEditor::OnGetActorFiltersForSceneOutliner( TSharedPtr<SceneOutline
 {
 	struct Local
 	{
-		static bool IsFilteredActor( const AActor* const Actor, TSharedRef<FPropertyEditor> PropertyEditor )
+		static bool IsFilteredActor( const AActor* Actor, TSharedRef<FPropertyEditor> PropertyEditor )
 		{
 			const TSharedRef<FPropertyNode> PropertyNode = PropertyEditor->GetPropertyNode();
 			UProperty* NodeProperty = PropertyNode->GetProperty();
@@ -384,7 +384,7 @@ void FPropertyEditor::OnGetActorFiltersForSceneOutliner( TSharedPtr<SceneOutline
 		}
 	};
 
-	OutFilters->Add( MakeShareable( new TDelegateFilter< const AActor* const >( TDelegateFilter< const AActor* const >::FPredicate::CreateStatic( &Local::IsFilteredActor, AsShared() ) ) ) );
+	OutFilters->AddFilterPredicate( SceneOutliner::FActorFilterPredicate::CreateStatic( &Local::IsFilteredActor, AsShared() ) );
 }
 
 void FPropertyEditor::OnResetToDefault()
diff --git a/Engine/Source/Editor/PropertyEditor/Private/Presentation/PropertyEditor/PropertyEditor.h b/Engine/Source/Editor/PropertyEditor/Private/Presentation/PropertyEditor/PropertyEditor.h
--- a/Engine/Source/Editor/PropertyEditor/Private/Presentation/PropertyEditor/PropertyEditor.h
+++ b/Engine/Source/Editor/PropertyEditor/Private/Presentation/PropertyEditor/PropertyEditor.h
@@ -1,7 +1,7 @@
 // Copyright 1998-2014 Epic Games, Inc. All Rights Reserved.
 #pragma once
 
-#include "Editor/SceneOutliner/Public/SceneOutlinerFilters.h"
+#include "Editor/SceneOutliner/Public/SceneOutliner.h"
 
 class FPropertyNode;
 class IPropertyUtilities;
diff --git a/Engine/Source/Editor/PropertyEditor/Private/PropertyCustomizationHelpers.cpp b/Engine/Source/Editor/PropertyEditor/Private/PropertyCustomizationHelpers.cpp
--- a/Engine/Source/Editor/PropertyEditor/Private/PropertyCustomizationHelpers.cpp
+++ b/Engine/Source/Editor/PropertyEditor/Private/PropertyCustomizationHelpers.cpp
@@ -225,13 +225,13 @@ namespace PropertyCustomizationHelpers
 			.OnActorSelected( OnActorSelectedFromPicker );
 	}
 
-	TSharedRef<SWidget> MakeActorPickerWithMenu( AActor* const InitialActor, const bool AllowClear, const TSharedPtr< SceneOutliner::FOutlinerFilters >& ActorFilters, FOnActorSelected OnSet, FSimpleDelegate OnClose, FSimpleDelegate OnUseSelected )
+	TSharedRef<SWidget> MakeActorPickerWithMenu( AActor* const InitialActor, const bool AllowClear, FOnShouldFilterActor ActorFilter, FOnActorSelected OnSet, FSimpleDelegate OnClose, FSimpleDelegate OnUseSelected )
 	{
 		return 
 			SNew( SPropertyMenuActorPicker )
 			.InitialActor(InitialActor)
 			.AllowClear(AllowClear)
-			.ActorFilters(ActorFilters)
+			.ActorFilter(ActorFilter)
 			.OnSet(OnSet)
 			.OnClose(OnClose)
 			.OnUseSelected(OnUseSelected);
diff --git a/Engine/Source/Editor/PropertyEditor/Private/UserInterface/PropertyEditor/SPropertyEditorAsset.cpp b/Engine/Source/Editor/PropertyEditor/Private/UserInterface/PropertyEditor/SPropertyEditorAsset.cpp
--- a/Engine/Source/Editor/PropertyEditor/Private/UserInterface/PropertyEditor/SPropertyEditorAsset.cpp
+++ b/Engine/Source/Editor/PropertyEditor/Private/UserInterface/PropertyEditor/SPropertyEditorAsset.cpp
@@ -8,7 +8,6 @@
 #include "AssetThumbnail.h"
 #include "Editor/ContentBrowser/Public/ContentBrowserModule.h"
 #include "Runtime/AssetRegistry/Public/AssetData.h"
-#include "Editor/SceneOutliner/Public/SceneOutlinerModule.h"
 #include "PropertyHandleImpl.h"
 #include "DelegateFilter.h"
 #include "Developer/AssetTools/Public/AssetToolsModule.h"
@@ -325,12 +324,9 @@ TSharedRef<SWidget> SPropertyEditorAsset::OnGetMenuContent()
 
 	if(bIsActor)
 	{
-		TSharedPtr<SceneOutliner::FOutlinerFilters> ActorFilters = MakeShareable( new SceneOutliner::FOutlinerFilters );
-		ActorFilters->Add( MakeShareable( new TDelegateFilter< const AActor* const >( TDelegateFilter< const AActor* const >::FPredicate::CreateSP( this, &SPropertyEditorAsset::IsFilteredActor ) ) ) );
-
 		return PropertyCustomizationHelpers::MakeActorPickerWithMenu(Cast<AActor>(Value.Object),
 																	 bAllowClear,
-																	 ActorFilters,
+																	 FOnShouldFilterActor::CreateSP( this, &SPropertyEditorAsset::IsFilteredActor ),
 																	 FOnActorSelected::CreateSP( this, &SPropertyEditorAsset::OnActorSelected),
 																	 FSimpleDelegate::CreateSP( this, &SPropertyEditorAsset::CloseComboButton ),
 																	 FSimpleDelegate::CreateSP( this, &SPropertyEditorAsset::OnUse ) );
diff --git a/Engine/Source/Editor/PropertyEditor/Private/UserInterface/PropertyEditor/SPropertyMenuActorPicker.cpp b/Engine/Source/Editor/PropertyEditor/Private/UserInterface/PropertyEditor/SPropertyMenuActorPicker.cpp
--- a/Engine/Source/Editor/PropertyEditor/Private/UserInterface/PropertyEditor/SPropertyMenuActorPicker.cpp
+++ b/Engine/Source/Editor/PropertyEditor/Private/UserInterface/PropertyEditor/SPropertyMenuActorPicker.cpp
@@ -4,7 +4,7 @@
 #include "SPropertyMenuActorPicker.h"
 #include "AssetRegistryModule.h"
 #include "DelegateFilter.h"
-#include "Editor/SceneOutliner/Public/SceneOutlinerModule.h"
+#include "Editor/SceneOutliner/Public/SceneOutliner.h"
 #include "PropertyEditorAssetConstants.h"
 
 #define LOCTEXT_NAMESPACE "PropertyEditor"
@@ -13,7 +13,7 @@ void SPropertyMenuActorPicker::Construct( const FArguments& InArgs )
 {
 	CurrentActor = InArgs._InitialActor;
 	bAllowClear = InArgs._AllowClear;
-	ActorFilters = InArgs._ActorFilters;
+	ActorFilter = InArgs._ActorFilter;
 	OnSet = InArgs._OnSet;
 	OnClose = InArgs._OnClose;
 	OnUseSelected = InArgs._OnUseSelected;
@@ -71,9 +71,9 @@ void SPropertyMenuActorPicker::Construct( const FArguments& InArgs )
 
 		FSceneOutlinerModule& SceneOutlinerModule = FModuleManager::Get().LoadModuleChecked<FSceneOutlinerModule>(TEXT("SceneOutliner"));
 
-		FSceneOutlinerInitializationOptions InitOptions;
+		SceneOutliner::FInitializationOptions InitOptions;
 		InitOptions.Mode = ESceneOutlinerMode::ActorPicker;
-		InitOptions.Filters = ActorFilters;
+		InitOptions.Filters->AddFilterPredicate(ActorFilter);
 		InitOptions.bFocusSearchBoxWhenOpened = true;
 		
 		MenuContent =
@@ -135,7 +135,7 @@ void SPropertyMenuActorPicker::OnPaste()
 	else
 	{
 		AActor* Actor = LoadObject<AActor>(NULL, *DestPath);
-		if(Actor && (!ActorFilters.IsValid() || ActorFilters->PassesAllFilters(Actor)))
+		if(Actor && (ActorFilter.IsBound() || ActorFilter.Execute(Actor)))
 		{
 			SetValue(Actor);
 		}
diff --git a/Engine/Source/Editor/PropertyEditor/Private/UserInterface/PropertyEditor/SPropertyMenuActorPicker.h b/Engine/Source/Editor/PropertyEditor/Private/UserInterface/PropertyEditor/SPropertyMenuActorPicker.h
--- a/Engine/Source/Editor/PropertyEditor/Private/UserInterface/PropertyEditor/SPropertyMenuActorPicker.h
+++ b/Engine/Source/Editor/PropertyEditor/Private/UserInterface/PropertyEditor/SPropertyMenuActorPicker.h
@@ -8,11 +8,11 @@ public:
 	SLATE_BEGIN_ARGS( SPropertyMenuActorPicker )
 		: _InitialActor(NULL)
 		, _AllowClear(true)
-		, _ActorFilters(NULL)
+		, _ActorFilter()
 	{}
 		SLATE_ARGUMENT( AActor*, InitialActor )
 		SLATE_ARGUMENT( bool, AllowClear )
-		SLATE_ARGUMENT( TSharedPtr< SceneOutliner::FOutlinerFilters >, ActorFilters )
+		SLATE_ARGUMENT( FOnShouldFilterActor, ActorFilter )
 		SLATE_EVENT( FOnActorSelected, OnSet )
 		SLATE_EVENT( FSimpleDelegate, OnClose )
 		SLATE_EVENT( FSimpleDelegate, OnUseSelected )
@@ -75,7 +75,8 @@ private:
 	/** Whether the asset can be 'None' in this case */
 	bool bAllowClear;
 
-	TSharedPtr< SceneOutliner::FOutlinerFilters > ActorFilters;
+	/** Delegate used to test whether an actor should be displayed or not */
+	FOnShouldFilterActor ActorFilter;
 
 	/** Delegate to call when our object value should be set */
 	FOnActorSelected OnSet;
diff --git a/Engine/Source/Editor/PropertyEditor/Public/PropertyCustomizationHelpers.h b/Engine/Source/Editor/PropertyEditor/Public/PropertyCustomizationHelpers.h
--- a/Engine/Source/Editor/PropertyEditor/Public/PropertyCustomizationHelpers.h
+++ b/Engine/Source/Editor/PropertyEditor/Public/PropertyCustomizationHelpers.h
@@ -32,7 +32,7 @@ namespace PropertyCustomizationHelpers
 	PROPERTYEDITOR_API TSharedRef<SWidget> MakeAssetPickerAnchorButton( FOnGetAllowedClasses OnGetAllowedClasses, FOnAssetSelected OnAssetSelectedFromPicker );
 	PROPERTYEDITOR_API TSharedRef<SWidget> MakeAssetPickerWithMenu( const FAssetData& InitialObject, const bool AllowClear, const TArray<const UClass*>& AllowedClasses, const TArray<UFactory*>& NewAssetFactories, FOnShouldFilterAsset OnShouldFilterAsset, FOnAssetSelected OnSet, FSimpleDelegate OnClose );
 	PROPERTYEDITOR_API TSharedRef<SWidget> MakeActorPickerAnchorButton( FOnGetActorFilters OnGetActorFilters, FOnActorSelected OnActorSelectedFromPicker );
-	PROPERTYEDITOR_API TSharedRef<SWidget> MakeActorPickerWithMenu( AActor* const InitialActor, const bool AllowClear, const TSharedPtr< SceneOutliner::FOutlinerFilters >& ActorFilters, FOnActorSelected OnSet, FSimpleDelegate OnClose, FSimpleDelegate OnUseSelected );
+	PROPERTYEDITOR_API TSharedRef<SWidget> MakeActorPickerWithMenu( AActor* const InitialActor, const bool AllowClear, FOnShouldFilterActor ActorFilter, FOnActorSelected OnSet, FSimpleDelegate OnClose, FSimpleDelegate OnUseSelected );
 	PROPERTYEDITOR_API TSharedRef<SWidget> MakeInteractiveActorPicker( FOnGetAllowedClasses OnGetAllowedClasses, FOnShouldFilterActor OnShouldFilterActor, FOnActorSelected OnActorSelectedFromPicker );
 
 	/** @return the UBoolProperty edit condition property if one exists. */
diff --git a/Engine/Source/Editor/SceneOutliner/Private/ActorTreeItem.cpp b/Engine/Source/Editor/SceneOutliner/Private/ActorTreeItem.cpp
new file mode 100644
--- /dev/null
+++ b/Engine/Source/Editor/SceneOutliner/Private/ActorTreeItem.cpp
@@ -0,0 +1,324 @@
+// Copyright 1998-2014 Epic Games, Inc. All Rights Reserved.
+
+#include "SceneOutlinerPrivatePCH.h"
+#include "SSceneOutliner.h"
+
+#include "SceneOutlinerFilters.h"
+
+#include "ActorTreeItem.h"
+#include "FolderTreeItem.h"
+#include "WorldTreeItem.h"
+
+#include "ActorEditorUtils.h"
+#include "LevelUtils.h"
+#include "MessageLog.h"
+#include "SSocketChooser.h"
+#include "ScopedTransaction.h"
+
+#define LOCTEXT_NAMESPACE "SceneOutliner_ActorTreeItem"
+
+namespace SceneOutliner
+{
+
+FDragValidationInfo FActorDropTarget::ValidateDrop(FDragDropPayload& DraggedObjects, UWorld& World) const
+{
+	if (DraggedObjects.Folders)
+	{
+		return FDragValidationInfo(FActorDragDropGraphEdOp::ToolTip_IncompatibleGeneric, LOCTEXT("FoldersOnActorError", "Cannot attach folders to actors"));
+	}
+
+	const AActor* DropTarget = Actor.Get();
+
+	if (!DropTarget || !DraggedObjects.Actors)
+	{
+		return FDragValidationInfo(FActorDragDropGraphEdOp::ToolTip_IncompatibleGeneric, FText());
+	}	
+
+	FText AttachErrorMsg;
+	bool bCanAttach = true;
+	bool bDraggedOntoAttachmentParent = true;
+
+	const auto& DragActors = DraggedObjects.Actors.GetValue();
+	for (const auto& Actor : DragActors)
+	{
+		AActor* DragActor = Actor.Get();
+		if (DragActor)
+		{
+			if (!GEditor->CanParentActors(DropTarget, DragActor, &AttachErrorMsg))
+			{
+				bCanAttach = false;
+			}
+			if (DragActor->GetAttachParentActor() != DropTarget)
+			{
+				bDraggedOntoAttachmentParent = false;
+			}
+		}
+	}
+
+	const FText ActorLabel = FText::FromString(DropTarget->GetActorLabel());
+	if (bDraggedOntoAttachmentParent)
+	{
+		if (DragActors.Num() == 1)
+		{
+			return FDragValidationInfo(FActorDragDropGraphEdOp::ToolTip_CompatibleDetach, ActorLabel);
+		}
+		else
+		{
+			return FDragValidationInfo(FActorDragDropGraphEdOp::ToolTip_CompatibleMultipleDetach, ActorLabel);
+		}
+	}
+	else if (bCanAttach)
+	{
+		if (DragActors.Num() == 1)
+		{
+			return FDragValidationInfo(FActorDragDropGraphEdOp::ToolTip_CompatibleAttach, ActorLabel);
+		}
+		else
+		{
+			return FDragValidationInfo(FActorDragDropGraphEdOp::ToolTip_CompatibleMultipleAttach, ActorLabel);
+		}
+	}
+	else
+	{
+		if (DragActors.Num() == 1)
+		{
+			return FDragValidationInfo(FActorDragDropGraphEdOp::ToolTip_IncompatibleGeneric, AttachErrorMsg);
+		}
+		else
+		{
+			const FText ReasonText = FText::Format(LOCTEXT("DropOntoText", "{DropTarget}. {AttachMessage}"), ActorLabel, AttachErrorMsg);
+			return FDragValidationInfo(FActorDragDropGraphEdOp::ToolTip_IncompatibleMultipleAttach, ReasonText);
+		}
+	}
+}
+
+void FActorDropTarget::OnDrop(FDragDropPayload& DraggedObjects, UWorld& World, const FDragValidationInfo& ValidationInfo, TSharedRef<SWidget> DroppedOnWidget)
+{
+	AActor* DropActor = Actor.Get();
+	if (!DropActor)
+	{
+		return;
+	}
+
+	FActorArray DraggedActors = DraggedObjects.Actors.GetValue();
+
+	FMessageLog EditorErrors("EditorErrors");
+	EditorErrors.NewPage(LOCTEXT("ActorAttachmentsPageLabel", "Actor attachment"));
+
+	if (ValidationInfo.TooltipType == FActorDragDropGraphEdOp::ToolTip_CompatibleMultipleDetach || ValidationInfo.TooltipType == FActorDragDropGraphEdOp::ToolTip_CompatibleDetach)
+	{
+		const FScopedTransaction Transaction( LOCTEXT("UndoAction_DetachActors", "Detach actors") );
+
+		for (const auto& WeakActor : DraggedActors)
+		{
+			if (auto* DragActor = WeakActor.Get())
+			{
+				DetachActorFromParent(DragActor);
+			}
+		}
+	}
+	else if (ValidationInfo.TooltipType == FActorDragDropGraphEdOp::ToolTip_CompatibleMultipleAttach || ValidationInfo.TooltipType == FActorDragDropGraphEdOp::ToolTip_CompatibleAttach)
+	{
+		// Show socket chooser if we have sockets to select
+
+		//@TODO: Should create a menu for each component that contains sockets, or have some form of disambiguation within the menu (like a fully qualified path)
+		// Instead, we currently only display the sockets on the root component
+		USceneComponent* Component = DropActor->GetRootComponent();
+		if ((Component != NULL) && (Component->HasAnySockets()))
+		{
+			// Create the popup
+			FSlateApplication::Get().PushMenu(
+				DroppedOnWidget,
+				SNew(SSocketChooserPopup)
+				.SceneComponent( Component )
+				.OnSocketChosen_Static(&FActorDropTarget::PerformAttachment, Actor, MoveTemp(DraggedActors) ),
+				FSlateApplication::Get().GetCursorPos(),
+				FPopupTransitionEffect( FPopupTransitionEffect::TypeInPopup )
+				);
+		}
+		else
+		{
+			PerformAttachment(NAME_None, Actor, MoveTemp(DraggedActors));
+		}
+	}
+
+	// Report errors
+	EditorErrors.Notify(NSLOCTEXT("ActorAttachmentError", "AttachmentsFailed", "Attachments Failed!"));
+}
+
+void FActorDropTarget::PerformAttachment(FName SocketName, TWeakObjectPtr<AActor> Parent, const FActorArray NewAttachments)
+{
+	AActor* ParentActor = Parent.Get();
+	if (ParentActor)
+	{
+		// modify parent and child
+		const FScopedTransaction Transaction( LOCTEXT("UndoAction_PerformAttachment", "Attach actors") );
+
+		// Attach each child
+		bool bAttached = false;
+		for (auto& Child : NewAttachments)
+		{
+			AActor* ChildActor = Child.Get();
+			if (GEditor->CanParentActors(ParentActor, ChildActor))
+			{
+				GEditor->ParentActors(ParentActor, ChildActor, SocketName);
+			}
+		}
+	}
+}
+
+void FActorDropTarget::DetachActorFromParent(AActor* ChildActor)
+{
+	USceneComponent* RootComp = ChildActor->GetRootComponent();
+	if (RootComp && RootComp->AttachParent)
+	{
+		AActor* OldParent = RootComp->AttachParent->GetOwner();
+		OldParent->Modify();
+		RootComp->DetachFromParent(true);
+		ChildActor->SetFolderPath(OldParent->GetFolderPath());
+	}
+}
+
+FActorTreeItem::FActorTreeItem(AActor* InActor)
+	: Actor(InActor)
+{
+	bExistsInCurrentWorldAndPIE = GEditor->ObjectsThatExistInEditorWorld.Get(InActor);
+}
+
+FTreeItemPtr FActorTreeItem::FindParent(const FTreeItemMap& ExistingItems) const
+{
+	AActor* ActorPtr = Actor.Get();
+	if (!ActorPtr)
+	{
+		return nullptr;
+	}
+
+	// Parents should have already been added to the tree
+	TWeakObjectPtr<AActor> ParentActor = ActorPtr->GetAttachParentActor();
+	if (ParentActor.IsValid())
+	{
+		return ExistingItems.FindRef(ParentActor);
+	}
+	else
+	{
+		const bool bShouldShowFolders = SharedData->Mode == ESceneOutlinerMode::ActorBrowsing || SharedData->bOnlyShowFolders;
+
+		const FName ActorFolder = ActorPtr->GetFolderPath();
+		if (bShouldShowFolders && !ActorFolder.IsNone())
+		{
+			return ExistingItems.FindRef(ActorFolder);
+		}
+	}
+
+	if (UWorld* World = ActorPtr->GetWorld())
+	{
+		return ExistingItems.FindRef(World);
+	}
+
+	return nullptr;
+}
+
+FTreeItemPtr FActorTreeItem::CreateParent() const
+{
+	AActor* ActorPtr = Actor.Get();
+	if (!ActorPtr)
+	{
+		return nullptr;
+	}
+
+	AActor* ParentActor = ActorPtr->GetAttachParentActor();
+	if (ParentActor)
+	{
+		return MakeShareable(new FActorTreeItem(ParentActor));
+	}
+	else
+	{
+		const bool bShouldShowFolders = SharedData->Mode == ESceneOutlinerMode::ActorBrowsing || SharedData->bOnlyShowFolders;
+
+		const FName ActorFolder = ActorPtr->GetFolderPath();
+		if (bShouldShowFolders && !ActorFolder.IsNone())
+		{
+			return MakeShareable(new FFolderTreeItem(ActorFolder));
+		}
+	}
+
+	if (UWorld* World = ActorPtr->GetWorld())
+	{
+		return MakeShareable(new FWorldTreeItem(World));
+	}
+
+	return nullptr;
+}
+
+void FActorTreeItem::Visit(const ITreeItemVisitor& Visitor) const
+{
+	Visitor.Visit(*this);
+}
+
+void FActorTreeItem::Visit(const IMutableTreeItemVisitor& Visitor)
+{
+	Visitor.Visit(*this);
+}
+
+FTreeItemID FActorTreeItem::GetID() const
+{
+	return FTreeItemID(Actor);
+}
+
+FString FActorTreeItem::GetDisplayString() const
+{
+	const AActor* ActorPtr = Actor.Get();
+	return ActorPtr ? ActorPtr->GetActorLabel() : LOCTEXT("ActorLabelForMissingActor", "(Deleted Actor)").ToString();
+}
+
+int32 FActorTreeItem::GetTypeSortPriority() const
+{
+	return ETreeItemSortOrder::Actor;
+}
+
+bool FActorTreeItem::CanInteract() const
+{
+	AActor* ActorPtr = Actor.Get();
+	if (!ActorPtr || Flags.bIsFilteredOut)
+	{
+		return false;
+	}
+
+	const bool bInSelected = true;
+	const bool bSelectEvenIfHidden = true;		// @todo outliner: Is this actually OK?
+	if (!GEditor->CanSelectActor(ActorPtr, bInSelected, bSelectEvenIfHidden))
+	{
+		return false;
+	}
+
+	return true;
+}
+
+void FActorTreeItem::PopulateDragDropPayload(FDragDropPayload& Payload) const
+{
+	AActor* ActorPtr = Actor.Get();
+	if (ActorPtr)
+	{
+		if (!Payload.Actors)
+		{
+			Payload.Actors = FActorArray();
+		}
+		Payload.Actors->Add(Actor);
+	}
+}
+
+FDragValidationInfo FActorTreeItem::ValidateDrop(FDragDropPayload& DraggedObjects, UWorld& World) const
+{
+	FActorDropTarget Target(Actor.Get());
+	return Target.ValidateDrop(DraggedObjects, World);
+}
+
+void FActorTreeItem::OnDrop(FDragDropPayload& DraggedObjects, UWorld& World, const FDragValidationInfo& ValidationInfo, TSharedRef<SWidget> DroppedOnWidget)
+{
+	FActorDropTarget Target(Actor.Get());
+	return Target.OnDrop(DraggedObjects, World, ValidationInfo, DroppedOnWidget);
+}
+
+}	// namespace SceneOutliner
+
+#undef LOCTEXT_NAMESPACE
\ No newline at end of file
diff --git a/Engine/Source/Editor/SceneOutliner/Private/FolderTreeItem.cpp b/Engine/Source/Editor/SceneOutliner/Private/FolderTreeItem.cpp
new file mode 100644
--- /dev/null
+++ b/Engine/Source/Editor/SceneOutliner/Private/FolderTreeItem.cpp
@@ -0,0 +1,305 @@
+// Copyright 1998-2014 Epic Games, Inc. All Rights Reserved.
+
+#include "SceneOutlinerPrivatePCH.h"
+#include "SSceneOutliner.h"
+
+#include "SceneOutlinerFilters.h"
+
+#include "EditorActorFolders.h"
+#include "FolderTreeItem.h"
+#include "ScopedTransaction.h"
+
+#define LOCTEXT_NAMESPACE "SceneOutliner_FolderTreeItem"
+
+namespace SceneOutliner
+{
+
+FDragValidationInfo FFolderDropTarget::ValidateDrop(FDragDropPayload& DraggedObjects, UWorld& World) const
+{
+	if (DraggedObjects.Folders)
+	{
+		// Iterate over all the folders that have been dragged
+		for (FName DraggedFolder : DraggedObjects.Folders.GetValue())
+		{
+			FName Leaf = GetFolderLeafName(DraggedFolder);
+			FName Parent = GetParentPath(DraggedFolder);
+
+			if (Parent == DestinationPath)
+			{
+				FFormatNamedArguments Args;
+				Args.Add(TEXT("SourceName"), FText::FromName(Leaf));
+
+				FText Text;
+				if (DestinationPath.IsNone())
+				{
+					Text = FText::Format(LOCTEXT("FolderAlreadyAssignedRoot", "{SourceName} is already assigned to root"), Args);
+				}
+				else
+				{
+					Args.Add(TEXT("DestPath"), FText::FromName(DestinationPath));
+					Text = FText::Format(LOCTEXT("FolderAlreadyAssigned", "{SourceName} is already assigned to {DestPath}"), Args);
+				}
+				
+				return FDragValidationInfo(FActorDragDropGraphEdOp::ToolTip_IncompatibleGeneric, Text);
+			}
+
+			const FString DragFolderPath = DraggedFolder.ToString();
+			const FString LeafName = Leaf.ToString();
+			const FString DstFolderPath = DestinationPath.IsNone() ? FString() : DestinationPath.ToString();
+			const FString NewPath = DstFolderPath / LeafName;
+
+			if (FActorFolders::Get().GetFolderProperties(World, FName(*NewPath)))
+			{
+				// The folder already exists
+				FFormatNamedArguments Args;
+				Args.Add(TEXT("DragName"), FText::FromString(LeafName));
+				return FDragValidationInfo(FActorDragDropGraphEdOp::ToolTip_IncompatibleGeneric,
+					FText::Format(LOCTEXT("FolderAlreadyExistsRoot", "A folder called \"{DragName}\" already exists at this level"), Args));
+			}
+			else if (DragFolderPath == DstFolderPath || DstFolderPath.StartsWith(DragFolderPath + "/"))
+			{
+				// Cannot drag as a child of itself
+				FFormatNamedArguments Args;
+				Args.Add(TEXT("FolderPath"), FText::FromName(DraggedFolder));
+				return FDragValidationInfo(
+					FActorDragDropGraphEdOp::ToolTip_IncompatibleGeneric,
+					FText::Format(LOCTEXT("ChildOfItself", "Cannot move \"{FolderPath}\" to be a child of itself"), Args));
+			}
+		}
+	}
+
+	if (DraggedObjects.Actors)
+	{
+		// Iterate over all the folders that have been dragged
+		for (auto WeakActor : DraggedObjects.Actors.GetValue())
+		{
+			AActor* Actor = WeakActor.Get();
+
+			if (Actor->GetFolderPath() == DestinationPath)
+			{
+				FFormatNamedArguments Args;
+				Args.Add(TEXT("SourceName"), FText::FromString(Actor->GetActorLabel()));
+
+				FText Text;
+				if (DestinationPath.IsNone())
+				{
+					Text = FText::Format(LOCTEXT("FolderAlreadyAssignedRoot", "{SourceName} is already assigned to root"), Args);
+				}
+				else
+				{
+					Args.Add(TEXT("DestPath"), FText::FromName(DestinationPath));
+					Text = FText::Format(LOCTEXT("FolderAlreadyAssigned", "{SourceName} is already assigned to {DestPath}"), Args);
+				}
+				
+				return FDragValidationInfo(FActorDragDropGraphEdOp::ToolTip_IncompatibleGeneric, Text);
+			}
+		}
+	}
+
+	// Everything else is a valid operation
+	if (DestinationPath.IsNone())
+	{
+		return FDragValidationInfo(FActorDragDropGraphEdOp::ToolTip_CompatibleGeneric, LOCTEXT("MoveToRoot", "Move to root"));
+	}
+	else
+	{
+		return FDragValidationInfo(FActorDragDropGraphEdOp::ToolTip_CompatibleGeneric, FText::Format(LOCTEXT("MoveInto", "Move into \"{DropPath}\""), FText::FromName(DestinationPath)));
+	}
+}
+
+void FFolderDropTarget::OnDrop(FDragDropPayload& DraggedObjects, UWorld& World, const FDragValidationInfo& ValidationInfo, TSharedRef<SWidget> DroppedOnWidget)
+{
+	const FScopedTransaction Transaction( LOCTEXT("MoveOutlinerItems", "Move Scene Outliner Items") );
+
+	if (DraggedObjects.Folders)
+	{
+		for (FName Folder : DraggedObjects.Folders.GetValue())
+		{
+			MoveFolderTo(Folder, DestinationPath, World);
+		}
+	}
+
+	if (DraggedObjects.Actors)
+	{
+		for (auto WeakActor : DraggedObjects.Actors.GetValue())
+		{
+			AActor* Actor = WeakActor.Get();
+			if (Actor)
+			{
+				Actor->SetFolderPath(DestinationPath);
+			}
+		}
+	}
+}
+
+FFolderTreeItem::FFolderTreeItem(FName InPath)
+	: Path(InPath)
+	, LeafName(GetFolderLeafName(InPath))
+{
+}
+
+FTreeItemPtr FFolderTreeItem::FindParent(const FTreeItemMap& ExistingItems) const
+{
+	const FName ParentPath = GetParentPath(Path);
+	if (!ParentPath.IsNone())
+	{
+		return ExistingItems.FindRef(ParentPath);
+	}
+
+	if (SharedData->RepresentingWorld)
+	{
+		return ExistingItems.FindRef(SharedData->RepresentingWorld);
+	}
+
+	return nullptr;
+}
+
+FTreeItemPtr FFolderTreeItem::CreateParent() const
+{
+	const FName ParentPath = GetParentPath(Path);
+	if (!ParentPath.IsNone())
+	{
+		return MakeShareable(new FFolderTreeItem(ParentPath));
+	}
+
+	if (SharedData->RepresentingWorld)
+	{
+		return MakeShareable(new FWorldTreeItem(SharedData->RepresentingWorld));
+	}
+
+	return nullptr;
+}
+
+void FFolderTreeItem::Visit(const ITreeItemVisitor& Visitor) const
+{
+	Visitor.Visit(*this);
+}
+
+void FFolderTreeItem::Visit(const IMutableTreeItemVisitor& Visitor)
+{
+	Visitor.Visit(*this);
+}
+
+FTreeItemID FFolderTreeItem::GetID() const
+{
+	return FTreeItemID(Path);
+}
+
+FString FFolderTreeItem::GetDisplayString() const
+{
+	return LeafName.ToString();
+}
+
+int32 FFolderTreeItem::GetTypeSortPriority() const
+{
+	return ETreeItemSortOrder::Folder;
+}
+
+
+bool FFolderTreeItem::CanInteract() const
+{
+	return !Flags.bIsFilteredOut;
+}
+
+void FFolderTreeItem::Delete()
+{
+	if (!SharedData->RepresentingWorld)
+	{
+		return;
+	}
+
+	struct FResetActorFolders : IMutableTreeItemVisitor
+	{
+		virtual void Visit(FActorTreeItem& ActorItem) const override
+		{
+			if (AActor* Actor = ActorItem.Actor.Get())
+			{
+				Actor->SetFolderPath(FName());
+			}	
+		}
+		virtual void Visit(FFolderTreeItem& FolderItem) const override
+		{
+			FResetActorFolders ResetFolders;
+			for (auto& Child : FolderItem.GetChildren())
+			{
+				Child.Pin()->Visit(ResetFolders);
+			}
+		}
+	};
+
+	const FScopedTransaction Transaction( LOCTEXT("DeleteFolder", "Delete Folder") );
+
+	FResetActorFolders ResetFolders;
+	for (auto& Child : GetChildren())
+	{
+		Child.Pin()->Visit(ResetFolders);
+	}
+
+	FActorFolders::Get().DeleteFolder(*SharedData->RepresentingWorld, Path);
+}
+
+void FFolderTreeItem::CreateSubFolder(TWeakPtr<SSceneOutliner> WeakOutliner)
+{
+	auto Outliner = WeakOutliner.Pin();
+
+	if (Outliner.IsValid() && SharedData->RepresentingWorld)
+	{
+		const FScopedTransaction Transaction(LOCTEXT("UndoAction_CreateFolder", "Create Folder"));
+
+		const FName NewFolderName = FActorFolders::Get().GetDefaultFolderName(*SharedData->RepresentingWorld, Path);
+		FActorFolders::Get().CreateFolder(*SharedData->RepresentingWorld, NewFolderName);
+
+		// At this point the new folder will be in our newly added list, so select it and open a rename when it gets refreshed
+		Outliner->OnItemAdded(NewFolderName, ENewItemAction::Select | ENewItemAction::Rename);
+	}
+}
+
+void FFolderTreeItem::OnExpansionChanged()
+{
+	if (!SharedData->RepresentingWorld)
+	{
+		return;
+	}
+
+	// Update the central store of folder properties with this folder's new expansion state
+	if (FActorFolderProps* Props = FActorFolders::Get().GetFolderProperties(*SharedData->RepresentingWorld, Path))
+	{
+		Props->bIsExpanded = Flags.bIsExpanded;
+	}
+}
+
+void FFolderTreeItem::GenerateContextMenu(FMenuBuilder& MenuBuilder, SSceneOutliner& Outliner)
+{
+	auto SharedOutliner = StaticCastSharedRef<SSceneOutliner>(Outliner.AsShared());
+
+	const FSlateIcon NewFolderIcon(FEditorStyle::GetStyleSetName(), "SceneOutliner.NewFolderIcon");
+	
+	MenuBuilder.AddMenuEntry(LOCTEXT("CreateSubFolder", "Create Sub Folder"), FText(), NewFolderIcon, FUIAction(FExecuteAction::CreateSP(this, &FFolderTreeItem::CreateSubFolder, TWeakPtr<SSceneOutliner>(SharedOutliner))));
+	MenuBuilder.AddMenuEntry(LOCTEXT("RenameFolder", "Rename"), FText(), FSlateIcon(), FUIAction(FExecuteAction::CreateSP(&Outliner, &SSceneOutliner::InitiateRename, AsShared())));
+	MenuBuilder.AddMenuEntry(LOCTEXT("DeleteFolder", "Delete"), FText(), FSlateIcon(), FUIAction(FExecuteAction::CreateSP(this, &FFolderTreeItem::Delete)));
+}
+
+void FFolderTreeItem::PopulateDragDropPayload(FDragDropPayload& Payload) const
+{
+	if (!Payload.Folders)
+	{
+		Payload.Folders = FFolderPaths();
+	}
+	Payload.Folders->Add(Path);
+}
+
+FDragValidationInfo FFolderTreeItem::ValidateDrop(FDragDropPayload& DraggedObjects, UWorld& World) const
+{
+	FFolderDropTarget Target(Path);
+	return Target.ValidateDrop(DraggedObjects, World);
+}
+
+void FFolderTreeItem::OnDrop(FDragDropPayload& DraggedObjects, UWorld& World, const FDragValidationInfo& ValidationInfo, TSharedRef<SWidget> DroppedOnWidget)
+{
+	FFolderDropTarget Target(Path);
+	return Target.OnDrop(DraggedObjects, World, ValidationInfo, DroppedOnWidget);
+}
+
+}	// namespace SceneOutliner
+
+#undef LOCTEXT_NAMESPACE
\ No newline at end of file
diff --git a/Engine/Source/Editor/SceneOutliner/Private/LevelBlueprintTreeItem.cpp b/Engine/Source/Editor/SceneOutliner/Private/LevelBlueprintTreeItem.cpp
new file mode 100644
--- /dev/null
+++ b/Engine/Source/Editor/SceneOutliner/Private/LevelBlueprintTreeItem.cpp
@@ -0,0 +1,121 @@
+// Copyright 1998-2014 Epic Games, Inc. All Rights Reserved.
+
+#include "SceneOutlinerPrivatePCH.h"
+
+#include "LevelBlueprintTreeItem.h"
+#include "LevelEditor.h"
+#include "Runtime/Engine/Classes/Engine/LevelScriptBlueprint.h"
+
+#include "SNotificationList.h"
+#include "NotificationManager.h"
+
+#include "LevelUtils.h"
+
+#define LOCTEXT_NAMESPACE "SceneOutliner_LevelBlueprintTreeItem"
+
+namespace SceneOutliner
+{
+
+FLevelBlueprintTreeItem::FLevelBlueprintTreeItem(const FLevelBlueprintHandle& InHandle)
+	: Handle(InHandle)
+{
+
+}
+
+FTreeItemPtr FLevelBlueprintTreeItem::FindParent(const FTreeItemMap& ExistingItems) const
+{
+	if (SharedData->RepresentingWorld)
+	{
+		return ExistingItems.FindRef(SharedData->RepresentingWorld);
+	}
+
+	return nullptr;
+}
+
+FTreeItemPtr FLevelBlueprintTreeItem::CreateParent() const
+{
+
+	if (SharedData->RepresentingWorld)
+	{
+		return MakeShareable(new FWorldTreeItem(SharedData->RepresentingWorld));
+	}
+
+	return nullptr;
+}
+
+void FLevelBlueprintTreeItem::Visit(const ITreeItemVisitor& Visitor) const
+{
+	Visitor.Visit(*this);
+}
+
+void FLevelBlueprintTreeItem::Visit(const IMutableTreeItemVisitor& Visitor)
+{
+	Visitor.Visit(*this);
+}
+
+FTreeItemID FLevelBlueprintTreeItem::GetID() const
+{
+	return FTreeItemID(Handle);
+}
+
+FString FLevelBlueprintTreeItem::GetDisplayString() const
+{
+	FFormatOrderedArguments Args;
+	if (ULevel* LevelPtr = Handle.ParentLevel.Get())
+	{
+		return FPackageName::GetShortName(LevelPtr->GetOutermost());
+	}
+	else
+	{
+		return LOCTEXT("UnknownLevel", "Unknown Level").ToString();
+	}
+}
+
+int32 FLevelBlueprintTreeItem::GetTypeSortPriority() const
+{
+	return ETreeItemSortOrder::Actor;
+}
+
+bool FLevelBlueprintTreeItem::CanInteract() const
+{
+	if (ULevel* Level = Handle.ParentLevel.Get())
+	{
+		const EObjectFlags InvalidSelectableFlags = RF_PendingKill|RF_BeginDestroyed|RF_Unreachable;
+		if (!FLevelUtils::IsLevelVisible(Level) || FLevelUtils::IsLevelLocked(Level) || Level->HasAnyFlags(InvalidSelectableFlags))
+		{
+			return false;
+		}
+	}
+	return !Flags.bIsFilteredOut;
+}
+
+void FLevelBlueprintTreeItem::GenerateContextMenu(FMenuBuilder& MenuBuilder, SSceneOutliner& Outliner)
+{
+	const FSlateIcon LevelBlueprintIcon(FEditorStyle::GetStyleSetName(), "LevelEditor.OpenLevelBlueprint");
+
+	MenuBuilder.AddMenuEntry(LOCTEXT("OpenLevelBlueprint", "Edit Level Blueprint"), FText(), LevelBlueprintIcon, FExecuteAction::CreateSP(this, &FLevelBlueprintTreeItem::Open));
+}
+
+FDragValidationInfo FLevelBlueprintTreeItem::ValidateDrop(FDragDropPayload& DraggedObjects, UWorld& World) const
+{
+	return FDragValidationInfo::Invalid();
+}
+
+void FLevelBlueprintTreeItem::Open() const
+{
+	ULevelScriptBlueprint* LevelBlueprint = Handle.GetOrCreateLevelBlueprint();
+	if (!LevelBlueprint || !FAssetEditorManager::Get().OpenEditorForAsset(LevelBlueprint))
+	{
+		FNotificationInfo Info(LOCTEXT("CannotOpenLevelBlueprint", "Unable to open the specified level blueprint."));
+		Info.ExpireDuration = 3.0f;
+		Info.bUseLargeFont = false;
+		Info.bFireAndForget = true;
+		Info.bUseSuccessFailIcons = true;
+		FSlateNotificationManager::Get().AddNotification(Info)->SetCompletionState(SNotificationItem::CS_Fail);
+	}
+}
+
+
+}		// namespace SceneOutliner
+
+#undef LOCTEXT_NAMESPACE
\ No newline at end of file
diff --git a/Engine/Source/Editor/SceneOutliner/Private/SOutlinerTreeView.cpp b/Engine/Source/Editor/SceneOutliner/Private/SOutlinerTreeView.cpp
--- a/Engine/Source/Editor/SceneOutliner/Private/SOutlinerTreeView.cpp
+++ b/Engine/Source/Editor/SceneOutliner/Private/SOutlinerTreeView.cpp
@@ -2,365 +2,146 @@
 
 #include "SceneOutlinerPrivatePCH.h"
 
-#include "MessageLog.h"
-#include "SSocketChooser.h"
-#include "ScopedTransaction.h"
-#include "SceneOutlinerTreeItems.h"
-
 #define LOCTEXT_NAMESPACE "SSceneOutliner"
 
 namespace SceneOutliner
 {
-	typedef TArray<TWeakPtr<TOutlinerFolderTreeItem>> FolderArray;
-	typedef TArray<TWeakObjectPtr<AActor>> ActorArray;
-
-	/* Struct used for validation of a drag/drop operation in the scene outliner */
-	struct FDragValidationInfo
+	static void UpdateOperationDecorator(const FDragDropEvent& Event, const FDragValidationInfo& ValidationInfo)
 	{
-		/* The tooltip type to display on the operation */
-		FActorDragDropGraphEdOp::ToolTipTextType TooltipType;
-
-		/* The tooltip text to display on the operation */
-		FText ValidationText;
+		const FSlateBrush* Icon = ValidationInfo.IsValid() ? FEditorStyle::GetBrush(TEXT("Graph.ConnectorFeedback.OK")) : FEditorStyle::GetBrush(TEXT("Graph.ConnectorFeedback.Error"));
 
-		/* Construct this validation information out of a tootip type and some text */
-		FDragValidationInfo(const FActorDragDropGraphEdOp::ToolTipTextType InTooltipType, const FText InValidationText)
-			: TooltipType(InTooltipType)
-			, ValidationText(InValidationText)
-		{}
-
-		/* Return a generic invalid result */
-		static FDragValidationInfo Invalid()
+		FDragDropOperation* Operation = Event.GetOperation().Get();
+		if (Operation)
 		{
-			return FDragValidationInfo(FActorDragDropGraphEdOp::ToolTip_IncompatibleGeneric, FText());
-		}
-		
-		/* @return true if this operation is valid, false otheriwse */ 
-		bool IsValid() const
-		{
-			switch(TooltipType)
+			if (Operation->IsOfType<FSceneOutlinerDragDropOp>())
 			{
-			case FActorDragDropGraphEdOp::ToolTip_Compatible:
-			case FActorDragDropGraphEdOp::ToolTip_CompatibleAttach:
-			case FActorDragDropGraphEdOp::ToolTip_CompatibleGeneric:
-			case FActorDragDropGraphEdOp::ToolTip_CompatibleMultipleAttach:
-			case FActorDragDropGraphEdOp::ToolTip_CompatibleDetach:
-			case FActorDragDropGraphEdOp::ToolTip_CompatibleMultipleDetach:
-				return true;
-			default:
-				return false;
+				auto* OutlinerOp = static_cast<FSceneOutlinerDragDropOp*>(Operation);
+				OutlinerOp->SetTooltip(ValidationInfo.ValidationText, Icon);
+			}
+			else if (Operation->IsOfType<FActorDragDropGraphEdOp>())
+			{
+				static_cast<FActorDragDropGraphEdOp*>(Operation)->SetToolTip(ValidationInfo.TooltipType, ValidationInfo.ValidationText);
+			}
+			else if (Operation->IsOfType<FDecoratedDragDropOp>())
+			{
+				auto* DecoratedOp = static_cast<FDecoratedDragDropOp*>(Operation);
+				DecoratedOp->SetToolTip(ValidationInfo.ValidationText, Icon);
 			}
 		}
-	};
+	}
 
-	/** Extract lists of dragged folders and actors from the specified drag drop operation */
-	void ExtractDraggedActorsAndFolders(TSharedPtr<FDragDropOperation> Operation, const ActorArray*& DraggedActors, const FolderArray*& DraggedFolders)
+	static void ResetOperationDecorator(const FDragDropEvent& Event)
 	{
-		if (!Operation.IsValid())
-		{
-			return;
-		}
-
-		if (Operation->IsOfType<FFolderActorDragDropOp>())
-		{
-			DraggedFolders = &StaticCastSharedPtr<FFolderActorDragDropOp>(Operation)->Folders;
-		}
-		else if (Operation->IsOfType<FFolderDragDropOp>())
+		FDragDropOperation* Operation = Event.GetOperation().Get();
+		if (Operation)
 		{
-			DraggedFolders = &StaticCastSharedPtr<FFolderDragDropOp>(Operation)->Folders;
-		}
-
-		if (Operation->IsOfType<FActorDragDropGraphEdOp>())
-		{
-			DraggedActors = &StaticCastSharedPtr<FActorDragDropGraphEdOp>(Operation)->Actors;
+			if (Operation->IsOfType<FSceneOutlinerDragDropOp>())
+			{
+				static_cast<FSceneOutlinerDragDropOp*>(Operation)->ResetTooltip();
+			}
+			else if (Operation->IsOfType<FDecoratedDragDropOp>())
+			{
+				static_cast<FDecoratedDragDropOp*>(Operation)->ResetToDefaultToolTip();
+			}
 		}
 	}
 
-	/** Check to see if the specified drag drop operation is relevant to the scene outliner */
-	bool IsDropOperationApplicable(TSharedPtr<FDragDropOperation> Operation)
+	static FReply OnDragDetected( const FGeometry& MyGeometry, const FPointerEvent& MouseEvent, TWeakPtr<SOutlinerTreeView> Table )
 	{
-		// We're only interested in actor drags and folder drags
-		return Operation.IsValid() && (Operation->IsOfType<FActorDragDropGraphEdOp>() || Operation->IsOfType<FFolderDragDropOp>());
-	}
-
-	/** Compute validation information for dropping folder(s) onto another folder */
-	static FDragValidationInfo ValidateDropOnFolder(TSharedRef<SSceneOutliner> SceneOutlinerRef, const FolderArray& DragFolders, FName DropTargetPath)
-	{
-		// Iterate over all the folders that have been dragged
-		for (auto FolderIt = DragFolders.CreateConstIterator(); FolderIt; ++FolderIt)
+		auto TablePtr = Table.Pin();
+		if (TablePtr.IsValid() && MouseEvent.IsMouseButtonDown( EKeys::LeftMouseButton ))
 		{
-			auto DragFolderPtr = FolderIt->Pin();
-			if (DragFolderPtr.IsValid())
+			auto Operation = SceneOutliner::CreateDragDropOperation(TablePtr->GetSelectedItems());
+			if (Operation.IsValid())
 			{
-				if (DragFolderPtr->GetParentPath() == DropTargetPath)
-				{
-					FFormatNamedArguments Args;
-					Args.Add(TEXT("SourceName"), FText::FromName(DragFolderPtr->LeafName));
-
-					FText Text;
-					if (DropTargetPath.IsNone())
-					{
-						Text = FText::Format(LOCTEXT("FolderAlreadyAssignedRoot", "{SourceName} is already assigned to root"), Args);
-					}
-					else
-					{
-						Args.Add(TEXT("DestPath"), FText::FromName(DropTargetPath));
-						Text = FText::Format(LOCTEXT("FolderAlreadyAssigned", "{SourceName} is already assigned to {DestPath}"), Args);
-					}
-					
-					return FDragValidationInfo(FActorDragDropGraphEdOp::ToolTip_IncompatibleGeneric, Text);
-				}
-
-				const FString DragFolderPath = DragFolderPtr->Path.ToString();
-				const FString LeafName = DragFolderPtr->LeafName.ToString();
-				const FString DstFolderPath = DropTargetPath.IsNone() ? FString() : DropTargetPath.ToString();
-				const FString NewPath = DstFolderPath / LeafName;
-
-				if (SceneOutlinerRef->FindFolderByPath(FName(*NewPath)).IsValid())
-				{
-					// The folder already exists
-					FFormatNamedArguments Args;
-					Args.Add(TEXT("DragName"), FText::FromString(LeafName));
-					return FDragValidationInfo(FActorDragDropGraphEdOp::ToolTip_IncompatibleGeneric,
-						FText::Format(LOCTEXT("FolderAlreadyExistsRoot", "A folder called \"{DragName}\" already exists at this level"), Args));
-				}
-				else if (DragFolderPath == DstFolderPath || DstFolderPath.StartsWith(DragFolderPath + "/"))
-				{
-					// Cannot drag as a child of itself
-					FFormatNamedArguments Args;
-					Args.Add(TEXT("FolderPath"), FText::FromName(DragFolderPtr->Path));
-					return FDragValidationInfo(
-						FActorDragDropGraphEdOp::ToolTip_IncompatibleGeneric,
-						FText::Format(LOCTEXT("ChildOfItself", "Cannot move \"{FolderPath}\" to be a child of itself"), Args));
-				}
+				return FReply::Handled().BeginDragDrop(Operation.ToSharedRef());
 			}
 		}
 
-		// Everything else is a valid operation
-		if (DropTargetPath.IsNone())
-		{
-			return FDragValidationInfo(FActorDragDropGraphEdOp::ToolTip_CompatibleGeneric, LOCTEXT("MoveToRoot", "Move to root"));
-		}
-		else
-		{
-			return FDragValidationInfo(FActorDragDropGraphEdOp::ToolTip_CompatibleGeneric, FText::Format(LOCTEXT("MoveInto", "Move into \"{DropPath}\""), FText::FromName(DropTargetPath)));
-		}
+		return FReply::Unhandled();
 	}
 
-	/** Compute validation information for dropping actor(s) onto another actor */
-	static FDragValidationInfo ValidateDropOnActor(const ActorArray& DragActors, const AActor& DropTarget)
+	FReply HandleDrop(TWeakPtr<ISceneOutliner> SceneOutlinerWeak, const FDragDropEvent& DragDropEvent, IDropTarget& DropTarget, FDragValidationInfo& ValidationInfo, bool bApplyDrop = false)
 	{
-		FText AttachErrorMsg;
-		bool CanAttach = true;
-		bool DraggedOntoAttachmentParent = true;
-
-		for (auto ActorIt = DragActors.CreateConstIterator(); ActorIt; ++ActorIt)
+		auto SceneOutlinerPtr = SceneOutlinerWeak.Pin();
+		if (!SceneOutlinerPtr.IsValid())
 		{
-			AActor* DragActor = ActorIt->Get();
-			if (DragActor)
-			{
-				if (!GEditor->CanParentActors(&DropTarget, DragActor, &AttachErrorMsg))
-				{
-					CanAttach = false;
-				}
-				if (DragActor->GetAttachParentActor() != &DropTarget)
-				{
-					DraggedOntoAttachmentParent = false;
-				}
-			}
+			return FReply::Unhandled();
 		}
 
-		if (DraggedOntoAttachmentParent)
+		const FSharedOutlinerData& SharedData = SceneOutlinerPtr->GetSharedData();
+		if (SharedData.Mode != ESceneOutlinerMode::ActorBrowsing)
 		{
-			if (DragActors.Num() == 1)
-			{
-				return FDragValidationInfo(FActorDragDropGraphEdOp::ToolTip_CompatibleDetach, FText::FromString(DropTarget.GetActorLabel()));
-			}
-			else
-			{
-				return FDragValidationInfo(FActorDragDropGraphEdOp::ToolTip_CompatibleMultipleDetach, FText::FromString(DropTarget.GetActorLabel()));
-			}
+			return FReply::Unhandled();	
 		}
-		else if (CanAttach)
+
+		// Don't handle this if we're not showing a hierarchy, not in browsing mode, or the drop operation is not applicable
+		if (!SharedData.bShowParentTree || !SharedData.RepresentingWorld)
 		{
-			if (DragActors.Num() == 1)
-			{
-				return FDragValidationInfo(FActorDragDropGraphEdOp::ToolTip_CompatibleAttach, FText::FromString(DropTarget.GetActorLabel()));
-			}
-			else
-			{
-				return FDragValidationInfo(FActorDragDropGraphEdOp::ToolTip_CompatibleMultipleAttach, FText::FromString(DropTarget.GetActorLabel()));
-			}
+			return FReply::Unhandled();
 		}
-		else
+
+		FDragDropPayload DraggedObjects;
+		// Validate now to make sure we don't doing anything we shouldn't
+		if (!DraggedObjects.ParseDrag(*DragDropEvent.GetOperation()))
 		{
-			if (DragActors.Num() == 1)
-			{
-				return FDragValidationInfo(FActorDragDropGraphEdOp::ToolTip_IncompatibleGeneric, AttachErrorMsg);
-			}
-			else
-			{
-				const FText ReasonText = FText::Format(LOCTEXT("DropOntoText", "{DropTarget}. {AttachMessage}"), FText::FromString(DropTarget.GetActorLabel()), AttachErrorMsg);
-				return FDragValidationInfo(FActorDragDropGraphEdOp::ToolTip_IncompatibleMultipleAttach, ReasonText);
-			}
+			// Return handled here to stop anything else trying to handle it - the operation is invalid as far as we're concerned
+			return FReply::Handled();
 		}
-	}
 
-	/** Check that it is valid to drop the specified folders and/or actors to the root */
-	static FDragValidationInfo ValidateDropOnRoot(TSharedRef<SSceneOutliner> SceneOutlinerRef, const ActorArray* DraggedActors, const FolderArray* DraggedFolders)
-	{
-		if (DraggedFolders && DraggedFolders->Num())
-		{
-			// Dropping folders at the root
-			const FDragValidationInfo ValidationInfo = ValidateDropOnFolder(SceneOutlinerRef, *DraggedFolders, FName());
+		ValidationInfo = DropTarget.ValidateDrop(DraggedObjects, *SharedData.RepresentingWorld);
 
-			// If the operation is invalid, or we don't have any actors to check for, return immediately
-			// If it is valid then we need to drop down and validate the actors below.
-			if (!ValidationInfo.IsValid() || !DraggedActors || DraggedActors->Num() == 0)
-			{
-				return ValidationInfo;
-			}
+		if (!ValidationInfo.IsValid())
+		{
+			// Return handled here to stop anything else trying to handle it - the operation is invalid as far as we're concerned
+			return FReply::Handled();
 		}
 
-		if (DraggedActors && DraggedActors->Num() > 0)
+		if (bApplyDrop)
 		{
-			// Actors are always valid
-			return FDragValidationInfo(FActorDragDropGraphEdOp::ToolTip_CompatibleGeneric, LOCTEXT("MoveToRoot", "Move to root"));
+			DropTarget.OnDrop(DraggedObjects, *SharedData.RepresentingWorld, ValidationInfo, SceneOutlinerPtr.ToSharedRef());
 		}
 
-		return FDragValidationInfo::Invalid();
+		return FReply::Handled();
 	}
 
-	/** Validate a drop event that wants to drop at the root of the tree */
-	static FDragValidationInfo ValidateRootDragDropEvent(TSharedRef<SSceneOutliner> SceneOutlinerRef, const FDragDropEvent& DragDropEvent)
+	struct FDropHandler : ITreeItemVisitor
 	{
-		const FolderArray* DraggedFolders = nullptr;
-		const ActorArray* DraggedActors = nullptr;
-
-		ExtractDraggedActorsAndFolders(DragDropEvent.GetOperation(), DraggedActors, DraggedFolders);
+		/** Pointer to the scene outliner */
+		const TWeakPtr<SSceneOutliner>& SceneOutlinerWeak;
 
-		return ValidateDropOnRoot(SceneOutlinerRef, DraggedActors, DraggedFolders);
-	}
-
-	/** Validate a drop event that wants to drop onto the specified destination row */
-	static FDragValidationInfo ValidateDropOnItem(TSharedRef<SSceneOutliner> SceneOutlinerRef, TSharedRef<TOutlinerTreeItem> Destination, const ActorArray* DraggedActors, const FolderArray* DraggedFolders)
-	{
-		if (DraggedFolders)
-		{
-			if (DraggedFolders->Num())
-			{
-				if (Destination->Type == TOutlinerTreeItem::Actor)
-				{
-					return FDragValidationInfo(FActorDragDropGraphEdOp::ToolTip_IncompatibleGeneric, LOCTEXT("FoldersOnActorError", "Cannot attach folders to actors"));
-				}
-				else
-				{
-					// Dropping folders onto a folder
-					auto DragFolder = StaticCastSharedRef<TOutlinerFolderTreeItem>(Destination);
+		/** The drag/drop event we are processing */
+		const FDragDropEvent& DragDropEvent;
 
-					const FDragValidationInfo FolderValidation = ValidateDropOnFolder(SceneOutlinerRef, *DraggedFolders, DragFolder->Path);
+		/** Whether or not we should actually apply the drop, or just validate it */
+		bool bApplyDrop;
 
-					// If the operation is invalid, or we don't have any actors to check for, return immediately
-					// If it is valid then we need to drop down and validate the actors below.
-					if (!FolderValidation.IsValid() || !DraggedActors || DraggedActors->Num() == 0)
-					{
-						return FolderValidation;
-					}
-				}
-			}
-		}
+		/** The resulting validation information */
+		mutable FDragValidationInfo ValidationInfo;
+		
+		/** The reply we shouyld return, if applicable */
+		mutable FReply Reply;
+
+		/** Constructor */
+		FDropHandler(const TWeakPtr<SSceneOutliner>& InSceneOutlinerWeak, const FDragDropEvent& InDragDropEvent, bool bInApplyDrop)
+			: SceneOutlinerWeak(InSceneOutlinerWeak)
+			, DragDropEvent(InDragDropEvent)
+			, bApplyDrop(bInApplyDrop)
+			, ValidationInfo(FDragValidationInfo::Invalid())
+			, Reply(FReply::Unhandled())
+		{}
 
-		//	--------------------------------------------------------------------------------
-		//	Validate dragging actors on the scene outliner
-		if (DraggedActors)
+		virtual void Visit(const FActorTreeItem& ActorItem) const override
 		{
-			if (Destination->Type == TOutlinerTreeItem::Actor)
-			{
-				// Dragged into an actor
-				auto ActorTreeItem = StaticCastSharedRef<TOutlinerActorTreeItem>(Destination);
-				AActor* DropActor = ActorTreeItem->Actor.Get();
-				if (DropActor)
-				{
-					return ValidateDropOnActor(*DraggedActors, *DropActor);
-				}
-			}
-			else
-			{
-				auto DropFolder = StaticCastSharedRef<TOutlinerFolderTreeItem>(Destination);
-
-				FFormatNamedArguments Args;
-				Args.Add(TEXT("DestName"), FText::FromName(DropFolder->LeafName));
-
-				// Check to see if any of the actors are already assigned to the destination
-				for (const auto& WeakActor : *DraggedActors)
-				{
-					AActor* DragActor = WeakActor.Get();
-					if (DragActor && !DragActor->GetAttachParentActor() && DragActor->GetFolderPath() == DropFolder->Path)
-					{
-						Args.Add(TEXT("SourceName"), FText::FromString(DragActor->GetActorLabel()));
-						const FText Text = FText::Format(LOCTEXT("MoveToFolderAlreadyAssigned", "{SourceName} is already assigned to {DestName}"), Args);
-						return FDragValidationInfo(FActorDragDropGraphEdOp::ToolTip_IncompatibleGeneric, Text);
-					}
-				}
-
-				const int32 NumDraggedFolders = DraggedFolders ? DraggedFolders->Num() : 0;
-				const int32 NumDraggedItems = DraggedActors->Num() + NumDraggedFolders;
-
-				// Calculate a singular drag operation source name
-				if (NumDraggedItems == 1)
-				{
-					if (DraggedActors->Num() == 1)
-					{
-						const AActor* Actor = DraggedActors->GetData()[0].Get();
-						if (Actor)
-						{
-							Args.Add(TEXT("SourceName"), FText::FromString(Actor->GetActorLabel()));
-						}
-					}
-					else
-					{
-						const auto DraggedFolder = DraggedFolders->GetData()[0].Pin();
-						if (DraggedFolder.IsValid())
-						{
-							Args.Add(TEXT("SourceName"), FText::FromString(DraggedFolder->LeafName.ToString()));
-						}
-					}
-				}
-
-				FText Text;
-				if (Args.Contains(TEXT("SourceName")))
-				{
-					Text = FText::Format(LOCTEXT("MoveToFolderSingle", "Move {SourceName} into {DestName}"), Args);
-				}
-				else
-				{
-					Text = FText::Format(LOCTEXT("MoveToFolderMulti", "Move items into {DestName}"), Args);
-				}
-
-				return FDragValidationInfo(FActorDragDropGraphEdOp::ToolTip_CompatibleGeneric, Text);
-			}
+			FActorDropTarget DropTarget(ActorItem.Actor.Get());
+			Reply = HandleDrop(SceneOutlinerWeak, DragDropEvent, DropTarget, ValidationInfo, bApplyDrop);
 		}
-
-		return FDragValidationInfo::Invalid();
-	}
-
-	/** Validate a drop event that wants to drop onto the specified destination row */
-	static FDragValidationInfo ValidateDragDropEvent(TSharedRef<SSceneOutliner> SceneOutlinerRef, const FDragDropEvent& DragDropEvent, TSharedRef<TOutlinerTreeItem> Destination)
-	{
-		// Only allow dropping objects when in actor browsing mode
-		if (SceneOutlinerRef->GetInitOptions().Mode != ESceneOutlinerMode::ActorBrowsing)
+		virtual void Visit(const FFolderTreeItem& FolderItem) const override
 		{
-			return FDragValidationInfo::Invalid();
+			FFolderDropTarget DropTarget(FolderItem.Path);
+			Reply = HandleDrop(SceneOutlinerWeak, DragDropEvent, DropTarget, ValidationInfo, bApplyDrop);
 		}
-
-		//	Validate dragging a mixture of actors and folders inside the scene outliner
-		const FolderArray* DraggedFolders = nullptr;
-		const ActorArray* DraggedActors = nullptr;
-
-		ExtractDraggedActorsAndFolders(DragDropEvent.GetOperation(), DraggedActors, DraggedFolders);
-
-		return ValidateDropOnItem(SceneOutlinerRef, Destination, DraggedActors, DraggedFolders);
-	}
+	};
 
 	void SOutlinerTreeView::Construct(const SOutlinerTreeView::FArguments& InArgs, TSharedRef<SSceneOutliner> Owner)
 	{
@@ -368,102 +149,27 @@ namespace SceneOutliner
 		STreeView::Construct(InArgs);
 	}
 
-	void SOutlinerTreeView::FlashHighlightOnItem( FOutlinerTreeItemPtr FlashHighlightOnItem )
+	void SOutlinerTreeView::FlashHighlightOnItem( FTreeItemPtr FlashHighlightOnItem )
 	{
 		TSharedPtr< SSceneOutlinerTreeRow > RowWidget = StaticCastSharedPtr< SSceneOutlinerTreeRow >( WidgetGenerator.GetWidgetForItem( FlashHighlightOnItem ) );
 		if( RowWidget.IsValid() )
 		{
-			TSharedPtr< SSceneOutlinerItemWidget > ItemWidget = RowWidget->GetItemWidget();
-			if( ItemWidget.IsValid() )
-			{
-				ItemWidget->FlashHighlight();
-			}
+			RowWidget->FlashHighlight();
 		}
 	}
 
-	bool SOutlinerTreeView::ValidateMoveSelectionTo(FName NewParent, FText* OutValidationText)
+	FReply SOutlinerTreeView::OnDragOver(const FGeometry& MyGeometry, const FDragDropEvent& DragDropEvent)
 	{
-		auto SceneOutlinerPtr = SceneOutlinerWeak.Pin();
-		if (!SceneOutlinerPtr.IsValid())
-		{
-			return false;
-		}
-
-		FolderArray Folders;
-		ActorArray Actors;
-
-		for (const auto& Item : GetSelectedItems())
-		{
-			if (Item->Type == TOutlinerTreeItem::Folder)
-			{
-				Folders.Add(StaticCastSharedPtr<TOutlinerFolderTreeItem>(Item));
-			}
-			else
-			{
-				Actors.Add(StaticCastSharedPtr<TOutlinerActorTreeItem>(Item)->Actor);
-			}
-		}
-
 		FDragValidationInfo ValidationInfo = FDragValidationInfo::Invalid();
-		if (NewParent.IsNone())
-		{
-			ValidationInfo = ValidateDropOnRoot(SceneOutlinerPtr.ToSharedRef(), &Actors, &Folders);
-		}
-		else
-		{
-			auto Folder = SceneOutlinerPtr->FindFolderByPath(NewParent);
-			if (Folder.IsValid())
-			{
-				ValidationInfo = ValidateDropOnItem(SceneOutlinerPtr.ToSharedRef(), Folder.ToSharedRef(), &Actors, &Folders);
-			}
-			else
-			{
-				ValidationInfo.ValidationText = FText::Format(LOCTEXT("UnableToFindFolderFailure", "Unable to find folder: {0}"), FText::FromName(NewParent));
-			}
-		}
-
-		if (OutValidationText)
-		{
-			*OutValidationText = ValidationInfo.ValidationText;
-		}
-		return ValidationInfo.IsValid();
-	}
 
-	FReply SOutlinerTreeView::OnDragOver(const FGeometry& MyGeometry, const FDragDropEvent& DragDropEvent)
-	{
-		auto SceneOutlinerPtr = SceneOutlinerWeak.Pin();
-		if (!SceneOutlinerPtr.IsValid())
-		{
-			return FReply::Unhandled();
-		}
+		FFolderDropTarget DropTarget(NAME_None);
 
+		auto Reply = HandleDrop(SceneOutlinerWeak, DragDropEvent, DropTarget, ValidationInfo);
 
-		TSharedPtr<FDecoratedDragDropOp> DecoratedDragOp = DragDropEvent.GetOperationAs<FDecoratedDragDropOp>();
-		if (DecoratedDragOp.IsValid())
+		if (Reply.IsEventHandled())
 		{
-			const FDragValidationInfo ValidationInfo = ValidateRootDragDropEvent(SceneOutlinerPtr.ToSharedRef(), DragDropEvent);
-
-			if (ValidationInfo.TooltipType == FActorDragDropGraphEdOp::ToolTip_CompatibleGeneric || ValidationInfo.TooltipType == FActorDragDropGraphEdOp::ToolTip_IncompatibleGeneric)
-			{
-				// Generic messages can go through the decorated drag/drop operation base class
-				DecoratedDragOp->CurrentHoverText = ValidationInfo.ValidationText;
-				if (ValidationInfo.TooltipType == FActorDragDropGraphEdOp::ToolTip_CompatibleGeneric)
-				{
-					DecoratedDragOp->CurrentIconBrush = FEditorStyle::GetBrush(TEXT("Graph.ConnectorFeedback.OK"));
-				}
-				else
-				{
-					DecoratedDragOp->CurrentIconBrush = FEditorStyle::GetBrush(TEXT("Graph.ConnectorFeedback.Error"));
-				}
-				return FReply::Handled();
-			}
-			else if (DecoratedDragOp->IsOfType<FActorDragDropGraphEdOp>())
-			{
-				// Display actor specific messages
-				auto ActorDragOp = StaticCastSharedPtr<FActorDragDropGraphEdOp>(DecoratedDragOp);
-				ActorDragOp->SetToolTip(ValidationInfo.TooltipType, ValidationInfo.ValidationText);
-				return FReply::Handled();
-			}
+			UpdateOperationDecorator(DragDropEvent, ValidationInfo);
+			return FReply::Handled();
 		}
 
 		return FReply::Unhandled();
@@ -471,308 +177,94 @@ namespace SceneOutliner
 	
 	void SOutlinerTreeView::OnDragLeave(const FDragDropEvent& DragDropEvent)
 	{
-		if( !SceneOutlinerWeak.Pin()->GetInitOptions().bShowParentTree )
+		if( SceneOutlinerWeak.Pin()->GetSharedData().bShowParentTree )
 		{
-			// We only allowing dropping to perform attachments when displaying the actors
-			// in their parent/child hierarchy
-			return;
-		}
-
-		auto DecoratedOp = DragDropEvent.GetOperationAs<FDecoratedDragDropOp>();
-		if (DecoratedOp.IsValid())
-		{
-			DecoratedOp->ResetToDefaultToolTip();
+			ResetOperationDecorator(DragDropEvent);
 		}
 	}
 
 	FReply SOutlinerTreeView::OnDrop(const FGeometry& MyGeometry, const FDragDropEvent& DragDropEvent)
 	{
-		auto SceneOutlinerPtr = SceneOutlinerWeak.Pin();
-
-		// Don't handle this if we're not showing a hierarchy, not in browsing mode, or the drop operation is not applicable
-		if (!SceneOutlinerPtr.IsValid() ||
-			!SceneOutlinerPtr->GetInitOptions().bShowParentTree ||
-			SceneOutlinerPtr->GetInitOptions().Mode != ESceneOutlinerMode::ActorBrowsing ||
-			!IsDropOperationApplicable(DragDropEvent.GetOperation())
-			)
-		{
-			return FReply::Unhandled();
-		}
-
-		// Validate now to make sure we don't doing anything we shouldn't
-		const FDragValidationInfo ValidationInfo = ValidateRootDragDropEvent(SceneOutlinerPtr.ToSharedRef(), DragDropEvent);
-		if (!ValidationInfo.IsValid())
-		{
-			return FReply::Handled();
-		}
-
-		const FScopedTransaction Transaction(LOCTEXT("UndoAction_DropOnSceneOutliner", "Drop on Scene Outliner"));
-		
-		const FolderArray* DraggedFolders = nullptr;
-		const ActorArray* DraggedActors = nullptr;
-		ExtractDraggedActorsAndFolders(DragDropEvent.GetOperation(), DraggedActors, DraggedFolders);
-
-		if (DraggedFolders)
-		{
-			for (auto& Folder : *DraggedFolders)
-			{
-				auto FolderPtr = Folder.Pin();
-				if (FolderPtr.IsValid())
-				{
-					SceneOutlinerPtr->MoveFolderTo(FolderPtr.ToSharedRef(), FName());
-				}
-			}
-		}
-
-		if (DraggedActors)
-		{
-			for (auto& Actor : *DraggedActors)
-			{
-				auto* ActorPtr = Actor.Get();
-				if (ActorPtr)
-				{
-					ActorPtr->SetFolderPath(FName());
-				}
-			}
-		}
+		FDragValidationInfo ValidationInfo = FDragValidationInfo::Invalid();
+		FFolderDropTarget DropTarget(NAME_None);
 
-		return FReply::Handled();
+		return HandleDrop(SceneOutlinerWeak, DragDropEvent, DropTarget, ValidationInfo, true);
 	}
 
 	FReply SSceneOutlinerTreeRow::OnDrop( const FGeometry& MyGeometry, const FDragDropEvent& DragDropEvent )
 	{
-		auto SceneOutlinerPtr = SceneOutlinerWeak.Pin();
-
-		// Don't handle this if we're not showing a hierarchy, not in browsing mode, or the drop operation is not applicable
-		if (!SceneOutlinerPtr.IsValid() ||
-			!SceneOutlinerPtr->GetInitOptions().bShowParentTree ||
-			SceneOutlinerPtr->GetInitOptions().Mode != ESceneOutlinerMode::ActorBrowsing ||
-			!IsDropOperationApplicable(DragDropEvent.GetOperation())
-			)
-		{
-			return FReply::Unhandled();
-		}
-
-		// Validate now to make sure we don't doing anything we shouldn't
-		const FDragValidationInfo ValidationInfo = ValidateDragDropEvent(SceneOutlinerPtr.ToSharedRef(), DragDropEvent, Item.ToSharedRef());
-		if (!ValidationInfo.IsValid())
-		{
-			// Return handled here to stop anything else trying to handle it - the operation is invalid as far as we're concerned
-			return FReply::Handled();
-		}
-
-		const FScopedTransaction Transaction(LOCTEXT("UndoAction_DropOnSceneOutliner", "Drop on Scene Outliner"));
-
-		const FolderArray* DraggedFolders = nullptr;
-		const ActorArray* DraggedActors = nullptr;
-
-		ExtractDraggedActorsAndFolders(DragDropEvent.GetOperation(), DraggedActors, DraggedFolders);
-
-		// Drop on a folder?
-		if (Item->Type == TOutlinerTreeItem::Folder && DraggedFolders)
-		{
-			auto DropFolder = StaticCastSharedPtr<TOutlinerFolderTreeItem>(Item);
-			for (auto& FolderIt : *DraggedFolders)
-			{
-				auto FolderItemPtr = FolderIt.Pin();
-				if (FolderItemPtr.IsValid())
-				{
-					SceneOutlinerPtr->MoveFolderTo(FolderItemPtr.ToSharedRef(), DropFolder->Path);
-				}
-			}
-		}
-
-		if (DraggedActors)
+		auto ItemPtr = Item.Pin();
+		if (ItemPtr.IsValid())
 		{
-			// Dragged some actors, deal with those now
-			if (Item->Type == TOutlinerTreeItem::Actor)
-			{
-				// Dropped on an actor
-				auto ActorTreeItem = StaticCastSharedPtr<TOutlinerActorTreeItem>(Item);
-				AActor* DropActor = ActorTreeItem->Actor.Get();
-				if (!DropActor)
-				{
-					return FReply::Handled();
-				}
-
-				FMessageLog EditorErrors("EditorErrors");
-				EditorErrors.NewPage(LOCTEXT("ActorAttachmentsPageLabel", "Actor attachment"));
-
-				if (ValidationInfo.TooltipType == FActorDragDropGraphEdOp::ToolTip_CompatibleMultipleDetach || ValidationInfo.TooltipType == FActorDragDropGraphEdOp::ToolTip_CompatibleDetach)
-				{
-					PerformDetachment(ActorTreeItem->Actor);
-				}
-				else if (ValidationInfo.TooltipType == FActorDragDropGraphEdOp::ToolTip_CompatibleMultipleAttach || ValidationInfo.TooltipType == FActorDragDropGraphEdOp::ToolTip_CompatibleAttach)
-				{
-					// Show socket chooser if we have sockets to select
-
-					//@TODO: Should create a menu for each component that contains sockets, or have some form of disambiguation within the menu (like a fully qualified path)
-					// Instead, we currently only display the sockets on the root component
-					USceneComponent* Component = DropActor->GetRootComponent();
-					if ((Component != NULL) && (Component->HasAnySockets()))
-					{
-						TWeakObjectPtr<AActor> Actor = ActorTreeItem->Actor;
-						// Create the popup
-						FSlateApplication::Get().PushMenu(
-							SceneOutlinerPtr.ToSharedRef(),
-							SNew(SSocketChooserPopup)
-							.SceneComponent( Component )
-							.OnSocketChosen(this, &SSceneOutlinerTreeRow::PerformAttachment, Actor, *DraggedActors ),
-							FSlateApplication::Get().GetCursorPos(),
-							FPopupTransitionEffect( FPopupTransitionEffect::TypeInPopup )
-							);
-					}
-					else
-					{
-						PerformAttachment(NAME_None, ActorTreeItem->Actor, *DraggedActors);
-					}
-				}
-
-				// Report errors
-				EditorErrors.Notify(NSLOCTEXT("ActorAttachmentError", "AttachmentsFailed", "Attachments Failed!"));
-			}
-			else
-			{
-				// Dropped on a folder
-				auto FolderTreeItem = StaticCastSharedPtr<TOutlinerFolderTreeItem>(Item);
-				for (auto& ActorIt : *DraggedActors)
-				{
-					auto* ActorPtr = ActorIt.Get();
-					if (ActorPtr)
-					{
-						ActorPtr->SetFolderPath(FolderTreeItem->Path);
-					}
-				}
-			}
+			FDropHandler Handler(SceneOutlinerWeak, DragDropEvent, true);
+			ItemPtr->Visit(Handler);
+			return Handler.Reply;
 		}
 
-		return FReply::Handled();
+		return FReply::Unhandled();
 	}
 
 	void SSceneOutlinerTreeRow::OnDragEnter( const FGeometry& MyGeometry, const FDragDropEvent& DragDropEvent )
 	{
-		auto SceneOutlinerPtr = SceneOutlinerWeak.Pin();
-		if( !SceneOutlinerPtr.IsValid() || !SceneOutlinerPtr->GetInitOptions().bShowParentTree )
+		auto ItemPtr = Item.Pin();
+		if (ItemPtr.IsValid())
 		{
-			// We only allowing dropping to perform attachments when displaying the actors
-			// in their parent/child hierarchy
-			return;
-		}
-
-		// Perform the validation and update the tooltip if necessary
-		const FDragValidationInfo ValidationInfo = ValidateDragDropEvent(SceneOutlinerPtr.ToSharedRef(), DragDropEvent, Item.ToSharedRef());
+			FDropHandler Handler(SceneOutlinerWeak, DragDropEvent, false);
+			ItemPtr->Visit(Handler);
 
-		TSharedPtr<FActorDragDropGraphEdOp> ActorDragOp = DragDropEvent.GetOperationAs<FActorDragDropGraphEdOp>();
-		if (ActorDragOp.IsValid())
-		{
-			ActorDragOp->SetToolTip(ValidationInfo.TooltipType, ValidationInfo.ValidationText);
-			return;
-		}
-
-		TSharedPtr<FDecoratedDragDropOp> DecoratedDragOp = DragDropEvent.GetOperationAs<FDecoratedDragDropOp>();
-		if (DecoratedDragOp.IsValid())
-		{
-			DecoratedDragOp->CurrentHoverText = ValidationInfo.ValidationText;
-			DecoratedDragOp->CurrentIconBrush = ValidationInfo.IsValid() ? FEditorStyle::GetBrush(TEXT("Graph.ConnectorFeedback.OK")) : FEditorStyle::GetBrush(TEXT("Graph.ConnectorFeedback.Error"));
+			UpdateOperationDecorator(DragDropEvent, Handler.ValidationInfo);
 		}
 	}
 
-	void SSceneOutlinerTreeRow::PerformAttachment( FName SocketName, TWeakObjectPtr<AActor> ParentActorPtr, const TArray< TWeakObjectPtr<AActor> > ChildrenPtrs )
+	void SSceneOutlinerTreeRow::OnDragLeave( const FDragDropEvent& DragDropEvent )
 	{
-		AActor* ParentActor = ParentActorPtr.Get();
-		if(ParentActor != NULL)
-		{
-			// modify parent and child
-			const FScopedTransaction Transaction( LOCTEXT("UndoAction_PerformAttachment", "Attach actors") );
-
-			// Attach each child
-			bool bAttached = false;
-			for(int32 i = 0; i < ChildrenPtrs.Num(); i++)
-			{
-				AActor* ChildActor = ChildrenPtrs[i].Get();
-				if (GEditor->CanParentActors(ParentActor, ChildActor))
-				{
-					GEditor->ParentActors(ParentActor, ChildActor, SocketName);
-					bAttached = true;
-				}
-			}
-
-			// refresh the tree, and ensure parent is expanded so we can still see the child if we attached something
-			if (bAttached)
-			{
-				SceneOutlinerWeak.Pin()->ExpandActor(ParentActor);
-			}
-		}
+		ResetOperationDecorator(DragDropEvent);
 	}
 
-	void SSceneOutlinerTreeRow::PerformDetachment( TWeakObjectPtr<AActor> ParentActorPtr )
+	FReply SSceneOutlinerTreeRow::OnMouseButtonDown( const FGeometry& MyGeometry, const FPointerEvent& MouseEvent )
 	{
-		AActor* ParentActor = ParentActorPtr.Get();
-		if(ParentActor != NULL)
+		auto ItemPtr = Item.Pin();
+		if (ItemPtr.IsValid() && ItemPtr->CanInteract())
 		{
-			bool bAttached = GEditor->DetachSelectedActors();
-			// refresh the tree, and ensure parent is expanded so we can still see the child if we attached something
-			if (bAttached)
+			if (MouseEvent.GetEffectingButton() == EKeys::LeftMouseButton)
 			{
-				SceneOutlinerWeak.Pin()->ExpandActor(ParentActor);
-			}
-		}
-	}
+				FReply Reply = SMultiColumnTableRow<FTreeItemPtr>::OnMouseButtonDown( MyGeometry, MouseEvent );
 
-	FReply SSceneOutlinerTreeRow::HandleDragDetected( const FGeometry& MyGeometry, const FPointerEvent& MouseEvent )
-	{
-		if ( MouseEvent.IsMouseButtonDown( EKeys::LeftMouseButton ) )
-		{
-			auto SceneOutlinerPtr = SceneOutlinerWeak.Pin();
-			if (!SceneOutlinerPtr.IsValid())
-			{
-				return FReply::Unhandled();
-			}
+				// We only support drag and drop when in actor browsing mode
+				if( SceneOutlinerWeak.Pin()->GetSharedData().Mode == ESceneOutlinerMode::ActorBrowsing )
+				{
+					return Reply.DetectDrag( SharedThis(this) , EKeys::LeftMouseButton );
+				}
 
-			// We only support drag and drop when in actor browsing mode
-			if( SceneOutlinerPtr->GetInitOptions().Mode == ESceneOutlinerMode::ActorBrowsing )
-			{
-				auto SelectedItems = SceneOutlinerPtr->GetSelectedItems();
-				return FReply::Handled().BeginDragDrop(SceneOutliner::CreateDragDropOperation(SelectedItems));
+				return Reply.PreventThrottling();
 			}
 		}
 
 		return FReply::Unhandled();
 	}
 
-	FReply SSceneOutlinerTreeRow::OnMouseButtonDown( const FGeometry& MyGeometry, const FPointerEvent& MouseEvent )
+	TSharedRef<SWidget> SSceneOutlinerTreeRow::GenerateWidgetForColumn( const FName& ColumnName )
 	{
-		if (MouseEvent.GetEffectingButton() == EKeys::LeftMouseButton)
+		auto ItemPtr = Item.Pin();
+		if (!ItemPtr.IsValid())
 		{
-			FReply Reply = SMultiColumnTableRow<FOutlinerTreeItemPtr>::OnMouseButtonDown( MyGeometry, MouseEvent );
-
-			// We only support drag and drop when in actor browsing mode
-			if( SceneOutlinerWeak.Pin()->GetInitOptions().Mode == ESceneOutlinerMode::ActorBrowsing )
-			{
-				return Reply.DetectDrag( SharedThis(this) , EKeys::LeftMouseButton );
-			}
-
-			return Reply.PreventThrottling();
+			return SNullWidget::NullWidget;
 		}
 
-		return FReply::Unhandled();
-	}
 
-	TSharedPtr< SSceneOutlinerItemWidget > SSceneOutlinerTreeRow::GetItemWidget()
-	{
-		return ItemWidget;
-	}
+		auto Outliner = SceneOutlinerWeak.Pin();
+		check(Outliner.IsValid());
 
-	TSharedRef<SWidget> SSceneOutlinerTreeRow::GenerateWidgetForColumn( const FName& ColumnName )
-	{
 		// Create the widget for this item
-		TSharedRef< SWidget > NewItemWidget = SceneOutlinerWeak.Pin()->GenerateWidgetForItemAndColumn( Item, ColumnName, FIsSelected::CreateSP(this, &SSceneOutlinerTreeRow::IsSelectedExclusively) );
-		if( ColumnName == ColumnID_ActorLabel )
+		TSharedRef<SWidget> NewItemWidget = SNullWidget::NullWidget;
+
+		auto Column = Outliner->GetColumns().FindRef(ColumnName);
+		if (Column.IsValid())
 		{
-			ItemWidget = StaticCastSharedPtr<SSceneOutlinerItemWidget>( TSharedPtr< SWidget >(NewItemWidget) );
+			NewItemWidget = Column->ConstructRowWidget(ItemPtr.ToSharedRef(), *this);
 		}
 
-		if( ColumnName == ColumnID_ActorLabel )
+		if( ColumnName == FBuiltInColumnTypes::Label() )
 		{
 			// The first column gets the tree expansion arrow for this row
 			return 
@@ -798,16 +290,122 @@ namespace SceneOutliner
 		}
 	}
 
-	void SSceneOutlinerTreeRow::Construct( const FArguments& InArgs, const TSharedRef<STableViewBase>& InOwnerTableView )
+	void SSceneOutlinerTreeRow::Construct( const FArguments& InArgs, const TSharedRef<SOutlinerTreeView>& OutlinerTreeView, TSharedRef<SSceneOutliner> SceneOutliner )
+	{
+		Item = InArgs._Item->AsShared();
+		SceneOutlinerWeak = SceneOutliner;
+
+		auto Args = FSuperRowType::FArguments()
+			.Style(&FEditorStyle::Get().GetWidgetStyle<FTableRowStyle>("SceneOutliner.TableViewRow"));
+
+		// We only support drag and drop when in actor browsing mode
+		if (SceneOutliner->GetSharedData().Mode == ESceneOutlinerMode::ActorBrowsing)
+		{
+			Args.OnDragDetected_Static(SceneOutliner::OnDragDetected, TWeakPtr<SOutlinerTreeView>(OutlinerTreeView));
+		}
+
+		SMultiColumnTableRow<FTreeItemPtr>::Construct(Args, OutlinerTreeView);
+	}
+
+	const float SSceneOutlinerTreeRow::HighlightRectLeftOffset = 0.0f;
+	const float SSceneOutlinerTreeRow::HighlightRectRightOffset = 0.0f;
+	const float SSceneOutlinerTreeRow::HighlightTargetSpringConstant = 25.0f;
+	const float SSceneOutlinerTreeRow::HighlightTargetEffectDuration = 0.5f;
+	const float SSceneOutlinerTreeRow::HighlightTargetOpacity = 0.8f;
+	const float SSceneOutlinerTreeRow::LabelChangedAnimOffsetPercent = 0.2f;
+    
+    void SSceneOutlinerTreeRow::FlashHighlight()
+    {
+        LastHighlightInteractionTime = FSlateApplication::Get().GetCurrentTime();
+    }
+
+	void SSceneOutlinerTreeRow::Tick( const FGeometry& AllottedGeometry, const double InCurrentTime, const float InDeltaTime )
+	{
+		// Call parent implementation.
+		SCompoundWidget::Tick( AllottedGeometry, InCurrentTime, InDeltaTime );
+
+		// We'll draw with the 'focused' look if we're either focused or we have a context menu summoned
+		const bool bShouldAppearFocused = HasKeyboardFocus();
+
+		// Update highlight 'target' effect
+		{
+			const float HighlightLeftX = HighlightRectLeftOffset;
+			const float HighlightRightX = HighlightRectRightOffset + AllottedGeometry.Size.X;
+
+			HighlightTargetLeftSpring.SetTarget( HighlightLeftX );
+			HighlightTargetRightSpring.SetTarget( HighlightRightX );
+
+			float TimeSinceHighlightInteraction = (float)( InCurrentTime - LastHighlightInteractionTime );
+			if( TimeSinceHighlightInteraction <= HighlightTargetEffectDuration || bShouldAppearFocused )
+			{
+				HighlightTargetLeftSpring.Tick( InDeltaTime );
+				HighlightTargetRightSpring.Tick( InDeltaTime );
+			}
+		}
+	}
+
+	int32 SSceneOutlinerTreeRow::OnPaint( const FPaintArgs& Args, const FGeometry& AllottedGeometry, const FSlateRect& MyClippingRect, FSlateWindowElementList& OutDrawElements, int32 LayerId, const FWidgetStyle& InWidgetStyle, bool bParentEnabled ) const
 	{
-		SceneOutlinerWeak = InArgs._SceneOutliner;
-		Item = InArgs._Item;
-
-		SMultiColumnTableRow< FOutlinerTreeItemPtr >::Construct(
-			FSuperRowType::FArguments()
-			.OnDragDetected(this, &SSceneOutlinerTreeRow::HandleDragDetected)
-			.Style(&FEditorStyle::Get().GetWidgetStyle<FTableRowStyle>("SceneOutliner.TableViewRow"))
-		,InOwnerTableView );
+		int32 StartLayer = SMultiColumnTableRow::OnPaint( Args, AllottedGeometry, MyClippingRect, OutDrawElements, LayerId, InWidgetStyle, bParentEnabled );
+
+		const int32 TextLayer = 1;	
+
+		// See if a disabled effect should be used
+		bool bEnabled = ShouldBeEnabled( bParentEnabled );
+		ESlateDrawEffect::Type DrawEffects = (bEnabled) ? ESlateDrawEffect::None : ESlateDrawEffect::DisabledEffect;
+
+		const double CurrentTime = FSlateApplication::Get().GetCurrentTime();
+
+		// We'll draw with the 'focused' look if we're either focused or we have a context menu summoned
+		const bool bShouldAppearFocused = HasKeyboardFocus();
+
+		const float DrawPositionY = ( AllottedGeometry.Size.Y / 2 ) - ( AllottedGeometry.Size.Y / 2 );
+
+		// Draw highlight targeting effect
+		const float TimeSinceHighlightInteraction = (float)( CurrentTime - LastHighlightInteractionTime );
+		if( TimeSinceHighlightInteraction <= HighlightTargetEffectDuration )
+		{
+
+			// Compute animation progress
+			float EffectAlpha = FMath::Clamp( TimeSinceHighlightInteraction / HighlightTargetEffectDuration, 0.0f, 1.0f );
+			EffectAlpha = 1.0f - EffectAlpha * EffectAlpha;  // Inverse square falloff (looks nicer!)
+
+			// Apply extra opacity falloff when dehighlighting
+			float EffectOpacity = EffectAlpha;
+
+			// Figure out a universally visible highlight color.
+			FColor HighlightTargetColorAndOpacity = ( (FLinearColor::White - ColorAndOpacity.Get())*0.5f + FLinearColor(+0.4f, +0.1f, -0.2f)) * InWidgetStyle.GetColorAndOpacityTint();
+			HighlightTargetColorAndOpacity.A = HighlightTargetOpacity * EffectOpacity * 255.0f;
+
+			// Compute the bounds offset of the highlight target from where the highlight target spring
+			// extents currently lie.  This is used to "grow" or "shrink" the highlight as needed.
+			const float LabelChangedAnimOffset = LabelChangedAnimOffsetPercent * AllottedGeometry.Size.Y;
+
+			// Choose an offset amount depending on whether we're highlighting, or clearing highlight
+			const float EffectOffset = EffectAlpha * LabelChangedAnimOffset;
+
+			const float HighlightLeftX = HighlightTargetLeftSpring.GetPosition() - EffectOffset;
+			const float HighlightRightX = HighlightTargetRightSpring.GetPosition() + EffectOffset;
+			const float HighlightTopY = 0.0f - LabelChangedAnimOffset;
+			const float HighlightBottomY = AllottedGeometry.Size.Y + EffectOffset;
+
+			const FVector2D DrawPosition = FVector2D( HighlightLeftX, HighlightTopY );
+			const FVector2D DrawSize = FVector2D( HighlightRightX - HighlightLeftX, HighlightBottomY - HighlightTopY );
+
+			const FSlateBrush* StyleInfo = FEditorStyle::GetBrush("SceneOutliner.ChangedItemHighlight");
+
+			// NOTE: We rely on scissor clipping for the highlight rectangle
+			FSlateDrawElement::MakeBox(
+				OutDrawElements,
+				LayerId + TextLayer,
+				AllottedGeometry.ToPaintGeometry( DrawPosition, DrawSize ),	// Position, Size, Scale
+				StyleInfo,													// Style
+				MyClippingRect,												// Clipping rect
+				DrawEffects,												// Effects to use
+				HighlightTargetColorAndOpacity );							// Color
+		}
+
+		return LayerId + TextLayer;
 	}
 }
 
