Added Worlds and Level Blueprints to the scene outliner

diff --git a/Engine/Source/Editor/SceneOutliner/Private/SOutlinerTreeView.h b/Engine/Source/Editor/SceneOutliner/Private/SOutlinerTreeView.h
--- a/Engine/Source/Editor/SceneOutliner/Private/SOutlinerTreeView.h
+++ b/Engine/Source/Editor/SceneOutliner/Private/SOutlinerTreeView.h
@@ -1,30 +1,19 @@
 // Copyright 1998-2014 Epic Games, Inc. All Rights Reserved.
 #pragma once
 
-#include "Editor/UnrealEd/Public/DragAndDrop/ActorDragDropGraphEdOp.h"
-#include "SSceneOutlinerItemWidget.h"
+#include "SceneOutlinerFwd.h"
 
 namespace SceneOutliner
 {
-	class SSceneOutliner;
 
-	typedef TSharedPtr<TOutlinerTreeItem> FOutlinerTreeItemPtr;
-	typedef TSharedRef<TOutlinerTreeItem> FOutlinerTreeItemRef;
-
-	/** IDs for list columns */
-	static const FName ColumnID_ActorLabel( "Actor" );
-
-	class SOutlinerTreeView : public STreeView< FOutlinerTreeItemPtr >
+	class SOutlinerTreeView : public STreeView< FTreeItemPtr >
 	{
 	public:
 		
 		/** Construct this widget */
 		void Construct(const FArguments& InArgs, TSharedRef<SSceneOutliner> Owner);
 
-		void FlashHighlightOnItem( FOutlinerTreeItemPtr FlashHighlightOnItem );
-
-		/** Check that it is valid to move the current selection to the specified folder path */
-		bool ValidateMoveSelectionTo(FName NewParent, FText* OutValidationText = nullptr);
+		void FlashHighlightOnItem( FTreeItemPtr FlashHighlightOnItem );
 
 	protected:
 
@@ -32,6 +21,8 @@ namespace SceneOutliner
 
 		virtual void OnDragLeave( const FDragDropEvent& DragDropEvent ) override;
 
+		//virtual FReply OnDragDetected( const FGeometry& MyGeometry, const FPointerEvent& MouseEvent ) override;
+
 		virtual FReply OnDrop( const FGeometry& MyGeometry, const FDragDropEvent& DragDropEvent ) override;
 
 		/** Weak reference to the outliner widget that owns this list */
@@ -40,43 +31,35 @@ namespace SceneOutliner
 
 	/** Widget that represents a row in the outliner's tree control.  Generates widgets for each column on demand. */
 	class SSceneOutlinerTreeRow
-		: public SMultiColumnTableRow< FOutlinerTreeItemPtr >
+		: public SMultiColumnTableRow< FTreeItemPtr >
 	{
 
 	public:
 
 		SLATE_BEGIN_ARGS( SSceneOutlinerTreeRow ) {}
 
-			/** The Scene Outliner widget that owns the tree.  We'll only keep a weak reference to it. */
-			SLATE_ARGUMENT( TSharedPtr< SSceneOutliner >, SceneOutliner )
-
 			/** The list item for this row */
-			SLATE_ARGUMENT( FOutlinerTreeItemPtr, Item )
+			SLATE_ARGUMENT( FTreeItemPtr, Item )
 
 		SLATE_END_ARGS()
 
 
 		/** Construct function for this widget */
-		void Construct( const FArguments& InArgs, const TSharedRef<STableViewBase>& InOwnerTableView );
-
+		void Construct( const FArguments& InArgs, const TSharedRef<SOutlinerTreeView>& OutlinerTreeView, TSharedRef<SSceneOutliner> SceneOutliner );
 
 		/** Overridden from SMultiColumnTableRow.  Generates a widget for this column of the tree row. */
 		virtual TSharedRef<SWidget> GenerateWidgetForColumn( const FName& ColumnName ) override;
 
-		TSharedPtr< SSceneOutlinerItemWidget > GetItemWidget();
+		void FlashHighlight();
 
 	protected:
 
 		virtual FReply OnMouseButtonDown(const FGeometry& MyGeometry, const FPointerEvent& MouseEvent) override;
 
-		FReply HandleDragDetected( const FGeometry& MyGeometry, const FPointerEvent& MouseEvent );
-		
-		void PerformDetachment(TWeakObjectPtr<AActor> ParentActorPtr );
-
-		void PerformAttachment(FName SocketName, TWeakObjectPtr<AActor> ParentActorPtr, const TArray< TWeakObjectPtr<AActor> > ChildrenPtrs);
-
 		virtual void OnDragEnter( const FGeometry& MyGeometry, const FDragDropEvent& DragDropEvent ) override;
 
+		virtual void OnDragLeave( const FDragDropEvent& DragDropEvent ) override;
+
 		virtual FReply OnDragOver( const FGeometry& MyGeometry, const FDragDropEvent& DragDropEvent ) override
 		{
 			return FReply::Handled();
@@ -84,16 +67,48 @@ namespace SceneOutliner
 
 		virtual FReply OnDrop( const FGeometry& MyGeometry, const FDragDropEvent& DragDropEvent ) override;
 
+		bool IsRowEnabled() const;
 
 	private:
 
 		/** Weak reference to the outliner widget that owns our list */
 		TWeakPtr< SSceneOutliner > SceneOutlinerWeak;
 
-		/** The widget contained in this row representing the associated item */
-		TSharedPtr< SSceneOutlinerItemWidget > ItemWidget;
-
 		/** The item associated with this row of data */
-		FOutlinerTreeItemPtr Item;
+		TWeakPtr<ITreeItem> Item;
+
+	protected:
+
+		virtual void Tick( const FGeometry& AllottedGeometry, const double InCurrentTime, const float InDeltaTime ) override;
+
+		virtual int32 OnPaint( const FPaintArgs& Args, const FGeometry& AllottedGeometry, const FSlateRect& MyClippingRect, FSlateWindowElementList& OutDrawElements, int32 LayerId, const FWidgetStyle& InWidgetStyle, bool bParentEnabled ) const override;
+
+	private:
+		/** How many pixels to extend the highlight rectangle's left side horizontally */
+		static const float HighlightRectLeftOffset;
+
+		/** How many pixels to extend the highlight rectangle's right side horizontally */
+		static const float HighlightRectRightOffset;
+
+		/** How quickly the highlight 'targeting' rectangle will slide around.  Larger is faster. */
+		static const float HighlightTargetSpringConstant;
+
+		/** Duration of animation highlight target effects */
+		static const float HighlightTargetEffectDuration;
+
+		/** Opacity of the highlight target effect overlay */
+		static const float HighlightTargetOpacity;
+
+		/** How large the highlight target effect will be when highlighting, as a scalar percentage of font height */
+		static const float LabelChangedAnimOffsetPercent;
+
+		/** Highlight "targeting" visual effect left position */
+		FFloatSpring1D HighlightTargetLeftSpring;
+
+		/** Highlight "targeting" visual effect right position */
+		FFloatSpring1D HighlightTargetRightSpring;
+
+		/** Last time that the user had a major interaction with the highlight */
+		double LastHighlightInteractionTime;
 	};
 }		// namespace SceneOutliner
