Added Worlds and Level Blueprints to the scene outliner

diff --git a/Engine/Source/Editor/SceneOutliner/Private/SSceneOutliner.cpp b/Engine/Source/Editor/SceneOutliner/Private/SSceneOutliner.cpp
--- a/Engine/Source/Editor/SceneOutliner/Private/SSceneOutliner.cpp
+++ b/Engine/Source/Editor/SceneOutliner/Private/SSceneOutliner.cpp
@@ -3,7 +3,6 @@
 
 #include "SceneOutlinerPrivatePCH.h"
 #include "SSceneOutliner.h"
-#include "SceneOutlinerTreeItems.h"
 
 #include "ScopedTransaction.h"
 
@@ -11,10 +10,13 @@
 
 #include "SSocketChooser.h"
 
-#include "LevelUtils.h"
-#include "ActorEditorUtils.h"
-#include "SceneOutlinerFilterHelpers.h"
 #include "SceneOutlinerFilters.h"
+#include "SceneOutlinerDragDrop.h"
+
+#include "SortHelper.h"
+
+#include "ActorEditorUtils.h"
+#include "LevelUtils.h"
 
 #include "EditorActorFolders.h"
 #include "SSearchBox.h"
@@ -26,112 +28,101 @@ DEFINE_LOG_CATEGORY_STATIC(LogSceneOutliner, Log, All);
 
 #define LOCTEXT_NAMESPACE "SSceneOutliner"
 
-// The amount of time that must past before the Scene Outliner will attempt a sort when in PIE/SIE.
+// The amount of time that must pass before the Scene Outliner will attempt a sort when in PIE/SIE.
 #define SCENE_OUTLINER_RESORT_TIMER 1.0f
 
 namespace SceneOutliner
 {
+	TSharedPtr< FOutlinerFilter > CreateSelectedActorFilter()
+	{
+		auto Predicate = FActorFilterPredicate::CreateStatic([](const AActor* InActor){			
+			return InActor->IsSelected();
+		});
+		return MakeShareable( new FOutlinerPredicateFilter(Predicate, EDefaultFilterBehaviour::Fail) );
+	}
 
-	/** Recursively set the selection of all actor descendants of the specified folder */
-	void SelectDescedantActors(TSharedRef<TOutlinerFolderTreeItem> FolderItem, bool bSelected)
+	TSharedPtr< FOutlinerFilter > CreateHideTemporaryActorsFilter()
 	{
-		for (const auto& WeakChild : FolderItem->ChildItems)
-		{
-			auto ChildItem = WeakChild.Pin();
-			if (ChildItem.IsValid())
-			{
-				if (ChildItem->Type == TOutlinerTreeItem::Actor)
-				{
-					if (AActor* Actor = StaticCastSharedPtr<TOutlinerActorTreeItem>(ChildItem)->Actor.Get())
-					{
-						GEditor->SelectActor(Actor, true, /*bNotify=*/false);
-					}
-				}
-				else
-				{
-					SelectDescedantActors(StaticCastSharedPtr<TOutlinerFolderTreeItem>(ChildItem).ToSharedRef(), bSelected);
-				}
-			}
-		}
+		return MakeShareable( new FOutlinerPredicateFilter( FActorFilterPredicate::CreateStatic( []( const AActor* InActor ){			
+			return InActor->GetWorld()->WorldType != EWorldType::PIE || GEditor->ObjectsThatExistInEditorWorld.Get(InActor);
+		} ), EDefaultFilterBehaviour::Pass ) );
 	}
 
-	/** Gets the label for a tree item */
-	FText GetLabelForItem( const TSharedRef<TOutlinerTreeItem> TreeItem )
+	TSharedPtr< FOutlinerFilter > CreateIsInCurrentLevelFilter()
 	{
-		if (TreeItem->Type == SceneOutliner::TOutlinerTreeItem::Actor)
-		{
-			const AActor* Actor = StaticCastSharedRef<const SceneOutliner::TOutlinerActorTreeItem>(TreeItem)->Actor.Get();
-			return Actor ? FText::FromString( Actor->GetActorLabel() ) : LOCTEXT("ActorLabelForMissingActor", "(Deleted Actor)");
-		}
-		else
-		{
-			return FText::FromName(StaticCastSharedRef<const SceneOutliner::TOutlinerFolderTreeItem>(TreeItem)->LeafName);
-		}
+		auto Predicate = FActorFilterPredicate::CreateStatic([](const AActor* InActor){
+			return InActor->GetLevel() == InActor->GetWorld()->GetCurrentLevel();
+		});
+		return MakeShareable( new FOutlinerPredicateFilter(Predicate, EDefaultFilterBehaviour::Pass) );
 	}
 
-	/** Get label as a string, much faster than as a FText */
-	FString GetLabelForItemAsString(const TSharedRef<TOutlinerTreeItem> TreeItem)
+	struct FItemSelection : IMutableTreeItemVisitor
 	{
-		if (TreeItem->Type == SceneOutliner::TOutlinerTreeItem::Actor)
-		{
-			const AActor* Actor = StaticCastSharedRef<const SceneOutliner::TOutlinerActorTreeItem>(TreeItem)->Actor.Get();
-			return Actor ? Actor->GetActorLabel() : LOCTEXT("ActorLabelForMissingActor", "(Deleted Actor)").ToString();
-		}
-		else
+		mutable TArray<FActorTreeItem*> Actors;
+		mutable TArray<FWorldTreeItem*> Worlds;
+		mutable TArray<FLevelBlueprintTreeItem*> LevelBlueprints;
+		mutable TArray<FFolderTreeItem*> Folders;
+
+		FItemSelection()
+		{}
+
+		FItemSelection(SOutlinerTreeView& Tree)
 		{
-			return StaticCastSharedRef<const SceneOutliner::TOutlinerFolderTreeItem>(TreeItem)->LeafName.ToString();
+			for (auto& Item : Tree.GetSelectedItems())
+			{
+				Item->Visit(*this);
+			}
 		}
-	}
 
-	namespace Helpers
-	{
-		/** Sorts tree items alphabetically */
-		struct FSortByActorAscending
+		TArray<TWeakObjectPtr<AActor>> GetWeakActors() const
 		{
-			bool operator()( const FOutlinerTreeItemPtr& A, const FOutlinerTreeItemPtr& B ) const
+			TArray<TWeakObjectPtr<AActor>> ActorArray;
+			for (const auto* ActorItem : Actors)
 			{
-				if ((A->Type == TOutlinerTreeItem::Folder) != (B->Type == TOutlinerTreeItem::Folder))
+				if (ActorItem->Actor.IsValid())
 				{
-					return A->Type == TOutlinerTreeItem::Folder;
+					ActorArray.Add(ActorItem->Actor);
 				}
-				return GetLabelForItemAsString(A.ToSharedRef()) < GetLabelForItemAsString(B.ToSharedRef());
 			}
-		};
-		struct FSortByActorDescending
+			return ActorArray;
+		}
+
+		TArray<AActor*> GetActorPtrs() const
 		{
-			bool operator()( const FOutlinerTreeItemPtr& A, const FOutlinerTreeItemPtr& B ) const
+			TArray<AActor*> ActorArray;
+			for (const auto* ActorItem : Actors)
 			{
-				if ((A->Type == TOutlinerTreeItem::Folder) != (B->Type == TOutlinerTreeItem::Folder))
+				if (AActor* Actor = ActorItem->Actor.Get())
 				{
-					return A->Type != TOutlinerTreeItem::Folder;
+					ActorArray.Add(Actor);
 				}
-				return GetLabelForItemAsString(A.ToSharedRef()) > GetLabelForItemAsString(B.ToSharedRef());
 			}
-		};
-	}
+			return ActorArray;
+		}
 
-	void PerformActionOnItem(TSharedPtr<SOutlinerTreeView> OutlinerTreeView, TSharedRef<TOutlinerTreeItem> NewItem, uint8 ActionMask)
-	{
-		if (ActionMask & NewItemActuator::Select)
+	private:
+		virtual void Visit(FActorTreeItem& ActorItem) const override
 		{
-			OutlinerTreeView->ClearSelection();
-			OutlinerTreeView->SetItemSelection(NewItem, true);
+			Actors.Add(&ActorItem);
 		}
-
-		if (ActionMask & NewItemActuator::Rename)
+		virtual void Visit(FWorldTreeItem& WorldItem) const override
 		{
-			NewItem->Flags.RenameWhenInView = true;
+			Worlds.Add(&WorldItem);
 		}
-
-		if (ActionMask & (NewItemActuator::ScrollIntoView | NewItemActuator::Rename))
+		virtual void Visit(FLevelBlueprintTreeItem& LevelBlueprintItem) const override
 		{
-			OutlinerTreeView->RequestScrollIntoView(NewItem);
+			LevelBlueprints.Add(&LevelBlueprintItem);
 		}
-	}
+		virtual void Visit(FFolderTreeItem& FolderItem) const override
+		{
+			Folders.Add(&FolderItem);
+		}
+	};
 
-	void SSceneOutliner::Construct( const FArguments& InArgs, const FSceneOutlinerInitializationOptions& InInitOptions )
+	void SSceneOutliner::Construct( const FArguments& InArgs, const FInitializationOptions& InInitOptions )
 	{
-		InitOptions = InInitOptions;
+		// Copy over the shared data from the initialization options
+		static_cast<FSharedDataBase&>(*SharedData) = static_cast<const FSharedDataBase&>(InInitOptions);
 
 		OnItemPicked = InArgs._OnItemPickedDelegate;
 
@@ -143,11 +134,12 @@ namespace SceneOutliner
 		bFullRefresh = true;
 		bNeedsRefresh = true;
 		bIsReentrant = false;
+		bSortDirty = true;
 		FilteredActorCount = 0;
 		SortOutlinerTimer = 0.0f;
 		bPendingFocusNextFrame = InInitOptions.bFocusSearchBoxWhenOpened;
 
-		SortByColumn = ColumnID_ActorLabel;
+		SortByColumn = FBuiltInColumnTypes::Label();
 		SortMode = EColumnSortMode::Ascending;
 
 		// @todo outliner: Should probably save this in layout!
@@ -156,9 +148,6 @@ namespace SceneOutliner
 		NoBorder = FEditorStyle::GetBrush( "LevelViewport.NoViewportBorder" );
 		PlayInEditorBorder = FEditorStyle::GetBrush( "LevelViewport.StartingPlayInEditorBorder" );
 		SimulateBorder = FEditorStyle::GetBrush( "LevelViewport.StartingSimulateBorder" );
-		MobilityStaticBrush = FEditorStyle::GetBrush( "ClassIcon.ComponentMobilityStaticPip" );
-		MobilityStationaryBrush = FEditorStyle::GetBrush( "ClassIcon.ComponentMobilityStationaryPip" );
-		MobilityMovableBrush = FEditorStyle::GetBrush( "ClassIcon.ComponentMobilityMovablePip" );
 
 		//Setup the SearchBox filter
 		{
@@ -168,79 +157,26 @@ namespace SceneOutliner
 
 		TSharedRef<SVerticalBox> VerticalBox = SNew(SVerticalBox);
 
-		//Setup the FilterCollection
-		//We use the filter collection provide, otherwise we create our own
+		// We use the filter collection provided, otherwise we create our own
 		Filters = InInitOptions.Filters.IsValid() ? InInitOptions.Filters : MakeShareable(new FOutlinerFilters);
 	
 		SearchBoxFilter->OnChanged().AddSP( this, &SSceneOutliner::FullRefresh );
 		Filters->OnChanged().AddSP( this, &SSceneOutliner::FullRefresh );
 
 		//Apply custom filters based on global preferences
-		if (InInitOptions.Mode == ESceneOutlinerMode::ActorBrowsing)
+		if (SharedData->Mode == ESceneOutlinerMode::ActorBrowsing)
 		{
 			ApplyShowOnlySelectedFilter( IsShowingOnlySelected() );
 			ApplyHideTemporaryActorsFilter( IsHidingTemporaryActors() );
 			ApplyShowOnlyCurrentLevelFilter( IsShowingOnlyCurrentLevel() );
 		}
 
-
-		if (InInitOptions.CustomColumnFactory.IsBound())
-		{
-			CustomColumn = InInitOptions.CustomColumnFactory.Execute(SharedThis(this));
-		}
-		else
-		{
-			CustomColumn = MakeShareable( new FSceneOutlinerActorInfoColumn( SharedThis( this ), ECustomColumnMode::ActorClass ) );
-		}
-
 		TSharedRef< SHeaderRow > HeaderRowWidget =
 			SNew( SHeaderRow )
 				// Only show the list header if the user configured the outliner for that
 				.Visibility(InInitOptions.bShowHeaderRow ? EVisibility::Visible : EVisibility::Collapsed);
 
-		if (InInitOptions.Mode == ESceneOutlinerMode::ActorBrowsing)
-		{
-			// Set up the gutter if we're in actor browsing mode
-
-			Gutter = MakeShareable(new FSceneOutlinerGutter(FOnSetItemVisibility::CreateSP(this, &SSceneOutliner::OnSetItemVisibility)));
-			HeaderRowWidget->AddColumn(
-				Gutter->ConstructHeaderRowColumn()
-				.FixedWidth(16.f)
-				.SortMode(this, &SSceneOutliner::GetColumnSortMode, Gutter->GetColumnID())
-				.OnSort(this, &SSceneOutliner::OnColumnSortModeChanged)
-			);
-		}
-		
-		// Actor label column
-		HeaderRowWidget->AddColumn(	
-			SHeaderRow::Column( ColumnID_ActorLabel )
-			.DefaultLabel( LOCTEXT("ActorColumnLabel", "Actor").ToString() )
-			.SortMode( this, &SSceneOutliner::GetColumnSortMode, ColumnID_ActorLabel )
-			.OnSort( this, &SSceneOutliner::OnColumnSortModeChanged )
-			.FillWidth( 5.0f )
-		);
-
-		if (InInitOptions.Mode == ESceneOutlinerMode::ActorBrowsing || InInitOptions.CustomColumnFactory.IsBound())
-		{
-			// Customizable actor data column is only viable when browsing OR when it has been bound specifically
-
-			auto CustomColumnHeaderRow = CustomColumn->ConstructHeaderRowColumn();
-
-			CustomColumnHeaderRow
-				.SortMode(this, &SSceneOutliner::GetColumnSortMode, CustomColumn->GetColumnID())
-				.OnSort(this, &SSceneOutliner::OnColumnSortModeChanged);
-
-			if (InInitOptions.CustomColumnFixedWidth > 0.0f)
-			{
-				CustomColumnHeaderRow.FixedWidth(InInitOptions.CustomColumnFixedWidth);
-			}
-			else
-			{
-				CustomColumnHeaderRow.FillWidth(2.0f);
-			}
-
-			HeaderRowWidget->AddColumn(CustomColumnHeaderRow);
-		}
+		SetupColumns(*HeaderRowWidget);
 
 		ChildSlot
 		[
@@ -259,14 +195,14 @@ namespace SceneOutliner
 		.VAlign(VAlign_Center)
 		[
 			SAssignNew( FilterTextBoxWidget, SSearchBox )
-			.Visibility( this, &SSceneOutliner::GetSearchBoxVisibility )
+			.Visibility( InInitOptions.bShowSearchBox ? EVisibility::Visible : EVisibility::Collapsed )
 			.HintText( LOCTEXT( "FilterSearch", "Search..." ) )
 			.ToolTipText( LOCTEXT("FilterSearchHint", "Type here to search (pressing enter selects the results)").ToString() )
 			.OnTextChanged( this, &SSceneOutliner::OnFilterTextChanged )
 			.OnTextCommitted( this, &SSceneOutliner::OnFilterTextCommitted )
 		];
 
-		if (InInitOptions.Mode == ESceneOutlinerMode::ActorBrowsing)
+		if (SharedData->Mode == ESceneOutlinerMode::ActorBrowsing)
 		{
 			Toolbar->AddSlot()
 				.VAlign(VAlign_Center)
@@ -341,7 +277,7 @@ namespace SceneOutliner
 			]
 		];
 
-		if (InInitOptions.Mode == ESceneOutlinerMode::ActorBrowsing)
+		if (SharedData->Mode == ESceneOutlinerMode::ActorBrowsing)
 		{
 			// Separator
 			VerticalBox->AddSlot()
@@ -398,7 +334,7 @@ namespace SceneOutliner
 					]
 				]
 			];
-		} //end if (InitOptions.Mode == ESceneOutlinerMode::ActorBrowsing)
+		} //end if (SharedData->Mode == ESceneOutlinerMode::ActorBrowsing)
 
 
 		// Don't allow tool-tips over the header
@@ -408,9 +344,8 @@ namespace SceneOutliner
 		// Populate our data set
 		Populate();
 
-
 		// We only synchronize selection when in actor browsing mode
-		if (InInitOptions.Mode == ESceneOutlinerMode::ActorBrowsing)
+		if (SharedData->Mode == ESceneOutlinerMode::ActorBrowsing)
 		{
 			// Populate and register to find out when the level's selection changes
 			OnLevelSelectionChanged( NULL );
@@ -443,13 +378,74 @@ namespace SceneOutliner
 
 		auto& Folders = FActorFolders::Get();
 		Folders.OnFolderCreate.AddSP(this, &SSceneOutliner::OnBroadcastFolderCreate);
+		Folders.OnFolderMove.AddSP(this, &SSceneOutliner::OnBroadcastFolderMove);
 		Folders.OnFolderDelete.AddSP(this, &SSceneOutliner::OnBroadcastFolderDelete);
 	}
 
+	void SSceneOutliner::SetupColumns(SHeaderRow& HeaderRow)
+	{
+		FSceneOutlinerModule& SceneOutlinerModule = FModuleManager::LoadModuleChecked<FSceneOutlinerModule>("SceneOutliner");
+
+		if (SharedData->ColumnMap.Num() == 0)
+		{
+			SharedData->UseDefaultColumns();
+		}
+
+		Columns.Empty(SharedData->ColumnMap.Num());
+		HeaderRow.ClearColumns();
+
+		// Get a list of sorted columns IDs to create
+		TArray<FName> SortedIDs;
+		SortedIDs.Reserve(SharedData->ColumnMap.Num());
+		SharedData->ColumnMap.GenerateKeyArray(SortedIDs);
+
+		SortedIDs.Sort([&](const FName& A, const FName& B){
+			return SharedData->ColumnMap[A].PriorityIndex < SharedData->ColumnMap[B].PriorityIndex;
+		});
+
+		for (const FName& ID : SortedIDs)
+		{
+			if (SharedData->ColumnMap[ID].Visibility == EColumnVisibility::Invisible)
+			{
+				continue;
+			}
+
+			TSharedPtr<ISceneOutlinerColumn> Column;
+
+			if (SharedData->ColumnMap[ID].Factory.IsBound())
+			{
+				Column = SharedData->ColumnMap[ID].Factory.Execute(*this);
+			}
+			else
+			{
+				Column = SceneOutlinerModule.FactoryColumn(ID, *this);
+			}
+
+			if (ensure(Column.IsValid()))
+			{
+				check(Column->GetColumnID() == ID);
+				Columns.Add(Column->GetColumnID(), Column);
+
+				auto ColumnArgs = Column->ConstructHeaderRowColumn();
+
+				if (Column->SupportsSorting())
+				{
+					ColumnArgs
+						.SortMode(this, &SSceneOutliner::GetColumnSortMode, Column->GetColumnID())
+						.OnSort(this, &SSceneOutliner::OnColumnSortModeChanged);
+				}
+
+				HeaderRow.AddColumn(ColumnArgs);
+			}
+		}
+
+		Columns.Shrink();
+	}
+
 	SSceneOutliner::~SSceneOutliner()
 	{
 		// We only synchronize selection when in actor browsing mode
-		if( InitOptions.Mode == ESceneOutlinerMode::ActorBrowsing )
+		if( SharedData->Mode == ESceneOutlinerMode::ActorBrowsing )
 		{
 			USelection::SelectionChangedEvent.RemoveAll(this);
 			USelection::SelectObjectEvent.RemoveAll(this);
@@ -471,23 +467,14 @@ namespace SceneOutliner
 		{
 			auto& Folders = FActorFolders::Get();
 			Folders.OnFolderCreate.RemoveAll(this);
+			Folders.OnFolderMove.RemoveAll(this);
 			Folders.OnFolderDelete.RemoveAll(this);
 		}
 	}
 
-	TArray<FOutlinerTreeItemPtr> SSceneOutliner::GetSelectedItems()
-	{
-		return OutlinerTreeView->GetSelectedItems();
-	}
-
-	const TSharedPtr<TOutlinerFolderTreeItem> SSceneOutliner::FindFolderByPath(FName Path) const
+	void SSceneOutliner::OnItemAdded(const FTreeItemID& ItemID, uint8 ActionMask)
 	{
-		auto* FolderPtr = FolderToTreeItemMap.Find(Path);
-		if (FolderPtr)
-		{
-			return *FolderPtr;
-		}
-		return nullptr;
+		NewItemActions.Add(ItemID, ActionMask);
 	}
 
 	FSlateColor SSceneOutliner::GetViewButtonForegroundColor() const
@@ -563,7 +550,7 @@ namespace SceneOutliner
 	{
 		if ( !SelectedActorFilter.IsValid() )
 		{
-			SelectedActorFilter = SceneOutlinerFilters::CreateSelectedActorFilter();
+			SelectedActorFilter = CreateSelectedActorFilter();
 		}
 
 		if ( bShowOnlySelected )
@@ -596,7 +583,7 @@ namespace SceneOutliner
 	{
 		if ( !HideTemporaryActorsFilter.IsValid() )
 		{
-			HideTemporaryActorsFilter = SceneOutlinerFilters::CreateHideTemporaryActorsFilter();
+			HideTemporaryActorsFilter = CreateHideTemporaryActorsFilter();
 		}
 
 		if ( bHideTemporaryActors )
@@ -628,7 +615,7 @@ namespace SceneOutliner
 	{
 		if ( !ShowOnlyActorsInCurrentLevelFilter.IsValid() )
 		{
-			ShowOnlyActorsInCurrentLevelFilter = SceneOutlinerFilters::CreateIsInCurrentLevelFilter();
+			ShowOnlyActorsInCurrentLevelFilter = CreateIsInCurrentLevelFilter();
 		}
 
 		if ( bShowOnlyActorsInCurrentLevel )
@@ -650,10 +637,9 @@ namespace SceneOutliner
 
 	const FSlateBrush* SSceneOutliner::OnGetBorderBrush() const
 	{
-		if (bRepresentingPlayWorld)
+		if (SharedData->bRepresentingPlayWorld)
 		{
-			return (GEditor->bIsSimulatingInEditor)? SimulateBorder 
-												   : PlayInEditorBorder;
+			return GEditor->bIsSimulatingInEditor ? SimulateBorder : PlayInEditorBorder;
 		}
 		else
 		{
@@ -663,13 +649,13 @@ namespace SceneOutliner
 
 	FSlateColor SSceneOutliner::OnGetBorderColorAndOpacity() const
 	{
-		return (bRepresentingPlayWorld)? FLinearColor(1.0f, 1.0f, 1.0f, 0.6f)
-									   : FLinearColor::Transparent;
+		return SharedData->bRepresentingPlayWorld	? FLinearColor(1.0f, 1.0f, 1.0f, 0.6f)
+									   				: FLinearColor::Transparent;
 	}
 
 	ESelectionMode::Type SSceneOutliner::GetSelectionMode() const
 	{
-		return (InitOptions.Mode == ESceneOutlinerMode::ActorBrowsing)? 
+		return (SharedData->Mode == ESceneOutlinerMode::ActorBrowsing)? 
 			ESelectionMode::Multi : ESelectionMode::Single;
 	}
 
@@ -691,974 +677,361 @@ namespace SceneOutliner
 		// Block events while we clear out the list.  We don't want actors in the level to become deselected
 		// while we doing this
 		TGuardValue<bool> ReentrantGuard(bIsReentrant, true);
-	
 
 		for (const FWorldContext& Context : GEngine->GetWorldContexts())
 		{
 			if (Context.WorldType == EWorldType::PIE)
 			{
-				RepresentingWorld = Context.World();
+				SharedData->RepresentingWorld = Context.World();
 				break;
 			}
 			else if (Context.WorldType == EWorldType::Editor)
 			{
-				RepresentingWorld = Context.World();
+				SharedData->RepresentingWorld = Context.World();
 			}
 		}
 		
-		if (!RepresentingWorld)
+		if (!CheckWorld())
 		{
 			return;
 		}
 
-		bRepresentingPlayWorld = RepresentingWorld->WorldType == EWorldType::PIE;
+		SharedData->bRepresentingPlayWorld = SharedData->RepresentingWorld->WorldType == EWorldType::PIE;
 
 		// Get a collection of items and folders which were formerly collapsed
-		// TODO: this will actually include non-parent items too; ideally only add item if it has children.
-		FParentsExpansionState ExpansionStateInfo;
-		GetParentsExpansionState(ExpansionStateInfo);
+		const FParentsExpansionState ExpansionStateInfo = GetParentsExpansionState();
 
 		bool bMadeAnySignificantChanges = false;
 		if(bFullRefresh)
 		{
-			// Clear the selection here - UpdateActorMapWithWorld will reconstruct it.
+			// Clear the selection here - RepopulateEntireTree will reconstruct it.
 			OutlinerTreeView->ClearSelection();
 
-			FParentActorsSet ParentActors;
-			UpdateActorMapWithWorld(RepresentingWorld, ParentActors);
-
-			SetParentsExpansionState(ParentActors, ExpansionStateInfo);
+			RepopulateEntireTree();
 
 			bMadeAnySignificantChanges = true;
 			bFullRefresh = false;
 		}
-		else
-		{
-			FParentActorsSet ParentActors;
 
-			// Add actors that have been added to the world
-			for(int32 ItemIdx = 0; ItemIdx < AddedItemsList.Num(); ++ItemIdx)
-			{
-				bMadeAnySignificantChanges |= AddItemToTree(AddedItemsList[ItemIdx], ParentActors);
-			}
-
-			// Attach actors to other actors
-			for(int32 ItemIdx = 0; ItemIdx < AttachedActorsList.Num(); ++ItemIdx)
+		// Only deal with 500 at a time
+		const int32 End = FMath::Min(PendingOperations.Num(), 500);
+		for (int32 Index = 0; Index < End; ++Index)
+		{
+			auto& PendingOp = PendingOperations[Index];
+			switch (PendingOp.Type)
 			{
-				const int32 RootIndex = RootTreeItems.Find(AttachedActorsList[ItemIdx]);
-
-				auto ActorItem = AttachedActorsList[ItemIdx];
-				auto* Actor = ActorItem->Actor.Get();
-				if (!Actor)
-				{
-					continue;
-				}
+				case FPendingTreeOperation::Added:
+					bMadeAnySignificantChanges = AddItemToTree(PendingOp.Item) || bMadeAnySignificantChanges;
+					break;
 
-				if (RootIndex != INDEX_NONE)
-				{
+				case FPendingTreeOperation::Moved:
 					bMadeAnySignificantChanges = true;
+					OnItemMoved(PendingOp.Item);
+					break;
 
-					// Item currently exists at the root of the tree, remove it then enumerate its parents
-					RootTreeItems.RemoveAt(RootIndex);
-
-
-					// Since the item is being removed from the root, whatever parent (or hierarchy of parents) it has needs to be added to the tree (and likely filtered out).
-					AActor* Parent = Actor->GetAttachParentActor();
-					while(Parent)
-					{
-						ParentActors.Add(Parent);
-						Parent = Parent->GetAttachParentActor();
-					}
-				}
-				else
-				{
-					// Item currently exists inside a folder - remove this item from our parent, and set the folder name
-					TSharedRef<TOutlinerFolderTreeItem>* ParentFolderItem = FolderToTreeItemMap.Find(Actor->GetFolderPath());
-
-
-					if (ParentFolderItem)
-					{
-						bMadeAnySignificantChanges = true;
-						(*ParentFolderItem)->ChildItems.Remove(AttachedActorsList[ItemIdx]);
-					}
-				}
-
-				// We only synchronize selection when in actor browsing mode
-				if( InitOptions.Mode == ESceneOutlinerMode::ActorBrowsing )
-				{
-					// Synchronize selection
-					if( GEditor->GetSelectedActors()->IsSelected( Actor ) )
-					{
-						OutlinerTreeView->SetItemSelection( AttachedActorsList[ItemIdx], true );
-					}
-				}
-			}
-
-			// Refresh items
-			for(int32 ItemIdx = 0; ItemIdx < RefreshItemsList.Num(); ++ItemIdx)
-			{
-				bMadeAnySignificantChanges = true;
-
-				if (RefreshItemsList[ItemIdx]->Type == TOutlinerTreeItem::Folder)
-				{
-					// Remove any dead actors from the folder
-					auto FolderItem = StaticCastSharedRef<TOutlinerFolderTreeItem>(RefreshItemsList[ItemIdx]);
-					for (int32 Index = FolderItem->ChildItems.Num() - 1; Index >= 0; --Index)
-					{
-						FOutlinerTreeItemPtr Child = FolderItem->ChildItems[Index].Pin();
-						if (!Child.IsValid())
-						{
-							FolderItem->ChildItems.RemoveAt(Index);
-						}
-						else if (Child->Type == TOutlinerTreeItem::Actor && !StaticCastSharedPtr<TOutlinerActorTreeItem>(Child)->Actor.Get())
-						{
-							FolderItem->ChildItems.RemoveAt(Index);
-						}
-					}
-				}
-
-				// Refresh items by removing and re-adding them
-				RemoveItemFromTree(RefreshItemsList[ItemIdx]);
-				AddItemToTree(RefreshItemsList[ItemIdx], ParentActors);
-			}
-
-			// Remove items from the tree
-			for(int32 ItemIdx = 0; ItemIdx < RemovedItemsList.Num(); ++ItemIdx)
-			{
-				bMadeAnySignificantChanges = true;
-				RemoveItemFromTree(RemovedItemsList[ItemIdx]);
-			}
-
-			if( InitOptions.bShowParentTree )
-			{
-				for( auto ActorIt = ParentActors.CreateConstIterator(); ActorIt; ++ActorIt )
-				{
-					const AActor* Actor = *ActorIt;
+				case FPendingTreeOperation::Removed:
+					bMadeAnySignificantChanges = true;
+					RemoveItemFromTree(PendingOp.Item);
+					break;
 
-					AddFilteredParentActorToTree(Actor);
-				}
+				default:
+					check(false);
+					break;
 			}
-
-			SetParentsExpansionState(ParentActors, ExpansionStateInfo);
 		}
-		
-		AddedItemsList.Empty();
-		RefreshItemsList.Empty();
-		RemovedItemsList.Empty();
-		AttachedActorsList.Empty();
+
+		PendingOperations.RemoveAt(0, End);
+		SetParentsExpansionState(ExpansionStateInfo);
 
 		if (bMadeAnySignificantChanges)
 		{
 			RequestSort();
 		}
 
-		NewItemActions.Empty();
-
-		// We're fully refreshed now.
-		bNeedsRefresh = false;
+		if (PendingOperations.Num() == 0)
+		{
+			// We're fully refreshed now.
+			NewItemActions.Empty();
+			bNeedsRefresh = false;
+		}
 	}
 
 	bool SSceneOutliner::ShouldShowFolders() const
 	{
-		return InitOptions.Mode == ESceneOutlinerMode::ActorBrowsing || InitOptions.bOnlyShowFolders;
+		return SharedData->Mode == ESceneOutlinerMode::ActorBrowsing || SharedData->bOnlyShowFolders;
 	}
 
 	void SSceneOutliner::EmptyTreeItems()
 	{
 		FilteredActorCount = 0;
+		ApplicableActors.Empty();
+
+		PendingOperations.Empty();
+		TreeItemMap.Reset();
 
-		ActorToTreeItemMap.Reset();
-		
-		FolderToTreeItemMap.Reset();
 		RootTreeItems.Empty();
 	}
 
-	void SSceneOutliner::UpdateActorMapWithWorld(UWorld* World, FParentActorsSet& ParentActors)
+	void SSceneOutliner::RepopulateEntireTree()
 	{
-		if ( World == NULL )
-		{
-			return;
-		}
-
 		EmptyTreeItems();
 
-		// We only synchronize selection when in actor browsing mode
-		const bool bSyncSelection = ( InitOptions.Mode == ESceneOutlinerMode::ActorBrowsing );
-		FOutlinerTreeItemPtr LastSelectedItem;
-
-		// Iterate over every actor in memory.  WARNING: This is potentially very expensive!
-		for( FActorIterator ActorIt(World); ActorIt; ++ActorIt )
+		if (!SharedData->bOnlyShowFolders)
 		{
-			AActor* Actor = *ActorIt;
-
-			TSharedRef<TOutlinerActorTreeItem> NewItem = MakeShareable( new TOutlinerActorTreeItem(Actor) );
-
-			const bool bAddedToTree = AddActorToTree(NewItem, ParentActors);
-
-			//Tag PlayWorld Actors that have counterparts in EditorWorld
-			NewItem->bExistsInCurrentWorldAndPIE = GEditor->ObjectsThatExistInEditorWorld.Get( Actor );
-
-			// Synchronize selection - only things actually displayed after filtering are wanted
-			if( bAddedToTree && bSyncSelection && GEditor->GetSelectedActors()->IsSelected( Actor ) )
-			{
-				LastSelectedItem = NewItem;
-			}
-		}
+			ConstructItemFor<FWorldTreeItem>(SharedData->RepresentingWorld);
 
-		/* Any parent actors identified might not have been added to the tree due to filtering. It is important
-		   to check these actors and add them to the list so their children can correctly appear during filtering. */
-		if( InitOptions.bShowParentTree )
-		{
-			// Add any parents for matched actors
-			for( auto ActorIt = ParentActors.CreateConstIterator(); ActorIt; ++ActorIt )
+			// Add the level blueprints for all levels
+			for (auto It = SharedData->RepresentingWorld->GetLevelIterator(); It ; ++It)
 			{
-				AddFilteredParentActorToTree(*ActorIt);
+				ConstructItemFor<FLevelBlueprintTreeItem>(FLevelBlueprintHandle(*It));
 			}
 
-			if (!IsShowingOnlySelected() && ShouldShowFolders())
+			// Iterate over every actor in memory. WARNING: This is potentially very expensive!
+			for( FActorIterator ActorIt(SharedData->RepresentingWorld); ActorIt; ++ActorIt )
 			{
-				// Add any folders which might match the current search terms
-				for (const auto& Pair : FActorFolders::Get().GetFolderPropertiesForWorld(*World))
+				AActor* Actor = *ActorIt;
+				if (Actor && IsActorDisplayable(Actor))
 				{
-					if (!FolderToTreeItemMap.Find(Pair.Key))
+					if (Filters->PassesAllFilters(FActorTreeItem(Actor)))
 					{
-						TSharedRef<TOutlinerFolderTreeItem> NewFolderItem = MakeShareable(new TOutlinerFolderTreeItem(Pair.Key));
-						if (!AddFolderToTree(NewFolderItem))
-						{
-							// If we didn't add that folder, attempt to add its parents too
-							TSet<FName> Parents;
-							NewFolderItem->GetParentPaths(Parents);
-							for (const auto& ParentPath : Parents)
-							{
-								AddFolderToTree(MakeShareable(new TOutlinerFolderTreeItem(ParentPath)));
-							}
-						}
+						ApplicableActors.Emplace(Actor);
 					}
+					ConstructItemFor<FActorTreeItem>(Actor);
 				}
 			}
 		}
 
-		// If we have a selection we may need to scroll to it
-		if( bSyncSelection && LastSelectedItem.IsValid() )
-		{
-			// Scroll last item into view - this means if we are multi-selecting, we show newest selection. @TODO Not perfect though
-			OutlinerTreeView->RequestScrollIntoView(LastSelectedItem);
-		}
-	}
-
-	void SSceneOutliner::RemoveItemFromTree(TSharedRef<TOutlinerTreeItem> InItem)
-	{
-		TSharedRef<TOutlinerFolderTreeItem>* ParentFolderItem = nullptr;
-		if (InItem->Type == TOutlinerTreeItem::Actor)
+		if (!IsShowingOnlySelected() && ShouldShowFolders())
 		{
-			auto ActorItemPtr = StaticCastSharedRef<TOutlinerActorTreeItem>(InItem);
-			auto* Actor = ActorItemPtr->Actor.Get();
-			// We only synchronize selection when in actor browsing mode
-			if( InitOptions.Mode == ESceneOutlinerMode::ActorBrowsing )
-			{
-				// Synchronize selection since we no longer clear the selection list!
-				if( GEditor->GetSelectedActors()->IsSelected( ActorItemPtr->Actor.Get() ) )
-				{
-					OutlinerTreeView->SetItemSelection( InItem, false );
-				}
-			}
-
-			if (Actor)
+			// Add any folders which might match the current search terms
+			for (const auto& Pair : FActorFolders::Get().GetFolderPropertiesForWorld(*SharedData->RepresentingWorld))
 			{
-				FName Path = Actor->GetFolderPath();
-				if (!Path.IsNone())
+				if (!TreeItemMap.Contains(Pair.Key))
 				{
-					ParentFolderItem = FolderToTreeItemMap.Find(Path);
+					ConstructItemFor<FFolderTreeItem>(Pair.Key);
 				}
 			}
-			ActorToTreeItemMap.Remove(ActorItemPtr->Actor);
-
-			if(!InItem->Flags.IsFilteredOut)
-			{
-				--FilteredActorCount;
-			}
-		}
-		else
-		{
-			auto FolderItemPtr = StaticCastSharedRef<TOutlinerFolderTreeItem>(InItem);
-			ParentFolderItem = FolderToTreeItemMap.Find(FolderItemPtr->GetParentPath());
-			FolderToTreeItemMap.Remove(FolderItemPtr->Path);
 		}
+	}
 
-		if (ParentFolderItem)
-		{
-			auto& Parent = (*ParentFolderItem);
-			Parent->ChildItems.Remove(InItem);
-		}
-		else
+	void SSceneOutliner::OnChildRemovedFromParent(ITreeItem& Parent)
+	{
+		if (Parent.Flags.bIsFilteredOut && !Parent.GetChildren().Num())
 		{
-			RootTreeItems.Remove(InItem);
+			// The parent no longer has any children that match the current search terms. Remove it.
+			RemoveItemFromTree(Parent.AsShared());
 		}
 	}
 
-	bool SSceneOutliner::AddItemToTree(TSharedRef<TOutlinerTreeItem> InItem, FParentActorsSet& ParentActors)
+	void SSceneOutliner::OnItemMoved(const FTreeItemRef& Item)
 	{
-		if (InItem->Type == TOutlinerTreeItem::Actor)
+		// Just remove the item if it no longer matches the filters
+		if (!Item->Flags.bIsFilteredOut && !SearchBoxFilter->PassesFilter(*Item))
 		{
-			return AddActorToTree(StaticCastSharedRef<TOutlinerActorTreeItem>(InItem), ParentActors);
+			// This will potentially remove any non-matching, empty parents as well
+			RemoveItemFromTree(Item);
 		}
 		else
 		{
-			return AddFolderToTree(StaticCastSharedRef<TOutlinerFolderTreeItem>(InItem));
+			// The item still matches the filters (or has children that do)
+			// When an item has been asked to move, it will still reside under its old parent
+			FTreeItemPtr Parent = Item->GetParent();
+			if (Parent.IsValid())
+			{
+				Parent->RemoveChild(Item);
+				OnChildRemovedFromParent(*Parent);
+			}
+			else
+			{
+				RootTreeItems.Remove(Item);
+			}
+
+			Parent = EnsureParentForItem(Item);
+			if (Parent.IsValid())
+			{
+				Parent->AddChild(Item);
+				OutlinerTreeView->SetItemExpansion(Parent, true);
+			}
+			else
+			{
+				RootTreeItems.Add(Item);
+			}
 		}
 	}
 
-	bool SSceneOutliner::AddActorToTree(TSharedRef<TOutlinerActorTreeItem> InActorItem, FParentActorsSet& ParentActors)
+	void SSceneOutliner::RemoveItemFromTree(FTreeItemRef InItem)
 	{
-		if (!InActorItem->Actor.IsValid())
+		if (TreeItemMap.Contains(InItem->GetID()))
 		{
-			return false;
-		}
+			auto Parent = InItem->GetParent();
 
-		AActor* Actor = InActorItem->Actor.Get();
-		if (!IsActorDisplayable(Actor) || !Filters->PassesAllFilters(Actor))
-		{
-			return false;
-		}
+			if (Parent.IsValid())
+			{
+				Parent->RemoveChild(InItem);
+				OnChildRemovedFromParent(*Parent);
+			}
+			else
+			{
+				RootTreeItems.Remove(InItem);
+			}
 
-		// Apply text filter
-		if (!SearchBoxFilter->PassesFilter(*InActorItem))
-		{
-			return false;
+			InItem->Visit(FFunctionalVisitor().Actor([&](const FActorTreeItem& ActorItem){
+				if (!ActorItem.Flags.bIsFilteredOut)
+				{
+					--FilteredActorCount;
+				}
+			}));
+
+			TreeItemMap.Remove(InItem->GetID());
 		}
+	}
 
-		if( InitOptions.bShowParentTree )
+	FTreeItemPtr SSceneOutliner::EnsureParentForItem(FTreeItemRef Item)
+	{
+		if (SharedData->bShowParentTree)
 		{
-			// OK so we know that this actor needs to be visible.  Let's gather any parent actors
-			// and make sure they're visible as well, so the user can see the hierarchy.
-			AActor* ParentActor = Actor->GetAttachParentActor();
-			if( ParentActor != NULL )
+			FTreeItemPtr Parent = Item->FindParent(TreeItemMap);
+			if (Parent.IsValid())
 			{
-				for( auto NextParent = ParentActor; NextParent != NULL; NextParent = NextParent->GetAttachParentActor() )
-				{
-					ParentActors.Add( NextParent );
-				}					
+				return Parent;
 			}
 			else
 			{
-				const FName ActorFolder = Actor->GetFolderPath();
-				if (ActorFolder.IsNone() || !ShouldShowFolders())
-				{
-					// Setup our root-level item list
-					RootTreeItems.Add( InActorItem );
-				}
-				else
+				auto NewParent = Item->CreateParent();
+				if (NewParent.IsValid())
 				{
-					// Ignore search text since we know we pass the filter, and will need to show parents
-					const bool bIgnoreSearch = true;
-					AddChildToParent(InActorItem, ActorFolder, bIgnoreSearch);
+					NewParent->Flags.bIsFilteredOut = !Filters->PassesAllFilters(*NewParent) || !SearchBoxFilter->PassesFilter(*NewParent);
+					AddUnfilteredItemToTree(NewParent.ToSharedRef());
+					return NewParent;
 				}
 			}
 		}
-		else
-		{
-			RootTreeItems.Add( InActorItem );
-		}
 
-		// Add the item to the map, making it available for display in the tree when its parent is looking for its children.
-		ActorToTreeItemMap.Add(Actor, InActorItem);
+		return nullptr;
+	}
 
-		++FilteredActorCount;
+	bool SSceneOutliner::AddItemToTree(FTreeItemRef Item)
+	{
+		const auto ItemID = Item->GetID();
 
-		// We only synchronize selection when in actor browsing mode
-		if( InitOptions.Mode == ESceneOutlinerMode::ActorBrowsing )
+		// If a tree item already exists that represents the same data, bail
+		if (TreeItemMap.Find(ItemID))
 		{
-			// Synchronize selection
-			if( GEditor->GetSelectedActors()->IsSelected( Actor ) )
-			{
-				OutlinerTreeView->SetItemSelection( InActorItem, true );
-			}
+			return false;
 		}
 
-		NewItemActions.ItemHasBeenCreated(OutlinerTreeView.ToSharedRef(), InActorItem);
+		// Set the filtered out flag
+		Item->Flags.bIsFilteredOut = !SearchBoxFilter->PassesFilter(*Item);
 
-		return true;
-	}
-
-	bool SSceneOutliner::AddFolderToTree(TSharedRef<TOutlinerFolderTreeItem> InFolderItem, bool bIgnoreSearchFilter)
-	{
-		if (!InitOptions.bShowParentTree || !Filters->PassesAllFilters(InFolderItem->Path))
+		if (!Item->Flags.bIsFilteredOut)
 		{
-			return false;
-		}
-
-		if (FolderToTreeItemMap.Find(InFolderItem->Path))
-		{
-			return false;
-		}
-
-		// Apply text filter
-		if (!bIgnoreSearchFilter && !SearchBoxFilter->PassesFilter(*InFolderItem))
-		{
-			return false;
-		}
-
-		// Ensure our parent is added first
-		FName ParentPath = InFolderItem->GetParentPath();
-
-		if (ParentPath.IsNone())
-		{
-			RootTreeItems.Add(InFolderItem);
-		}
-		else
-		{
-			// Ignore search text since we know we pass the filter, and will need to show parents
-			const bool bIgnoreSearch = true;
-			AddChildToParent(InFolderItem, ParentPath, bIgnoreSearch);
-		}
-
-		// Add the item to the map, making it available for display in the tree when its parent is looking for its children.
-		FolderToTreeItemMap.Add(InFolderItem->Path, InFolderItem);
-
-		if (FActorFolderProps* Props = FActorFolders::Get().GetFolderProperties(*RepresentingWorld, InFolderItem->Path))
-		{
-			OutlinerTreeView->SetItemExpansion(InFolderItem, Props->bIsExpanded);
-		}
-
-		NewItemActions.ItemHasBeenCreated(OutlinerTreeView.ToSharedRef(), InFolderItem);
-
-		return true;
-	}
-
-	void SSceneOutliner::AddFilteredParentActorToTree(const AActor* InActor)
-	{
-		/* Parents that are filtered out (or otherwise found and added to the list) will not appear in the actor map at this point
-		   and should be added as requested. If they have no parent, they will be added to the root of the tree, otherwise they will
-		   simply be added to the map, so when displaying the tree, their display item can be found. */
-		if( !ActorToTreeItemMap.Contains( InActor ) )
-		{
-			TSharedRef<TOutlinerActorTreeItem> NewItem = MakeShareable( new TOutlinerActorTreeItem(InActor) );
-
-			ActorToTreeItemMap.Add( InActor, NewItem );
-
-			// Setup our root-level item list
-			const auto* ParentActor = InActor->GetAttachParentActor();
-			if( ParentActor == NULL )
-			{
-				RootTreeItems.Add( NewItem );
-			}
+			AddUnfilteredItemToTree(Item);
 
-			// We only synchronize selection when in actor browsing mode
-			if( InitOptions.Mode == ESceneOutlinerMode::ActorBrowsing )
+			// Check if we need to do anything with this new item
+			if (uint8* ActionMask = NewItemActions.Find(ItemID))
 			{
-				// Synchronize selection
-				if( GEditor->GetSelectedActors()->IsSelected( InActor ) )
+				if (*ActionMask & ENewItemAction::Select)
 				{
-					OutlinerTreeView->SetItemSelection( NewItem, true );
+					OutlinerTreeView->ClearSelection();
+					OutlinerTreeView->SetItemSelection(Item, true);
 				}
-			}
-
-			NewItem->Flags.IsFilteredOut = true;
-		}
-	}
-
-	void SSceneOutliner::GetParentsExpansionState(FParentsExpansionState& ExpansionStateInfo) const
-	{
-		for (const auto& Item : ActorToTreeItemMap)
-		{
-			if (!OutlinerTreeView->IsItemExpanded(Item.Value) && Item.Key.IsValid())
-			{
-				ExpansionStateInfo.CollapsedItems.Add(Item.Key.Get());
-			}
-		}
-	}
-
-	void SSceneOutliner::SetParentsExpansionState(const FParentActorsSet& ParentActors, const FParentsExpansionState& ExpansionStateInfo) const
-	{
-		for (const auto* Actor : ParentActors)
-		{
-			// Set parent actor expansion state according to previous state
-			const auto& Item = ActorToTreeItemMap.FindChecked(Actor);
-			const bool bShouldExpand = !ExpansionStateInfo.CollapsedItems.Contains(Actor);
-			OutlinerTreeView->SetItemExpansion(Item, bShouldExpand);
-		}
-	}
-
-	void SSceneOutliner::PopulateSearchStrings( const TOutlinerTreeItem& TreeItem, OUT TArray< FString >& OutSearchStrings ) const
-	{
-		if (TreeItem.Type == TOutlinerTreeItem::Actor)
-		{
-			AActor* ActorPtr = static_cast<const TOutlinerActorTreeItem&>(TreeItem).Actor.Get();
-			if (ActorPtr)
-			{
-				OutSearchStrings.Add(ActorPtr->GetActorLabel());
-				CustomColumn->PopulateActorSearchStrings( ActorPtr, OUT OutSearchStrings );
-			}
-		}
-		else
-		{
-			OutSearchStrings.Add(static_cast<const TOutlinerFolderTreeItem&>(TreeItem).LeafName.ToString());
-			OutSearchStrings.Add(LOCTEXT("FolderClassName", "Folder").ToString());
-		}
-	}
-
-	EVisibility SSceneOutliner::IsActorClassNameVisible() const
-	{
-		return CustomColumn->ProvidesSearchStrings() || FilterTextBoxWidget->GetText().IsEmpty() ? EVisibility::Collapsed : EVisibility::Visible;
-	}
-
-	FString SSceneOutliner::GetClassNameForItem( FOutlinerTreeItemPtr TreeItem ) const
-	{
-		FString LabelText;
-		const FString& CurrentFilterText = FilterTextBoxWidget->GetText().ToString();
 
-		if( !CustomColumn->ProvidesSearchStrings() && !CurrentFilterText.IsEmpty() )
-		{
-			// If the user is currently searching, then also display the actor class name.  The class name is
-			// searchable, too, and we want to indicate that to the user.
-			if (TreeItem->Type == TOutlinerTreeItem::Actor)
-			{
-				AActor* Actor = StaticCastSharedPtr<TOutlinerActorTreeItem>(TreeItem)->Actor.Get();
-				if( Actor )
+				if (*ActionMask & ENewItemAction::Rename)
 				{
-					FString ActorClassName;
-					Actor->GetClass()->GetName( ActorClassName );
-					LabelText = FString::Printf( TEXT( " (%s)" ), *ActorClassName );
+					PendingRenameItem = Item;
 				}
-			}
-			else
-			{
-				LabelText = FString::Printf(TEXT(" (%s)"), *LOCTEXT("FolderClassName", "Folder").ToString());
-			}
-		}
-
-		return LabelText;
-	}
-
-	FSlateColor SSceneOutliner::GetColorAndOpacityForItem( FOutlinerTreeItemPtr TreeItem ) const
-	{
-		FString LabelText;
-		bool bDarkenItem = false;
-		if (TreeItem->Type == TOutlinerTreeItem::Actor)
-		{
-			auto ActorTreeItem = StaticCastSharedPtr<TOutlinerActorTreeItem>(TreeItem);
-			AActor* Actor = ActorTreeItem->Actor.Get();
-
-			if( Actor == NULL )
-			{
-				return FLinearColor( 0.2f, 0.2f, 0.25f );		// Deleted actor!
-			}
 
-			if ( bRepresentingPlayWorld && !ActorTreeItem->bExistsInCurrentWorldAndPIE )
-			{
-				// Highlight actors that are exclusive to PlayWorld
-				return FLinearColor( 0.12f, 0.56f, 1.0f );
-			}
-
-			// Darken items that aren't suitable targets for an active drag and drop actor attachment action
-			if ( FSlateApplication::Get().IsDragDropping() )
-			{
-				TSharedPtr<FDragDropOperation> DragDropOp = FSlateApplication::Get().GetDragDroppingContent();
-
-				if (DragDropOp->IsOfType<FActorDragDropGraphEdOp>())
+				if (*ActionMask & (ENewItemAction::ScrollIntoView | ENewItemAction::Rename))
 				{
-					TSharedPtr<FActorDragDropGraphEdOp> DragDropActorOp = StaticCastSharedPtr<FActorDragDropGraphEdOp>(DragDropOp);
-
-					bool CanAttach = true;
-					for(int32 i=0; i<DragDropActorOp->Actors.Num(); i++)
-					{
-						AActor* ChildActor = DragDropActorOp->Actors[i].Get();
-						if (ChildActor)
-						{
-							if (!GEditor->CanParentActors(Actor, ChildActor))
-							{
-								bDarkenItem = true;
-								break;
-							}
-						}
-					}
+					ScrollItemIntoView(Item);
 				}
 			}
-
-			// also darken items that are non selectable in the active mode(s)
-			const bool bInSelected = true;
-			const bool bSelectEvenIfHidden = true;		// @todo outliner: Is this actually OK?
-			bDarkenItem |= !GEditor->CanSelectActor( Actor, bInSelected, bSelectEvenIfHidden );
-		}
-
-		if( !InitOptions.bShowParentTree )
-		{
-			return FSlateColor::UseForeground();
-		}
-
-		// Highlight text differently if it doesn't match the search filter (e.g., parent actors to child actors that
-		// match search criteria.)
-		if( TreeItem->Flags.IsFilteredOut )
-		{
-			return FLinearColor( 0.30f, 0.30f, 0.30f );
 		}
 
-		return (!bDarkenItem) ? FSlateColor::UseForeground() : FLinearColor( 0.30f, 0.30f, 0.30f );
+		return true;
 	}
 
-	const FSlateBrush* SSceneOutliner::GetIconForItem(FOutlinerTreeItemRef Item) const
+	void SSceneOutliner::AddUnfilteredItemToTree(FTreeItemRef Item)
 	{
-		if (Item->Type == TOutlinerTreeItem::Actor)
-		{
-			AActor* Actor = StaticCastSharedRef<TOutlinerActorTreeItem>(Item)->Actor.Get();
-			return FClassIconFinder::FindIconForActor(Actor);
-		}
-		else
-		{
-			auto Folder = StaticCastSharedRef<TOutlinerFolderTreeItem>(Item);
-			if (OutlinerTreeView->IsItemExpanded(Item) && Folder->ChildItems.Num())
-			{
-				return FEditorStyle::Get().GetBrush(TEXT("SceneOutliner.FolderOpen"));
-			}
-			else
-			{
-				return FEditorStyle::Get().GetBrush(TEXT("SceneOutliner.FolderClosed"));
-			}
-		}
-	}
+		Item->SharedData = SharedData;
 
-	FText SSceneOutliner::GetToolTipTextForItemIcon( FOutlinerTreeItemPtr TreeItem ) const
-	{
-		FText ToolTipText;
-		if (TreeItem->Type == TOutlinerTreeItem::Actor)
-		{
-			auto ActorTreeItem = StaticCastSharedPtr<TOutlinerActorTreeItem>(TreeItem);
-			AActor* Actor = ActorTreeItem->Actor.Get();
-			if( Actor )
-			{
-				ToolTipText = FText::FromString( Actor->GetClass()->GetName() );
-				if( InitOptions.Mode == ESceneOutlinerMode::ActorBrowsing )
-				{
-					USceneComponent* RootComponent = Actor->GetRootComponent();
-					if( RootComponent )
-					{
-						FFormatNamedArguments Args;
-						Args.Add( TEXT("ActorClassName"), ToolTipText );
+		auto Parent = EnsureParentForItem(Item);
 
-						if( RootComponent->Mobility == EComponentMobility::Static )
-						{
-							ToolTipText = FText::Format( LOCTEXT( "ComponentMobility_Static", "{ActorClassName} with static mobility" ), Args );
-						}
-						else if( RootComponent->Mobility == EComponentMobility::Stationary )
-						{
-							ToolTipText = FText::Format( LOCTEXT( "ComponentMobility_Stationary", "{ActorClassName} with stationary mobility" ), Args );
-						}
-						else if( RootComponent->Mobility == EComponentMobility::Movable )
-						{
-							ToolTipText = FText::Format( LOCTEXT( "ComponentMobility_Movable", "{ActorClassName} with movable mobility" ), Args );
-						}
-					}
-				}
-			}
-		}
-		else
-		{
-			ToolTipText = LOCTEXT("FolderClassName", "Folder");
-		}
+		const auto ItemID = Item->GetID();
+		check(!TreeItemMap.Contains(ItemID));
 
-		return ToolTipText;
-	}
-
-	const FSlateBrush* SSceneOutliner::GetBrushForComponentMobilityIcon( FOutlinerTreeItemPtr TreeItem ) const
-	{
-		const FSlateBrush* IconBrush = MobilityStaticBrush;
-		
+		TreeItemMap.Add(ItemID, Item);
 
-		if (TreeItem->Type == TOutlinerTreeItem::Folder)
+		if (Parent.IsValid())
 		{
-			return IconBrush;
+			Parent->AddChild(Item);
 		}
-
-		AActor* Actor = StaticCastSharedPtr<TOutlinerActorTreeItem>(TreeItem)->Actor.Get();
-		if( Actor )
+		else
 		{
-			USceneComponent* RootComponent = Actor->GetRootComponent();
-			if( RootComponent )
-			{
-				if( RootComponent->Mobility == EComponentMobility::Stationary )
-				{
-					IconBrush = MobilityStationaryBrush;
-				}
-				else if( RootComponent->Mobility == EComponentMobility::Movable )
-				{
-					IconBrush = MobilityMovableBrush;
-				}
-			}
+			RootTreeItems.Add(Item);
 		}
 
-		return IconBrush;
+		Item->Visit(FOnItemAddedToTree(*this));
 	}
 
-	TSharedRef< SWidget > SSceneOutliner::GenerateWidgetForItemAndColumn( FOutlinerTreeItemPtr Item, const FName ColumnID, FIsSelected InIsSelected ) const
+	SSceneOutliner::FParentsExpansionState SSceneOutliner::GetParentsExpansionState() const
 	{
-		struct Local
-		{
-			/** Gets the tool-tip text for an actor */
-			static FString GetToolTipTextForActor( FOutlinerTreeItemPtr TreeItem )
-			{
-				if (TreeItem->Type == TOutlinerTreeItem::Actor)
-				{
-					const auto Actor = StaticCastSharedPtr<TOutlinerActorTreeItem>(TreeItem)->Actor;
-					if( Actor != NULL )
-					{
-						FFormatNamedArguments Args;
-						Args.Add(TEXT("ID_Name"), LOCTEXT("CustomColumnMode_InternalName", "ID Name"));
-						Args.Add(TEXT("Name"), FText::FromString(Actor->GetName()));
-						return FText::Format(LOCTEXT("ActorNameTooltip", "{ID_Name}: {Name}"), Args).ToString();
-					}
-				}
-				else
-				{
-					FFormatNamedArguments Args;
-					Args.Add(TEXT("Name"), FText::FromName(StaticCastSharedPtr<TOutlinerFolderTreeItem>(TreeItem)->LeafName));
-					return FText::Format(LOCTEXT("FolderNameTooltip", "Folder: {Name}"), Args).ToString();
-				}
-
-				return FString();
-			}
-
-		};
-
-		TSharedPtr< SWidget > TableRowContent;
-		if( ColumnID == ColumnID_ActorLabel )
+		FParentsExpansionState States;
+		for (const auto& Pair : TreeItemMap)
 		{
-			TSharedPtr< SInlineEditableTextBlock > InlineTextBlock;
-			TSharedPtr< SHorizontalBox > HBox = SNew( SHorizontalBox );
-
-			const FSlateColor IconColor = InitOptions.Mode == ESceneOutlinerMode::ActorBrowsing ? FLinearColor::White : FEditorStyle::GetSlateColor( "DefaultForeground" );
-
-			auto IconOverlay = SNew(SOverlay)
-				+SOverlay::Slot()
-				.HAlign(HAlign_Right)
-				.VAlign(VAlign_Center)
-				[
-					SNew( SImage )
-					.Image( this, &SSceneOutliner::GetIconForItem, Item.ToSharedRef() )
-					.ToolTipText( this, &SSceneOutliner::GetToolTipTextForItemIcon, Item )
-					.ColorAndOpacity( IconColor )
-				];
-
-			if( InitOptions.Mode == ESceneOutlinerMode::ActorBrowsing )
+			if (Pair.Value->GetChildren().Num())
 			{
-				// Add the component mobility icon
-				IconOverlay->AddSlot()
-					.HAlign(HAlign_Left)
-					[
-						SNew( SImage )
-						.Image( this, &SSceneOutliner::GetBrushForComponentMobilityIcon, Item )
-					];
+				States.Add(Pair.Key, Pair.Value->Flags.bIsExpanded);
 			}
-
-			HBox->AddSlot()
-				.AutoWidth()
-				.VAlign(VAlign_Center)
-				.Padding( FMargin(0.0f, 0.0f, 6.0f, 0.0f) )
-				[
-					SNew(SBox)
-					.WidthOverride(18.f)
-					[
-						IconOverlay
-					]
-				];
-
-			HBox->AddSlot()
-				.FillWidth(1.0f)
-				.VAlign(VAlign_Center)
-				.Padding( 0.0f, 2.0f, 0.0f, 2.0f )
-				[
-					SAssignNew( InlineTextBlock, SInlineEditableTextBlock )
-
-						// Bind a delegate for getting the actor's name.  This ensures that the current actor name is
-						// shown accurately, even if it gets renamed while the tree is visible. (Without requiring us
-						// to refresh anything.)
- 						.Text_Static( &GetLabelForItem, Item.ToSharedRef() )
-
-						// Bind our filter text as the highlight string for the text block, so that when the user
-						// starts typing search criteria, this text highlights
-						.HighlightText( this, &SSceneOutliner::GetFilterHighlightText )
-
-						// Use the actor's actual FName as the tool-tip text
-						.ToolTipText_Static( &Local::GetToolTipTextForActor, Item )
-
-						// Actor names that pass the search criteria are shown in a brighter color
-						.ColorAndOpacity( this, &SSceneOutliner::GetColorAndOpacityForItem, Item )
-
-						.OnTextCommitted( this, &SSceneOutliner::OnItemLabelCommitted, Item )
-
-						.OnVerifyTextChanged( this, &SSceneOutliner::OnVerifyItemLabelChanged, Item )
-
-						.IsSelected( InIsSelected )
-				];
-
-			HBox->AddSlot()
-				.AutoWidth()
-				.Padding( 0.0f, 3.0f, 3.0f, 0.0f )
-				[
-					SNew( STextBlock )
-
-						// Bind a delegate for getting the actor's class name
- 						.Text( this, &SSceneOutliner::GetClassNameForItem, Item )
-
-						.Visibility( this, &SSceneOutliner::IsActorClassNameVisible )
-
-						// Bind our filter text as the highlight string for the text block, so that when the user
-						// starts typing search criteria, this text highlights
-						.HighlightText( this, &SSceneOutliner::GetFilterHighlightText )
-				];
-
-			TableRowContent = 
-				SNew( SSceneOutlinerItemWidget )
-					.Content()
-					[
-						HBox.ToSharedRef()
-					];
-
-			Item->RenameRequestEvent.BindSP( InlineTextBlock.Get(), &SInlineEditableTextBlock::EnterEditingMode );
-
-		}
-		else if( Gutter.IsValid() && ColumnID == Gutter->GetColumnID() )
-		{
-			TableRowContent = Gutter->ConstructRowWidget( Item.ToSharedRef() );
-		}
-		else if( ensure( ColumnID == CustomColumn->GetColumnID() ) )
-		{
-			TableRowContent = CustomColumn->ConstructRowWidget( Item.ToSharedRef() );
 		}
-
-		return TableRowContent.ToSharedRef();
+		return States;
 	}
 
-	void SSceneOutliner::OnItemLabelCommitted( const FText& InLabel, ETextCommit::Type InCommitInfo, FOutlinerTreeItemPtr TreeItem )
+	void SSceneOutliner::SetParentsExpansionState(const FParentsExpansionState& ExpansionStateInfo) const
 	{
-		if (TreeItem->Type == TOutlinerTreeItem::Actor)
-		{
-			auto ActorTreeItem = StaticCastSharedPtr<TOutlinerActorTreeItem>(TreeItem);
-			auto* Actor = ActorTreeItem->Actor.Get();
-
-			if (Actor && Actor->IsActorLabelEditable() && !InLabel.ToString().Equals(Actor->GetActorLabel(), ESearchCase::CaseSensitive))
-			{
-				const FScopedTransaction Transaction( LOCTEXT( "SceneOutlinerRenameActorTransaction", "Rename Actor" ) );
-				Actor->SetActorLabel( InLabel.ToString() );
-
-				// Set the keyboard focus back to the SceneOutliner, so the user can perform keyboard commands
-				FWidgetPath SceneOutlinerWidgetPath;
-				FSlateApplication::Get().GeneratePathToWidgetUnchecked( SharedThis( this ), SceneOutlinerWidgetPath );
-				FSlateApplication::Get().SetKeyboardFocus( SceneOutlinerWidgetPath, EFocusCause::SetDirectly );
-			}
-		}
-		else
+		for (const auto& Pair : TreeItemMap)
 		{
-			auto FolderTreeItem = StaticCastSharedPtr<TOutlinerFolderTreeItem>(TreeItem);
-			if (!InLabel.ToString().Equals(FolderTreeItem->LeafName.ToString(), ESearchCase::CaseSensitive))
+			auto& Item = Pair.Value;
+			if (Item->GetChildren().Num())
 			{
-				// Rename the item
-				FName NewPath = FolderTreeItem->GetParentPath();
-				if (NewPath.IsNone())
+				const bool* bIsExpanded = ExpansionStateInfo.Find(Pair.Key);
+				if (bIsExpanded)
 				{
-					NewPath = FName(*InLabel.ToString());
+					OutlinerTreeView->SetItemExpansion(Item, *bIsExpanded);
 				}
 				else
 				{
-					NewPath = FName(*(NewPath.ToString() / InLabel.ToString()));
-				}
-
-				if (FActorFolders::Get().RenameFolderInWorld(*RepresentingWorld, FolderTreeItem->Path, NewPath))
-				{
-					NewItemActions.WhenCreated(NewPath, NewItemActuator::Select | NewItemActuator::ScrollIntoView);
+					OutlinerTreeView->SetItemExpansion(Item, Item->Flags.bIsExpanded);
 				}
 			}
 		}
 	}
 
-	bool SSceneOutliner::OnVerifyItemLabelChanged( const FText& InLabel, FText& OutErrorMessage, FOutlinerTreeItemPtr TreeItem )
+	void SSceneOutliner::PopulateSearchStrings(const ITreeItem& Item, TArray< FString >& OutSearchStrings) const
 	{
-		FText TrimmedLabel = FText::TrimPrecedingAndTrailing(InLabel);
-
-		if (TrimmedLabel.IsEmpty())
+		for (const auto& Pair : Columns)
 		{
-			OutErrorMessage = LOCTEXT( "RenameFailed_LeftBlank", "Names cannot be left blank" );
-			return false;
+			Pair.Value->PopulateSearchStrings(Item, OutSearchStrings);
 		}
-
-		if (TrimmedLabel.ToString().Len() >= NAME_SIZE)
-		{
-			FFormatNamedArguments Arguments;
-			Arguments.Add( TEXT("CharCount"), NAME_SIZE );
-			OutErrorMessage = FText::Format(LOCTEXT("RenameFailed_TooLong", "Names must be less than {CharCount} characters long."), Arguments);
-			return false;
-		}
-
-		FString LabelString = TrimmedLabel.ToString();
-		if (TreeItem->Type == TOutlinerTreeItem::Folder)
-		{
-			auto FolderTreeItem = StaticCastSharedPtr<TOutlinerFolderTreeItem>(TreeItem);
-			if (FolderTreeItem->LeafName == FName(*LabelString))
-			{
-				return true;
-			}
-
-			int32 Dummy = 0;
-			if (LabelString.FindChar('/', Dummy) || LabelString.FindChar('\\', Dummy))
-			{
-				OutErrorMessage = LOCTEXT("RenameFailed_InvalidChar", "Folder names cannot contain / or \\.");
-				return false;
-			}
-
-			// Validate that this folder doesn't exist already
-			FName NewPath = FolderTreeItem->GetParentPath();
-			if (NewPath.IsNone())
-			{
-				NewPath = FName(*LabelString);
-			}
-			else
-			{
-				NewPath = FName(*(NewPath.ToString() / LabelString));
-			}
-
-			if (FolderToTreeItemMap.Find(NewPath))
-			{
-				OutErrorMessage = LOCTEXT("RenameFailed_AlreadyExists", "A folder with this name already exists at this level");
-				return false;
-			}
-		}
-
-		return true;
-	}
-
-	void SSceneOutliner::ExpandActor(AActor* Actor)
-	{
-		if( !InitOptions.bShowParentTree )
-		{
-			return;
-		}
-
-		OutlinerTreeView->SetItemExpansion(*ActorToTreeItemMap.Find(Actor), true);
 	}
 
-	TArray<TSharedRef<TOutlinerFolderTreeItem>> SSceneOutliner::GetSelectedFolders() const
+	TArray<FFolderTreeItem*> SSceneOutliner::GetSelectedFolders() const
 	{
-		TArray<TSharedRef<TOutlinerFolderTreeItem>> Folders;
-		for (auto& Item : OutlinerTreeView->GetSelectedItems())
-		{
-			if (Item->Type == TOutlinerTreeItem::Folder)
-			{
-				Folders.Add(StaticCastSharedPtr<TOutlinerFolderTreeItem>(Item).ToSharedRef());
-			}
-		}
-		return Folders;
+		return FItemSelection(*OutlinerTreeView).Folders;
 	}
 
-	TSharedPtr<SWidget> SSceneOutliner::OnOpenContextMenu() const
+	TSharedPtr<SWidget> SSceneOutliner::OnOpenContextMenu()
 	{
 		TArray<AActor*> SelectedActors;
 		GEditor->GetSelectedActors()->GetSelectedObjects<AActor>( SelectedActors );
 
-		if (SelectedActors.Num() && InitOptions.ContextMenuOverride.IsBound())
+		/** Legacy context menu override only supports actors */
+		if (SelectedActors.Num() && SharedData->ContextMenuOverride.IsBound())
 		{
-			return InitOptions.ContextMenuOverride.Execute();
+			return SharedData->ContextMenuOverride.Execute();
 		}
 
-		if (InitOptions.Mode == ESceneOutlinerMode::ActorBrowsing)
+		if (SharedData->Mode == ESceneOutlinerMode::ActorBrowsing)
 		{
 			return BuildDefaultContextMenu();
 		}
@@ -1666,52 +1039,50 @@ namespace SceneOutliner
 		return TSharedPtr<SWidget>();
 	}
 
-	TSharedPtr<SWidget> SSceneOutliner::BuildDefaultContextMenu() const
+	TSharedPtr<SWidget> SSceneOutliner::BuildDefaultContextMenu()
 	{
-		// Build up the menu
-		const bool bCloseAfterSelection = true;
-		FMenuBuilder MenuBuilder(bCloseAfterSelection, TSharedPtr<FUICommandList>(), InitOptions.DefaultMenuExtender);
+		if (!CheckWorld())
+		{
+			return nullptr;
+		}
 
-		auto SelectedItems = OutlinerTreeView->GetSelectedItems();
-		auto SelectedFolders = GetSelectedFolders();
+		// Build up the menu for a selection
+		const bool bCloseAfterSelection = true;
+		FMenuBuilder MenuBuilder(bCloseAfterSelection, TSharedPtr<FUICommandList>(), SharedData->DefaultMenuExtender);
 
-		TSharedPtr<TOutlinerFolderTreeItem> ParentFolder;
-		if (SelectedFolders.Num() == 1)
+		const auto NumSelectedItems = OutlinerTreeView->GetNumItemsSelected();
+		if (NumSelectedItems == 1)
 		{
-			ParentFolder = SelectedFolders[0];
+			OutlinerTreeView->GetSelectedItems()[0]->GenerateContextMenu(MenuBuilder, *this);
 		}
 
-		const FSlateIcon NewFolderIcon(FEditorStyle::GetStyleSetName(), "SceneOutliner.NewFolderIcon");
 		// We always create a section here, even if there is no parent so that clients can still extend the menu
-		MenuBuilder.BeginSection("FolderSection", LOCTEXT("ThisFolderSection", "This Folder"));
+		MenuBuilder.BeginSection("MainSection");
 		{
-			if (ParentFolder.IsValid())
-			{
-				MenuBuilder.AddMenuEntry(LOCTEXT("CreateSubFolder", "Create Sub Folder"), FText(), NewFolderIcon, FUIAction(FExecuteAction::CreateSP(this, &SSceneOutliner::CreateSubFolder, ParentFolder->Path)));
-				MenuBuilder.AddMenuEntry(LOCTEXT("RenameFolder", "Rename Folder"), FText(), FSlateIcon(), FUIAction(FExecuteAction::CreateSP(this, &SSceneOutliner::RenameFolder, ParentFolder.ToSharedRef())));
-				MenuBuilder.AddMenuEntry(LOCTEXT("DeleteFolder", "Delete Folder"), FText(), FSlateIcon(), FUIAction(FExecuteAction::CreateSP(this, &SSceneOutliner::DeleteFolder, ParentFolder.ToSharedRef())));
-			}
-		}
-		MenuBuilder.EndSection();
+			FItemSelection ItemSelection(*OutlinerTreeView);
 
-		// Don't add any of these menu items if we're not showing the parent tree
-		if (InitOptions.bShowParentTree)
-		{
-			if (SelectedItems.Num() == 0)
+			// Don't add any of these menu items if we're not showing the parent tree
+			if (SharedData->bShowParentTree)
 			{
-				MenuBuilder.AddMenuEntry(LOCTEXT("CreateFolder", "Create Folder"), FText(), NewFolderIcon, FUIAction(FExecuteAction::CreateSP(this, &SSceneOutliner::CreateFolder)));
-			}
-			else
-			{
-				MenuBuilder.BeginSection("SelectionSection", LOCTEXT("OutlinerSelection", "Current Outliner Selection"));
+				if (NumSelectedItems == 0)
+				{
+					const FSlateIcon NewFolderIcon(FEditorStyle::GetStyleSetName(), "SceneOutliner.NewFolderIcon");
+					MenuBuilder.AddMenuEntry(LOCTEXT("CreateFolder", "Create Folder"), FText(), NewFolderIcon, FUIAction(FExecuteAction::CreateSP(this, &SSceneOutliner::CreateFolder)));
+				}
+				else
 				{
+					// Can't move worlds or level blueprints
+					const bool bCanMoveSelection = ItemSelection.Worlds.Num() + ItemSelection.LevelBlueprints.Num() == 0;
+					if (bCanMoveSelection)
+					{
 					MenuBuilder.AddSubMenu(
-						LOCTEXT("MoveActorsToFolder", "Move To Folder"),
-						LOCTEXT("MoveActorsToFolder_Tooltip", "Move selection to a folder"),
+						LOCTEXT("MoveActorsTo", "Move To"),
+						LOCTEXT("MoveActorsTo_Tooltip", "Move selection to another folder"),
 						FNewMenuDelegate::CreateSP(this,  &SSceneOutliner::FillFoldersSubMenu));
+					}
 
 					// If we've only got folders selected, show the selection sub menu
-					if (SelectedFolders.Num() == SelectedItems.Num())
+					if (ItemSelection.Folders.Num() == NumSelectedItems)
 					{
 						MenuBuilder.AddSubMenu(
 							LOCTEXT("SelectSubmenu", "Select"),
@@ -1721,109 +1092,146 @@ namespace SceneOutliner
 				}
 			}
 		}
+		MenuBuilder.EndSection();
 
 		return MenuBuilder.MakeWidget();
 	}
 
 	void SSceneOutliner::FillFoldersSubMenu(FMenuBuilder& MenuBuilder) const
 	{
-		MenuBuilder.AddMenuEntry(LOCTEXT( "CreateNew", "Create New" ), LOCTEXT( "CreateNew_ToolTip", "Move to a new folder" ),
+		MenuBuilder.AddMenuEntry(LOCTEXT( "CreateNew", "Create New Folder" ), LOCTEXT( "CreateNew_ToolTip", "Move to a new folder" ),
 			FSlateIcon(FEditorStyle::GetStyleSetName(), "SceneOutliner.NewFolderIcon"), FExecuteAction::CreateSP(this, &SSceneOutliner::CreateFolder));
 
-		if (OutlinerTreeView->ValidateMoveSelectionTo(FName()))
+		AddMoveToFolderOutliner(MenuBuilder);
+	}
+
+	TSharedRef<TSet<FName>> SSceneOutliner::GatherInvalidMoveToDestinations() const
+	{
+		// We use a pointer here to save copying the whole array for every invocation of the filter delegate
+		TSharedRef<TSet<FName>> ExcludedParents(new TSet<FName>());
+
+		struct FFindInvalidFolders : ITreeItemVisitor
+		{
+			TSet<FName>& ExcludedParents;
+			const TMap<FTreeItemID, FTreeItemPtr>& TreeItemMap;
+			FFindInvalidFolders(TSet<FName>& InExcludedParents, const TMap<FTreeItemID, FTreeItemPtr>& InTreeItemMap)
+				: ExcludedParents(InExcludedParents), TreeItemMap(InTreeItemMap)
+			{}
+
+			static bool ItemHasSubFolders(const TWeakPtr<ITreeItem>& WeakItem)
+			{
+				bool bHasSubFolder = false;
+				WeakItem.Pin()->Visit(FFunctionalVisitor().Folder([&](const FFolderTreeItem&){
+					bHasSubFolder = true;
+				}));
+				return bHasSubFolder;
+			}
+
+			virtual void Visit(const FActorTreeItem& ActorItem) const override
+			{
+				if (AActor* Actor = ActorItem.Actor.Get())
+				{
+					// We exclude actor parent folders if they don't have any sub folders
+					const FName& Folder = Actor->GetFolderPath();
+					if (!Folder.IsNone() && !ExcludedParents.Contains(Folder))
+					{
+						auto FolderItem = TreeItemMap.FindRef(Folder);
+						if (FolderItem.IsValid() && !FolderItem->GetChildren().ContainsByPredicate(&ItemHasSubFolders))
+						{
+							ExcludedParents.Add(Folder);
+						}
+					}
+				}
+			}
+
+			virtual void Visit(const FFolderTreeItem& Folder) const override
+			{
+				// Cannot move into its parent
+				const FName ParentPath = GetParentPath(Folder.Path);
+				if (!ParentPath.IsNone())
+				{
+					ExcludedParents.Add(ParentPath);
+				}
+				else
+				{
+					// Failing that, cannot move into itself, or any child
+					ExcludedParents.Add(Folder.Path);
+				}
+			}
+		};
+
+		auto Visitor = FFindInvalidFolders(*ExcludedParents, TreeItemMap);
+		for (const auto& Item : OutlinerTreeView->GetSelectedItems())
 		{
-			MenuBuilder.AddMenuEntry(LOCTEXT( "MoveToRoot", "Move To Root" ), LOCTEXT( "MoveToRoot_ToolTip", "Move to the root of the scene outliner" ),
-				FSlateIcon(FEditorStyle::GetStyleSetName(), "SceneOutliner.MoveToRoot"), FExecuteAction::CreateSP(this, &SSceneOutliner::MoveSelectionTo, FName()));
+			Item->Visit(Visitor);
 		}
 
-		AddMoveToFolderOutliner(MenuBuilder);
+		return ExcludedParents;
 	}
 
 	void SSceneOutliner::AddMoveToFolderOutliner(FMenuBuilder& MenuBuilder) const
 	{
-		// We don't show this if there is no world (unusual) or if there aren't any folders in the world
-		const int32 NumFoldersInWorld = FActorFolders::Get().GetFolderPropertiesForWorld(*RepresentingWorld).Num();
-		if (!RepresentingWorld || !NumFoldersInWorld)
+		// We don't show this if there aren't any folders in the world
+		if (!FActorFolders::Get().GetFolderPropertiesForWorld(*SharedData->RepresentingWorld).Num())
 		{
 			return;
 		}
 
 		// Add a mini scene outliner for choosing an existing folder
-		FSceneOutlinerInitializationOptions MiniSceneOutlinerInitOptions;
+		FInitializationOptions MiniSceneOutlinerInitOptions;
 		MiniSceneOutlinerInitOptions.bShowHeaderRow = false;
 		MiniSceneOutlinerInitOptions.bFocusSearchBoxWhenOpened = true;
 		MiniSceneOutlinerInitOptions.bOnlyShowFolders = true;
-			
+		
+		// Don't show any folders that are a child of any of the selected folders
+		auto ExcludedParents = GatherInvalidMoveToDestinations();
+		if (ExcludedParents->Num())
 		{
-			// Don't show any folders that are a child of any of the selected folders
-
-			// We use a pointer here to save copying the whole array for every invocation of the filter delegate
-			TSharedRef<TSet<FName>> ExcludedParents(new TSet<FName>());
-			for (const auto& Item : OutlinerTreeView->GetSelectedItems())
-			{
-				if (Item->Type == TOutlinerTreeItem::Folder)
-				{
-					// Can't move the folder into its current parent (it already exists)
-					const auto Folder = StaticCastSharedPtr<TOutlinerFolderTreeItem>(Item);
-					const FName ParentPath = Folder->GetParentPath();
-					ExcludedParents->Add(ParentPath.IsNone() ? Folder->Path : ParentPath);
-				}
-				else if (AActor* Actor = StaticCastSharedPtr<TOutlinerActorTreeItem>(Item)->Actor.Get())
+			// Add a filter if necessary
+			auto FilterOutChildFolders = [](FName Path, TSharedRef<TSet<FName>> InExcludedParents){
+				for (const auto& Parent : *InExcludedParents)
 				{
-					// We exclude actor parent folders if they don't have any sub folders
-					const FName& Folder = Actor->GetFolderPath();
-					if (ExcludedParents->Contains(Folder))
+					if (Path == Parent || FActorFolders::PathIsChildOf(Path.ToString(), Parent.ToString()))
 					{
-						continue;
+						return false;
 					}
+				}
+				return true;
+			};
 
-					auto HasSubFolders = [](TWeakPtr<TOutlinerTreeItem> WeakItem){
-						auto TreeItem = WeakItem.Pin();
-						return TreeItem.IsValid() && TreeItem->Type == TOutlinerTreeItem::Folder;
-					};
+			MiniSceneOutlinerInitOptions.Filters->AddFilterPredicate(FFolderFilterPredicate::CreateStatic(FilterOutChildFolders, ExcludedParents), EDefaultFilterBehaviour::Pass);
+		}
 
-					auto* ParentFolder = FolderToTreeItemMap.Find(Folder);
-					if (ParentFolder && !(*ParentFolder)->ChildItems.ContainsByPredicate(HasSubFolders))
-					{
-						ExcludedParents->Add(Folder);
-					}
-				}
-			}
+		{
+			// Filter in/out the world according to whether it is valid to move to/from the root
+			FDragDropPayload DraggedObjects(OutlinerTreeView->GetSelectedItems());
 
-			if (ExcludedParents->Num())
-			{
-				// Add a filter if necessary
-				auto FilterOutChildFolders = [](FName Path, TSharedRef<TSet<FName>> InExcludedParents){
-					for (const auto& Parent : *InExcludedParents)
-					{
-						if (Path == Parent || FActorFolders::PathIsChildOf(Path.ToString(), Parent.ToString()))
-						{
-							return false;
-						}
-					}
+			const bool bMoveToRootValid = FFolderDropTarget(FName()).ValidateDrop(DraggedObjects, *SharedData->RepresentingWorld).IsValid();
 
-					return true;
-				};
+			MiniSceneOutlinerInitOptions.Filters->AddFilterPredicate(FWorldFilterPredicate::CreateStatic([](const UWorld*, bool bMoveToRootValid){
+				return bMoveToRootValid;
+			}, bMoveToRootValid), EDefaultFilterBehaviour::Pass);
+		}
 
-				MiniSceneOutlinerInitOptions.Filters->AddFilterPredicate(SceneOutliner::FFolderFilterPredicate::CreateStatic(FilterOutChildFolders, ExcludedParents));
-			}
+		// Don't show the actor info column
+		MiniSceneOutlinerInitOptions.UseDefaultColumns();
+		MiniSceneOutlinerInitOptions.ColumnMap.Remove(FBuiltInColumnTypes::ActorInfo());
 
-			// Actor selector to allow the user to choose a folder
-			FSceneOutlinerModule& SceneOutlinerModule = FModuleManager::LoadModuleChecked<FSceneOutlinerModule>("SceneOutliner");
-			TSharedRef< SWidget > MiniSceneOutliner =
-				SNew(SVerticalBox)
-				+ SVerticalBox::Slot()
-				.MaxHeight(400.0f)
-				[
-					SNew(SceneOutliner::SSceneOutliner, MiniSceneOutlinerInitOptions)
-					.IsEnabled(FSlateApplication::Get().GetNormalExecutionAttribute())
-					.OnItemPickedDelegate(FOnSceneOutlinerItemPicked::CreateSP(this, &SSceneOutliner::MoveSelectionTo))
-				];
+		// Actor selector to allow the user to choose a folder
+		FSceneOutlinerModule& SceneOutlinerModule = FModuleManager::LoadModuleChecked<FSceneOutlinerModule>("SceneOutliner");
+		TSharedRef< SWidget > MiniSceneOutliner =
+			SNew(SVerticalBox)
+			+ SVerticalBox::Slot()
+			.MaxHeight(400.0f)
+			[
+				SNew(SSceneOutliner, MiniSceneOutlinerInitOptions)
+				.IsEnabled(FSlateApplication::Get().GetNormalExecutionAttribute())
+				.OnItemPickedDelegate(FOnSceneOutlinerItemPicked::CreateSP(this, &SSceneOutliner::MoveSelectionTo))
+			];
 
-			MenuBuilder.BeginSection(FName(), LOCTEXT("ExistingFolders", "Existing:"));
-			MenuBuilder.AddWidget(MiniSceneOutliner, FText::GetEmpty(), false);
-			MenuBuilder.EndSection();
-		}
+		MenuBuilder.BeginSection(FName(), LOCTEXT("ExistingFolders", "Existing:"));
+		MenuBuilder.AddWidget(MiniSceneOutliner, FText::GetEmpty(), false);
+		MenuBuilder.EndSection();
 	}
 
 	void SSceneOutliner::FillSelectionSubMenu(FMenuBuilder& MenuBuilder) const
@@ -1840,6 +1248,25 @@ namespace SceneOutliner
 			FExecuteAction::CreateSP(this, &SSceneOutliner::SelectFoldersDescendants));
 	}
 
+	struct FSelectActors : ITreeItemVisitor
+	{
+		virtual void Visit(const FActorTreeItem& ActorItem) const override
+		{
+			if (AActor* Actor = ActorItem.Actor.Get())
+			{
+				GEditor->SelectActor(Actor, true, /*bNotify=*/false);
+			}
+		}
+
+		virtual void Visit(const FFolderTreeItem& FolderItem) const override
+		{
+			for (auto& Child : FolderItem.GetChildren())
+			{
+				Child.Pin()->Visit(FSelectActors());
+			}
+		}
+	};
+
 	void SSceneOutliner::SelectFoldersImmediateChildren() const
 	{
 		auto SelectedFolders = GetSelectedFolders();
@@ -1848,18 +1275,13 @@ namespace SceneOutliner
 			// We'll batch selection changes instead by using BeginBatchSelectOperation()
 			GEditor->GetSelectedActors()->BeginBatchSelectOperation();
 
+			OutlinerTreeView->ClearSelection();
+
 			for (const auto& Folder : SelectedFolders)
 			{
-				for (const auto& WeakChild : Folder->ChildItems)
+				for (auto& Child : Folder->GetChildren())
 				{
-					auto ChildItem = WeakChild.Pin();
-					if (ChildItem.IsValid() && ChildItem->Type == TOutlinerTreeItem::Actor)
-					{
-						if (AActor* Actor = StaticCastSharedPtr<TOutlinerActorTreeItem>(ChildItem)->Actor.Get())
-						{
-							GEditor->SelectActor(Actor, true, /*bNotify=*/false);
-						}
-					}
+					Child.Pin()->Visit(FSelectActors());
 				}
 			}
 
@@ -1876,9 +1298,11 @@ namespace SceneOutliner
 			// We'll batch selection changes instead by using BeginBatchSelectOperation()
 			GEditor->GetSelectedActors()->BeginBatchSelectOperation();
 
+			OutlinerTreeView->ClearSelection();
+
 			for (const auto& Folder : SelectedFolders)
 			{
-				SelectDescedantActors(Folder, true);
+				Folder->Visit(FSelectActors());
 			}
 
 			GEditor->GetSelectedActors()->EndBatchSelectOperation();
@@ -1886,42 +1310,42 @@ namespace SceneOutliner
 		}
 	}
 
-	FName SSceneOutliner::MoveFolderTo(TSharedRef<TOutlinerFolderTreeItem> Folder, FName NewParent)
+	void SSceneOutliner::MoveSelectionTo(FTreeItemRef NewParent)
 	{
-		FName NewPath;
-		if (NewParent.IsNone())
-		{
-			NewPath = Folder->LeafName;
-		}
-		else
+		struct FMoveToFolder : ITreeItemVisitor
 		{
-			NewPath = FName(*(NewParent.ToString() / Folder->LeafName.ToString()));
-		}
+			SSceneOutliner& Outliner;
+			FMoveToFolder(SSceneOutliner& InOutliner) : Outliner(InOutliner) {}
 
-		if (FActorFolders::Get().RenameFolderInWorld(*RepresentingWorld, Folder->Path, NewPath))
-		{
-			return NewPath;
-		}
+			virtual void Visit(const FFolderTreeItem& Folder) const override
+			{
+				Outliner.MoveSelectionTo(Folder.Path);
+			}
+			virtual void Visit(const FWorldTreeItem&) const override
+			{
+				Outliner.MoveSelectionTo(FName());
+			}
+		};
 
-		return FName();
+		NewParent->Visit(FMoveToFolder(*this));
 	}
 
-	void SSceneOutliner::MoveSelectionTo(TSharedRef<TOutlinerTreeItem> NewParent)
+	void SSceneOutliner::MoveSelectionTo(FName NewParent)
 	{
-		if (NewParent->Type == TOutlinerTreeItem::Folder)
+		if (!CheckWorld())
 		{
-			MoveSelectionTo(StaticCastSharedRef<TOutlinerFolderTreeItem>(NewParent)->Path);
+			return;
 		}
-	}
 
-	void SSceneOutliner::MoveSelectionTo(FName NewParent)
-	{
 		FSlateApplication::Get().DismissAllMenus();
+		
+		FFolderDropTarget 	DropTarget(NewParent);
+		FDragDropPayload 	DraggedObjects(OutlinerTreeView->GetSelectedItems());
 
-		FText ValidationText;
-		if (!OutlinerTreeView->ValidateMoveSelectionTo(NewParent, &ValidationText))
+		FDragValidationInfo Validation = DropTarget.ValidateDrop(DraggedObjects, *SharedData->RepresentingWorld);
+		if (!Validation.IsValid())
 		{
-			FNotificationInfo Info(ValidationText);
+			FNotificationInfo Info(Validation.ValidationText);
 			Info.ExpireDuration = 3.0f;
 			Info.bUseLargeFont = false;
 			Info.bFireAndForget = true;
@@ -1930,379 +1354,199 @@ namespace SceneOutliner
 			return;
 		}
 
-		auto SelectedItems = GetSelectedItems();
-
 		const FScopedTransaction Transaction( LOCTEXT("MoveOutlinerItems", "Move Scene Outliner Items") );
-		for (const auto& Item : SelectedItems)
-		{
-			if (Item->Type == TOutlinerTreeItem::Folder)
-			{
-				auto Folder = StaticCastSharedPtr<TOutlinerFolderTreeItem>(Item).ToSharedRef();
-				auto NewPath = MoveFolderTo(Folder, NewParent);
-				if (!NewPath.IsNone() && SelectedItems.Num() == 1)
-				{
-					NewItemActions.WhenCreated(NewPath, NewItemActuator::Select | NewItemActuator::ScrollIntoView);
-				}
-			}
-			else
-			{
-				auto ActorItem = StaticCastSharedPtr<TOutlinerActorTreeItem>(Item);
-				auto* Actor = ActorItem->Actor.Get();
-				if (Actor)
-				{
-					Actor->SetFolderPath(NewParent);
-
-					if (SelectedItems.Num() == 1)
-					{
-						NewItemActions.WhenCreated(ActorItem.ToSharedRef(), NewItemActuator::Select | NewItemActuator::ScrollIntoView);
-					}
-				}
-			}
-		}
-	}
-
-	bool SSceneOutliner::AddChildToParent(FOutlinerTreeItemRef TreeItem, FName ParentPath, bool bIgnoreSearchFilter)
-	{
-		TSharedRef<TOutlinerFolderTreeItem>* ExistingParent = FolderToTreeItemMap.Find(ParentPath);
-		if (ExistingParent)
-		{
-			(*ExistingParent)->ChildItems.Add(TreeItem);
-			return true;
-		}
-		else
-		{
-			TSharedRef<TOutlinerFolderTreeItem> NewFolderItem = MakeShareable(new TOutlinerFolderTreeItem(ParentPath));
-			if (AddFolderToTree(NewFolderItem, bIgnoreSearchFilter))
-			{
-				NewFolderItem->ChildItems.Add(TreeItem);
-				return true;
-			}
-		}
-
-		return false;
+		DropTarget.OnDrop(DraggedObjects, *SharedData->RepresentingWorld, Validation, SNullWidget::NullWidget);
 	}
 
 	FReply SSceneOutliner::OnCreateFolderClicked()
 	{
-		if (!RepresentingWorld)
-		{
-			return FReply::Unhandled();
-		}
-
 		CreateFolder();
-
 		return FReply::Handled();
 	}
 
 	void SSceneOutliner::CreateFolder()
 	{
-		if (!RepresentingWorld)
+		if (!CheckWorld())
 		{
 			return;
 		}
 
+		UWorld& World = *SharedData->RepresentingWorld;
 		const FScopedTransaction Transaction(LOCTEXT("UndoAction_CreateFolder", "Create Folder"));
 
-		const FName NewFolderName = FActorFolders::Get().GetDefaultFolderNameForSelection(*RepresentingWorld);
-		FActorFolders::Get().CreateFolderContainingSelection(*RepresentingWorld, NewFolderName);
+		const FName NewFolderName = FActorFolders::Get().GetDefaultFolderNameForSelection(World);
+		FActorFolders::Get().CreateFolderContainingSelection(World, NewFolderName);
 
-		// Clear the current actor selection.
-		GEditor->SelectNone(false, true, true);
+		auto PreviouslySelectedItems = OutlinerTreeView->GetSelectedItems();
 
-		// Move any selected folders into the new folder name
-		for (auto& Item : GetSelectedItems())
+		auto Visit = [&](const FFolderTreeItem& Folder)
 		{
-			if (Item->Type == TOutlinerTreeItem::Folder)
-			{
-				auto Folder = StaticCastSharedPtr<TOutlinerFolderTreeItem>(Item).ToSharedRef();
-				MoveFolderTo(Folder, NewFolderName);
-			}
-		}
+			MoveFolderTo(Folder.Path, NewFolderName, World);
+		};
+		auto Visitor = FFunctionalVisitor().Folder(Visit);
 
-		// At this point the new folder will be in our newly added list, so select it and open a rename when it gets refreshed
-		NewItemActions.WhenCreated(NewFolderName, NewItemActuator::Select | NewItemActuator::Rename);
-	}
-	
-	void SSceneOutliner::CreateSubFolder(FName Parent)
-	{
-		if (!RepresentingWorld)
+		// Move any selected folders into the new folder name
+		for (const auto& Item : PreviouslySelectedItems)
 		{
-			return;
+			Item->Visit(Visitor);
 		}
 
-		const FScopedTransaction Transaction(LOCTEXT("UndoAction_CreateFolder", "Create Folder"));
-
-		const FName NewFolderName = FActorFolders::Get().GetDefaultFolderName(*RepresentingWorld, Parent);
-		FActorFolders::Get().CreateFolder(*RepresentingWorld, NewFolderName);
-
 		// At this point the new folder will be in our newly added list, so select it and open a rename when it gets refreshed
-		NewItemActions.WhenCreated(NewFolderName, NewItemActuator::Select | NewItemActuator::Rename);
+		NewItemActions.Add(NewFolderName, ENewItemAction::Select | ENewItemAction::Rename);
 	}
 
 	void SSceneOutliner::OnBroadcastFolderCreate(UWorld& InWorld, FName NewPath)
 	{
-		if (!ShouldShowFolders() || &InWorld != RepresentingWorld)
+		if (!ShouldShowFolders() || &InWorld != SharedData->RepresentingWorld)
 		{
 			return;
 		}
 
-		if (!FolderToTreeItemMap.Contains(NewPath))
+		if (!TreeItemMap.Contains(NewPath))
 		{
-			TSharedRef<TOutlinerFolderTreeItem> NewFolderItem = MakeShareable(new TOutlinerFolderTreeItem(NewPath));
-			AddedItemsList.Add(NewFolderItem);
-
-			FActorFolders::Get().SetSelectedFolderPath(NewPath);
-
-			Refresh();
+			ConstructItemFor<FFolderTreeItem>(NewPath);
 		}
 	}
 
-	void SSceneOutliner::DetachChildFromParent(TSharedRef<TOutlinerTreeItem> Child, FName ParentPath)
+	void SSceneOutliner::OnBroadcastFolderMove(UWorld& InWorld, FName OldPath, FName NewPath)
 	{
-		if (ParentPath.IsNone())
+		if (!ShouldShowFolders() || &InWorld != SharedData->RepresentingWorld)
 		{
-			RootTreeItems.Remove(Child);
+			return;
 		}
-		else
+
+		auto Item = TreeItemMap.FindRef(OldPath);
+		if (Item.IsValid())
 		{
-			auto* ParentFolder = FolderToTreeItemMap.Find(ParentPath);
-			if (ParentFolder)
-			{
-				ParentFolder->Get().ChildItems.Remove(Child);
-			}
+			// Remove it from the map under the old ID (which is derived from the folder path)
+			TreeItemMap.Remove(Item->GetID());
+
+			// Now change the path and put it back in the map with its new ID
+			auto Folder = StaticCastSharedPtr<FFolderTreeItem>(Item);
+			Folder->Path = NewPath;
+			Folder->LeafName = GetFolderLeafName(NewPath);
+
+			TreeItemMap.Add(Item->GetID(), Item);
+
+			// Add an operation to move the item in the hierarchy
+			PendingOperations.Emplace(FPendingTreeOperation::Moved, Item.ToSharedRef());
+			Refresh();
 		}
 	}
 
 	void SSceneOutliner::OnBroadcastFolderDelete(UWorld& InWorld, FName Path)
 	{
-		if (&InWorld != RepresentingWorld)
+		if (&InWorld != SharedData->RepresentingWorld)
 		{
 			return;
 		}
 
-		auto* Folder = FolderToTreeItemMap.Find(Path);
+		auto* Folder = TreeItemMap.Find(Path);
 		if (Folder)
 		{
-			RemovedItemsList.Add(*Folder);
+			PendingOperations.Emplace(FPendingTreeOperation::Removed, Folder->ToSharedRef());
 			Refresh();
 		}
 	}
 
-	void SSceneOutliner::RenameFolder(TSharedRef<TOutlinerFolderTreeItem> Folder)
+	void SSceneOutliner::ScrollItemIntoView(FTreeItemPtr Item)
 	{
-		// Ensure that the item we want to rename is visible in the tree
-		Folder->Flags.RenameWhenInView = true;
-		OutlinerTreeView->RequestScrollIntoView(Folder);
-	}
-
-	void SSceneOutliner::DeleteFolder(TSharedRef<TOutlinerFolderTreeItem> Folder)
-	{
-		const FScopedTransaction Transaction( LOCTEXT("DeleteFolder", "Delete Folder") );
-
-		auto Children = Folder->ChildItems;
-		for (auto& WeakChild : Children)
+		auto Parent = Item->GetParent();
+		while(Parent.IsValid())
 		{
-			auto Item = WeakChild.Pin();
-			if (!Item.IsValid())
-			{
-				continue;
-			}
-
-			if (Item->Type == TOutlinerTreeItem::Actor)
-			{
-				AActor* Actor = StaticCastSharedPtr<TOutlinerActorTreeItem>(Item)->Actor.Get();
-				if (Actor)
-				{
-					Actor->SetFolderPath(FName());
-				}
-			}
-			else
-			{
-				DeleteFolder(StaticCastSharedPtr<TOutlinerFolderTreeItem>(Item).ToSharedRef());
-			}
+			OutlinerTreeView->SetItemExpansion(Parent->AsShared(), true);
+			Parent = Parent->GetParent();
 		}
 
-		if (RepresentingWorld)
+		OutlinerTreeView->RequestScrollIntoView(Item);
+	}
+
+	void SSceneOutliner::InitiateRename(TSharedRef<ITreeItem> Item)
+	{
+		if (Item->CanInteract())
 		{
-			FActorFolders::Get().DeleteFolder(*RepresentingWorld, Folder->Path);
+			PendingRenameItem = Item;
+			ScrollItemIntoView(Item);
 		}
-
-		Refresh();
 	}
 
-	TSharedRef< ITableRow > SSceneOutliner::OnGenerateRowForOutlinerTree( FOutlinerTreeItemPtr Item, const TSharedRef< STableViewBase >& OwnerTable )
+	TSharedRef< ITableRow > SSceneOutliner::OnGenerateRowForOutlinerTree( FTreeItemPtr Item, const TSharedRef< STableViewBase >& OwnerTable )
 	{
-		return
-			SNew( SSceneOutlinerTreeRow, OwnerTable )
-				.SceneOutliner( SharedThis( this ) )
-				.Item( Item );
+		return SNew( SSceneOutlinerTreeRow, OutlinerTreeView.ToSharedRef(), SharedThis(this) ).Item( Item );
 	}
 
-
-	void SSceneOutliner::OnGetChildrenForOutlinerTree( FOutlinerTreeItemPtr InParent, TArray< FOutlinerTreeItemPtr >& OutChildren )
+	void SSceneOutliner::OnGetChildrenForOutlinerTree( FTreeItemPtr InParent, TArray< FTreeItemPtr >& OutChildren )
 	{
-		if( !InitOptions.bShowParentTree )
-		{
-			return;
-		}
-
-		if (InParent->Type == TOutlinerTreeItem::Actor)
-		{
-			auto Actor = StaticCastSharedPtr<TOutlinerActorTreeItem>(InParent)->Actor;
-			AActor* ParentActor = Actor.Get();
-			if( !ParentActor )
-			{
-				return;
-			}
-			// Does this actor have anything attached?
-			static TArray< AActor* > AttachedActors;
-			AttachedActors.Reset();
-			// @todo outliner: Perf: This call always will allocate heap memory and do a bit of work, even for actors with no children!
-			ParentActor->GetAttachedActors( AttachedActors );
-
-			if( AttachedActors.Num() == 0 )
-			{
-				return;
-			}
-
-			// Walk the list of attached actors and add their tree nodes to the output array
-			for( TArray< AActor* >::TConstIterator AttachedActorIt( AttachedActors ); AttachedActorIt; ++AttachedActorIt )
-			{
-				AActor* ChildActor = *AttachedActorIt;
-				if( !ensure( ChildActor != NULL ) )
-				{
-					continue;
-				}
-
-				// NOTE: We don't call Filters->PassesAllFilters() here because we may need to display branches
-				// in the tree for actors that don't pass our actor inclusion list, so that we can
-				// represent the heirarchy correctly.  We only call IsActorDisplayable() to avoid showing
-				// deleted actors and such.
-
-				// We MUST limit child actors displayed to those whose root component is attached to this actor
-				// Otherwise, the attached actors could be present elsewhere in the hierarchy, causing a crash
-				if( ChildActor->GetAttachParentActor() == ParentActor && IsActorDisplayable( ChildActor ) )
-				{
-					// NOTE: We don't need to apply filters here because we always repopulate the tree after
-					// the filter changes.  Besides, we may need to display child actors that were filtered out
-					// simply because they may have their own children that do match the filter.
-
-					// Find this actor in our list of tree items for each actor
-					auto* ChildTreeItem = ActorToTreeItemMap.Find( ChildActor );
-
-					// There's a chance we won't actually have a tree item for the child actor (e.g. because
-					// Outliner wasn't notified when an actor was destroyed or reparented), so we handle that here.
-					// In any other case though, we should always have an item for the actor in our list
-					if( ChildTreeItem )
-					{
-						OutChildren.Add( *ChildTreeItem );
-					}
-				}
-			}
-		}
-		else
-		{
-			// Folders
-			auto& ChildItems = StaticCastSharedPtr<TOutlinerFolderTreeItem>(InParent)->ChildItems;
-			for (auto ChildIt = ChildItems.CreateIterator(); ChildIt; ++ChildIt)
-			{
-				auto Child = ChildIt->Pin();
-				if (Child.IsValid())
-				{
-					OutChildren.Add( Child );
-				}
-			}
-
-		}
-
-		if (OutChildren.Num())
-		{
-			if (SortByColumn == ColumnID_ActorLabel)
-			{
-				if (SortMode == EColumnSortMode::Ascending)
-				{
-					OutChildren.Sort(Helpers::FSortByActorAscending());
-				}
-				else if (SortMode == EColumnSortMode::Descending)
-				{
-					OutChildren.Sort(Helpers::FSortByActorDescending());
-				}
-			}
-			else if (Gutter.IsValid() && SortByColumn == Gutter->GetColumnID())
+		if( SharedData->bShowParentTree )
+		{
+			for (auto& WeakChild : InParent->GetChildren())
 			{
-				Gutter->SortItems(OutChildren, SortMode);
+				auto Child = WeakChild.Pin();
+				// Should never have bogus entries in this list
+				check(Child.IsValid());
+				OutChildren.Add(Child);
 			}
-			else if (CustomColumn.IsValid())
+
+			// If the item needs it's children sorting, do that now
+			if (OutChildren.Num() && InParent->Flags.bChildrenRequireSort)
 			{
-				CustomColumn->SortItems(OutChildren, SortMode);
+				// Sort the children we returned
+				SortItems(OutChildren);
+
+				// Empty out the children and repopulate them in the correct order
+				InParent->Children.Empty();
+				for (auto& Child : OutChildren)
+				{
+					InParent->Children.Emplace(Child);
+				}
+				
+				// They no longer need sorting
+				InParent->Flags.bChildrenRequireSort = false;
 			}
 		}
 	}
 
-
 	bool SSceneOutliner::IsActorDisplayable( const AActor* Actor ) const
 	{
-		return (
-			!InitOptions.bOnlyShowFolders && 
-			Actor->IsEditable() &&			// Only show actors that are allowed to be selected and drawn in editor
-			Actor->IsListedInSceneOutliner() &&
-			!Actor->IsTemplate() &&			// Should never happen, but we never want CDOs displayed
-			!FActorEditorUtils::IsABuilderBrush(Actor) &&	// Don't show the builder brush
-			!Actor->IsA( AWorldSettings::StaticClass() ) &&		// Don't show the WorldSettings actor, even though it is technically editable
-			!Actor->IsPendingKill() &&						// We don't want to show actors that are about to go away
-			FLevelUtils::IsLevelVisible( Actor->GetLevel() ) &&					// Only show Actors whose level is visible
-			!FLevelEditorViewportClient::GetDropPreviewActors().Contains(Actor) );	// Don't display drag and drop preview actors (transient temp actors)
+		return	!SharedData->bOnlyShowFolders && 												// Don't show actors if we're only showing folders
+				Actor->IsEditable() &&															// Only show actors that are allowed to be selected and drawn in editor
+				Actor->IsListedInSceneOutliner() &&
+				!Actor->HasAnyFlags(RF_Transient) &&											// Don't show transient actors
+				!Actor->IsTemplate() &&															// Should never happen, but we never want CDOs displayed
+				!FActorEditorUtils::IsABuilderBrush(Actor) &&									// Don't show the builder brush
+				!Actor->IsA( AWorldSettings::StaticClass() ) &&									// Don't show the WorldSettings actor, even though it is technically editable
+				!Actor->IsPendingKill() &&														// We don't want to show actors that are about to go away
+				FLevelUtils::IsLevelVisible( Actor->GetLevel() );								// Only show Actors whose level is visible
 	}
 
-
-	void SSceneOutliner::OnOutlinerTreeSelectionChanged( FOutlinerTreeItemPtr TreeItem, ESelectInfo::Type SelectInfo )
+	void SSceneOutliner::OnOutlinerTreeSelectionChanged( FTreeItemPtr TreeItem, ESelectInfo::Type SelectInfo )
 	{
 		if (SelectInfo == ESelectInfo::Direct)
 		{
 			return;
 		}
 
-		if( InitOptions.Mode == ESceneOutlinerMode::ActorPicker )
+		if( SharedData->Mode == ESceneOutlinerMode::ActorPicker )
 		{
-			// In actor picking mode, we check to see if we have a selected actor, and if so, fire
-			// off the notification to whoever is listening.  This may often cause the widget itself
-			// to be enqueued for destruction
-			FOutlinerData SelectedItems = OutlinerTreeView->GetSelectedItems();
-			if( SelectedItems.Num() > 0 )
+			// In actor picking mode, we fire off the notification to whoever is listening.
+			// This may often cause the widget itself to be enqueued for destruction
+			if( OutlinerTreeView->GetNumItemsSelected() > 0 )
 			{
-				auto FirstItem = SelectedItems[0];
-				if (FirstItem->Type == TOutlinerTreeItem::Actor)
-				{
-					AActor* FirstSelectedActor = StaticCastSharedPtr<TOutlinerActorTreeItem>(FirstItem)->Actor.Get();
-					if( FirstSelectedActor != NULL )				
-					{
-						// Only fire the callback if this actor actually passes our required filters.  It's possible that
-						// some actors displayed are parent actors which don't actually match the requested criteria.
-						if( IsActorDisplayable( FirstSelectedActor ) && 
-							Filters->PassesAllFilters( FirstSelectedActor ) &&
-							SearchBoxFilter->PassesFilter( *FirstItem ) )
-						{
-							// Signal that an actor was selected
-							OnItemPicked.ExecuteIfBound( FirstItem.ToSharedRef() );
-						}
-					}
-				}
-				else
+				auto FirstItem = OutlinerTreeView->GetSelectedItems()[0];
+				if (FirstItem->CanInteract())
 				{
-					// A folder was picked
-					OnItemPicked.ExecuteIfBound( FirstItem.ToSharedRef() );
+					OnItemPicked.ExecuteIfBound( FirstItem->AsShared() );
 				}
 			}
 		}
 
 		// We only synchronize selection when in actor browsing mode
-		else if( InitOptions.Mode == ESceneOutlinerMode::ActorBrowsing )
+		else if( SharedData->Mode == ESceneOutlinerMode::ActorBrowsing )
 		{
 			if( !bIsReentrant )
 			{
 				TGuardValue<bool> ReentrantGuard(bIsReentrant,true);
 
-				// @todo outliner: Could be called frequently.  Use dirty bit?
 				// @todo outliner: Can be called from non-interactive selection
 
 				// The tree let us know that selection has changed, but wasn't able to tell us
@@ -2310,41 +1554,23 @@ namespace SceneOutliner
 				// selected actors to match the control's selection set.
 
 				// Make a list of all the actors that should now be selected in the world.
-				TArray< AActor* > ActorsToSelect;
-				TArray<FOutlinerTreeItemRef> FoldersToSelect;
-				const FOutlinerData SelectedItems = OutlinerTreeView->GetSelectedItems();
-				for( FOutlinerData::TConstIterator SelectedItemIt( SelectedItems ); SelectedItemIt; ++SelectedItemIt )
-				{
-					FOutlinerTreeItemPtr SelectedTreeItem = *SelectedItemIt;
-					if (SelectedTreeItem->Type == TOutlinerTreeItem::Actor)
-					{
-						AActor* Actor = StaticCastSharedPtr<TOutlinerActorTreeItem>(SelectedTreeItem)->Actor.Get();
-						if( Actor != NULL )
-						{
-							ActorsToSelect.Add( Actor );
-						}
-					}
-					else
-					{
-						FoldersToSelect.Add(SelectedTreeItem.ToSharedRef());
-					}
-				}
+				FItemSelection Selection(*OutlinerTreeView);
+				auto SelectedActors = TSet<AActor*>(Selection.GetActorPtrs());
 
 				bool bChanged = false;
-
-				for( int32 ActorIndex = 0; ActorIndex < ActorsToSelect.Num() && !bChanged; ActorIndex++ )
+				for (auto* Actor : SelectedActors)
 				{
-					AActor* Actor = ActorsToSelect[ ActorIndex ];
-					if ( !GEditor->GetSelectedActors()->IsSelected( Actor ) )
+					if (!GEditor->GetSelectedActors()->IsSelected(Actor))
 					{
-						// Actor has been selected
 						bChanged = true;
+						break;
 					}
 				}
-				for( FSelectionIterator SelectionIt( *GEditor->GetSelectedActors() ); SelectionIt && !bChanged; ++SelectionIt )
+
+				for (FSelectionIterator SelectionIt( *GEditor->GetSelectedActors() ); SelectionIt && !bChanged; ++SelectionIt)
 				{
 					AActor* Actor = CastChecked< AActor >( *SelectionIt );
-					if ( !ActorsToSelect.Contains( Actor ) )
+					if (!SelectedActors.Contains(Actor))
 					{
 						// Actor has been deselected
 						bChanged = true;
@@ -2356,9 +1582,9 @@ namespace SceneOutliner
 							TArray<AActor*> GroupActors;
 							DeselectedGroupActor->GetGroupActors( GroupActors );
 
-							for ( auto GroupActor : GroupActors )
+							for (auto* GroupActor : GroupActors)
 							{
-								ActorsToSelect.RemoveSingle( GroupActor );
+								SelectedActors.Remove(GroupActor);
 							}
 
 						}
@@ -2380,9 +1606,8 @@ namespace SceneOutliner
 					const bool bShouldSelect = true;
 					const bool bNotifyAfterSelect = false;
 					const bool bSelectEvenIfHidden = true;	// @todo outliner: Is this actually OK?
-					for( int32 ActorIndex = 0; ActorIndex < ActorsToSelect.Num(); ActorIndex++ )
+					for (auto* Actor : SelectedActors)
 					{
-						AActor* Actor = ActorsToSelect[ ActorIndex ];
 						UE_LOG(LogSceneOutliner, Verbose,  TEXT("Clicking on Actor (scene outliner): %s (%s)"), *Actor->GetClass()->GetName(), *Actor->GetActorLabel());
 						GEditor->SelectActor( Actor, bShouldSelect, bNotifyAfterSelect, bSelectEvenIfHidden );
 					}
@@ -2394,26 +1619,7 @@ namespace SceneOutliner
 					GEditor->NoteSelectionChange();
 				}
 
-				// Synchronize selection - actors may or may not have been selected depending on 
-				// the editor mode we are in (and their selectablity). Grouping may also have selected
-				// other actors.
-				OutlinerTreeView->ClearSelection();
-				for( FSelectionIterator SelectionIt( *GEditor->GetSelectedActors() ); SelectionIt; ++SelectionIt )
-				{
-					AActor* Actor = CastChecked< AActor >( *SelectionIt );
-					auto* SelectedTreeItem = ActorToTreeItemMap.Find( Actor );
-					if (SelectedTreeItem)
-					{
-						OutlinerTreeView->SetItemSelection(*SelectedTreeItem, true);
-					}
-				}
-				for (auto FolderIt = FoldersToSelect.CreateConstIterator(); FolderIt; ++FolderIt)
-				{
-					OutlinerTreeView->SetItemSelection(*FolderIt, true);
-				}
-
-				// Broadcast selection changed delegate
-				SelectionChanged.Broadcast();
+				bActorSelectionDirty = true;
 			}
 		}
 	}
@@ -2422,7 +1628,7 @@ namespace SceneOutliner
 	void SSceneOutliner::OnLevelSelectionChanged(UObject* Obj)
 	{
 		// We only synchronize selection when in actor browsing mode
-		if( InitOptions.Mode == ESceneOutlinerMode::ActorBrowsing )
+		if( SharedData->Mode == ESceneOutlinerMode::ActorBrowsing )
 		{
 			// @todo outliner: Because we are not notified of which items are being added/removed from selection, we have
 			// no immediate means to incrementally update the tree when selection changes.
@@ -2432,141 +1638,108 @@ namespace SceneOutliner
 			{
 				FullRefresh();
 			}
-
-			if( !bIsReentrant )
+			else if (!bIsReentrant)
 			{
-				TGuardValue<bool> ReentrantGuard(bIsReentrant, true);
-
-				// @todo outliner: Called twice quite frequently (once for deselect, once for select.)  Change to dirty bit pattern.
+				OutlinerTreeView->ClearSelection();
+				bActorSelectionDirty = true;
 
-				// The selected actors in the level have changed, so update our tree control's selection state
-				// to match the level.
+				// Scroll last item into view - this means if we are multi-selecting, we show newest selection. @TODO Not perfect though
+				if (AActor* LastSelectedActor = GEditor->GetSelectedActors()->GetBottom<AActor>())
 				{
-					OutlinerTreeView->ClearSelection();
-
-					ValidateOutlinerTreeView();
-
-					// Synchronize selection
-					FOutlinerTreeItemPtr LastSelectedItem;
-					for( FSelectionIterator SelectionIt( *GEditor->GetSelectedActors() ); SelectionIt; ++SelectionIt )
+					auto TreeItem = TreeItemMap.FindRef(LastSelectedActor);
+					if (TreeItem.IsValid())
 					{
-						AActor* Actor = CastChecked< AActor >( *SelectionIt );
-						auto* TreeItem = ActorToTreeItemMap.Find( Actor );
-						if( TreeItem )
+						if (!OutlinerTreeView->IsItemVisible(TreeItem))
 						{
-							OutlinerTreeView->SetItemSelection( *TreeItem, true );
-
-							LastSelectedItem = *TreeItem;
+							ScrollItemIntoView(TreeItem);
 						}
 					}
-
-					// Scroll last item into view - this means if we are multi-selecting, we show newest selection. @TODO Not perfect though
-					if(LastSelectedItem.IsValid() && !OutlinerTreeView->IsItemVisible( LastSelectedItem ))
+					else
 					{
-						OutlinerTreeView->RequestScrollIntoView(LastSelectedItem);
+						OnItemAdded(LastSelectedActor, ENewItemAction::ScrollIntoView);
 					}
 				}
-
 			}
 		}
 	}
 
-
-	void SSceneOutliner::ValidateOutlinerTreeView()
+	void SSceneOutliner::OnOutlinerTreeDoubleClick( FTreeItemPtr TreeItem )
 	{
-		bool bRefresh = false;
-		for( FActorToTreeItemMap::TConstIterator ActorIt( ActorToTreeItemMap ); ActorIt; ++ActorIt )
-		{
-			AActor* Actor = ActorIt.Key().Get();
-			if( Actor != NULL )
-			{
-				//Check whether actor is marked for deletion
-				if( Actor->HasAnyFlags(RF_PendingKill) )
-				{
-					bRefresh = true;
-					break;
-				}
-			}
-		}
-		if( bRefresh )
+		// We only deal with double clicks when in actor browsing mode
+		if( SharedData->Mode == ESceneOutlinerMode::ActorBrowsing )
 		{
-			Populate();
-		}
-	}
+			auto ExpandCollapseFolder = [&](const FFolderTreeItem& Folder){
+				auto Shared = const_cast<FFolderTreeItem&>(Folder).AsShared();
+				OutlinerTreeView->SetItemExpansion(Shared, !OutlinerTreeView->IsItemExpanded(Shared));
+			};
 
-	void SSceneOutliner::OnOutlinerTreeDoubleClick( FOutlinerTreeItemPtr TreeItem )
-	{
-		// We only modify viewport cameras when in actor browsing mode
-		if( InitOptions.Mode == ESceneOutlinerMode::ActorBrowsing )
-		{
-			// Move viewport cameras to the actor that was double-clicked on
-			const FOutlinerData SelectedItems = OutlinerTreeView->GetSelectedItems();
-			TArray< AActor* > SelectedActors;
-			for( FOutlinerData::TConstIterator SelectedItemIt( SelectedItems ); SelectedItemIt; ++SelectedItemIt )
+			if (TreeItem->CanInteract())
 			{
-				FOutlinerTreeItemPtr SelectedTreeItem = *SelectedItemIt;
-				if (SelectedTreeItem->Type == TOutlinerTreeItem::Actor)
-				{
-					AActor* Actor = StaticCastSharedPtr<TOutlinerActorTreeItem>(SelectedTreeItem)->Actor.Get();
-					if( Actor != NULL )
-					{
-						SelectedActors.Add( Actor );
-					}
-				}
-				else
-				{
-					OutlinerTreeView->SetItemExpansion(SelectedTreeItem, !OutlinerTreeView->IsItemExpanded(SelectedTreeItem));
-				}
-				
-			}
+				TreeItem->Visit(FFunctionalVisitor()
+
+					.Actor([&](const FActorTreeItem&){
+						// Move all actors into view
+						FItemSelection Selection(*OutlinerTreeView);
+						if( Selection.Actors.Num() > 0 )
+						{
+							const bool bActiveViewportOnly = false;
+							GEditor->MoveViewportCamerasToActor( Selection.GetActorPtrs(), bActiveViewportOnly );
+						}
+					})
 
-			if( SelectedActors.Num() > 0 )
+					.Folder(ExpandCollapseFolder)
+
+					.LevelBlueprint([](const FLevelBlueprintTreeItem& LevelBlueprintItem){
+						LevelBlueprintItem.Open();
+					})
+
+					.World([](const FWorldTreeItem& WorldItem){
+						WorldItem.OpenWorldSettings();
+					})
+
+				);
+
+			}
+			else
 			{
-				const bool bActiveViewportOnly = false;
-				GEditor->MoveViewportCamerasToActor( SelectedActors, bActiveViewportOnly );
+				TreeItem->Visit(FFunctionalVisitor()
+
+					.Folder(ExpandCollapseFolder)
+
+					.Actor([&](const FActorTreeItem& Item){
+						// Move just this actor into view
+						if (AActor* Actor = Item.Actor.Get())
+						{
+							const bool bActiveViewportOnly = false;
+							GEditor->MoveViewportCamerasToActor( *Actor, bActiveViewportOnly );
+						}
+					})
+				);
 			}
 		}
 	}
 
-	void SSceneOutliner::OnOutlinerTreeItemScrolledIntoView( FOutlinerTreeItemPtr TreeItem, const TSharedPtr<ITableRow>& Widget )
+	void SSceneOutliner::OnOutlinerTreeItemScrolledIntoView( FTreeItemPtr TreeItem, const TSharedPtr<ITableRow>& Widget )
 	{
-		auto FolderItem = StaticCastSharedPtr<TOutlinerFolderTreeItem>(TreeItem);
-		if (TreeItem->Flags.RenameWhenInView)
+		if (TreeItem == PendingRenameItem.Pin())
 		{
-			TreeItem->Flags.RenameWhenInView = false;
-			if(Widget.IsValid() && Widget->GetContent().IsValid())
-			{
-				TreeItem->RenameRequestEvent.ExecuteIfBound();
-			}
+			PendingRenameItem = nullptr;
+			TreeItem->RenameRequestEvent.ExecuteIfBound();
 		}
 	}
 	
-	void SSceneOutliner::OnItemExpansionChanged(FOutlinerTreeItemPtr TreeItem, bool bIsExpanded) const
+	void SSceneOutliner::OnItemExpansionChanged(FTreeItemPtr TreeItem, bool bIsExpanded) const
 	{
-		if (RepresentingWorld && TreeItem->Type == TOutlinerTreeItem::Folder)
-		{
-			// Update the central store of folder properties with this folder's new expansion state
-
-			auto FolderItem = StaticCastSharedPtr<TOutlinerFolderTreeItem>(TreeItem);
-
-			if (FActorFolderProps* Props = FActorFolders::Get().GetFolderProperties(*RepresentingWorld, FolderItem->Path))
-			{
-				Props->bIsExpanded = bIsExpanded;
-			}
+		TreeItem->Flags.bIsExpanded = bIsExpanded;
+		TreeItem->OnExpansionChanged();
 
-			// Expand any children that are also expanded
-			for (const auto& WeakChild : FolderItem->ChildItems)
+		// Expand any children that are also expanded
+		for (auto WeakChild : TreeItem->GetChildren())
+		{
+			auto Child = WeakChild.Pin();
+			if (Child->Flags.bIsExpanded)
 			{
-				auto Child = WeakChild.Pin();
-				if (Child.IsValid() && Child->Type == TOutlinerTreeItem::Folder)
-				{
-					auto ChildItem = StaticCastSharedPtr<TOutlinerFolderTreeItem>(Child);
-					const FActorFolderProps* FolderPtr = FActorFolders::Get().GetFolderProperties(*RepresentingWorld, ChildItem->Path);
-					if (FolderPtr && FolderPtr->bIsExpanded)
-					{
-						OutlinerTreeView->SetItemExpansion(ChildItem, true);
-					}
-				}
+				OutlinerTreeView->SetItemExpansion(Child, true);
 			}
 		}
 	}
@@ -2583,18 +1756,20 @@ namespace SceneOutliner
 
 	void SSceneOutliner::OnLevelActorsAdded(AActor* InActor)
 	{
-		// @todo outliner: Currently, attachment changes aren't causing us to re-populate
 		if( !bIsReentrant )
 		{
-			if( InActor && RepresentingWorld == InActor->GetWorld() )
+			if( InActor && SharedData->RepresentingWorld == InActor->GetWorld() && IsActorDisplayable(InActor) )
 			{
-				if(!ActorToTreeItemMap.Find(InActor))
+				if (!TreeItemMap.Find(InActor))
 				{
-					FOutlinerTreeItemRef NewItem = MakeShareable( new TOutlinerActorTreeItem((AActor*)InActor) );
-					AddedItemsList.AddUnique(NewItem);
-				}
+					// Update the total actor count that match the filters
+					if (Filters->PassesAllFilters(FActorTreeItem(InActor)))
+					{
+						ApplicableActors.Emplace(InActor);
+					}
 
-				Refresh();
+					ConstructItemFor<FActorTreeItem>(InActor);
+				}
 			}
 		}
 	}
@@ -2603,20 +1778,13 @@ namespace SceneOutliner
 	{
 		if( !bIsReentrant )
 		{
-			if( InActor && RepresentingWorld == InActor->GetWorld() )
+			if( InActor && SharedData->RepresentingWorld == InActor->GetWorld() )
 			{
-				auto* ItemPtr = ActorToTreeItemMap.Find(InActor);
-				if(ItemPtr)
-				{
-					RemovedItemsList.AddUnique(*ItemPtr);
+				ApplicableActors.Remove(InActor);
 
-					// Refresh the parent folder as well if possible
-					const FName Path = InActor->GetFolderPath();
-					auto* FolderItem = FolderToTreeItemMap.Find(Path);
-					if (FolderItem)
-					{
-						RefreshItemsList.AddUnique(*FolderItem);
-					}
+				if (auto* ItemPtr = TreeItemMap.Find(InActor))
+				{
+					PendingOperations.Emplace(FPendingTreeOperation::Removed, ItemPtr->ToSharedRef());
 					Refresh();
 				}
 			}
@@ -2628,12 +1796,11 @@ namespace SceneOutliner
 		// InActor can be equal to InParent in cases of components being attached internally. The Scene Outliner does not need to do anything in this case.
 		if( !bIsReentrant && InActor != InParent )
 		{
-			if( InActor && RepresentingWorld == InActor->GetWorld() )
+			if( InActor && SharedData->RepresentingWorld == InActor->GetWorld() )
 			{
-				auto* ItemPtr = ActorToTreeItemMap.Find(InActor);
-				if(ItemPtr)
+				if (auto* ItemPtr = TreeItemMap.Find(InActor))
 				{
-					AttachedActorsList.AddUnique(*ItemPtr);
+					PendingOperations.Emplace(FPendingTreeOperation::Moved, ItemPtr->ToSharedRef());
 					Refresh();
 				}
 			}
@@ -2645,13 +1812,11 @@ namespace SceneOutliner
 		// InActor can be equal to InParent in cases of components being attached internally. The Scene Outliner does not need to do anything in this case.
 		if( !bIsReentrant && InActor != InParent)
 		{
-			if( InActor && RepresentingWorld == InActor->GetWorld() )
+			if( InActor && SharedData->RepresentingWorld == InActor->GetWorld() )
 			{
-				auto* ItemPtr = ActorToTreeItemMap.Find(InActor);
-
-				if(ItemPtr)
+				if (auto* ItemPtr = TreeItemMap.Find(InActor))
 				{
-					RefreshItemsList.AddUnique(*ItemPtr);
+					PendingOperations.Emplace(FPendingTreeOperation::Moved, ItemPtr->ToSharedRef());
 					Refresh();
 				}
 				else
@@ -2660,66 +1825,34 @@ namespace SceneOutliner
 					OnLevelActorsAdded(InActor);
 				}
 			}
-
-			if( InParent && RepresentingWorld == InParent->GetWorld() )
-			{
-				// See if any of the remaining filtered actors still have this as a parent, 
-				// and if not remove it from the list too. If by chance its name matches the 
-				// current filter it'll get automatically re-added to the list.
-				bool bHasVisibileChild = false;
-				TArray<AActor*> AttachedActors;
-				InParent->GetAttachedActors(AttachedActors);
-				for (TArray< AActor* >::TConstIterator AttachedActorIt(AttachedActors); AttachedActorIt; ++AttachedActorIt)
-				{
-					const AActor* ChildActor = *AttachedActorIt;
-
-					if (ChildActor != NULL && ActorToTreeItemMap.Contains( ChildActor ) )
-					{
-						bHasVisibileChild = true;
-						break;
-					}
-				}
-				if ( !bHasVisibileChild )
-				{
-					// No children remaining, make sure this gets re-evaluated for the list too
-					auto* ItemPtr = ActorToTreeItemMap.Find(InParent);
-
-					if(ItemPtr)
-					{
-						RefreshItemsList.AddUnique(*ItemPtr);
-						Refresh();
-					}
-				}
-			}
 		}
 	}
 
 	/** Called by the engine when an actor's folder is changed */
 	void SSceneOutliner::OnLevelActorFolderChanged(const AActor* InActor, FName OldPath)
 	{
-		auto* ActorTreeItem = ActorToTreeItemMap.Find(InActor);
+		auto* ActorTreeItem = TreeItemMap.Find(InActor);
 		if (!ShouldShowFolders() || !InActor || !ActorTreeItem)
 		{
 			return;
 		}
 		
-		// Detach the item from its old parent while we know the old path
-		DetachChildFromParent(*ActorTreeItem, OldPath);
-
-		// Refresh the tree item
-		RefreshItemsList.AddUnique(*ActorTreeItem);
+		PendingOperations.Emplace(FPendingTreeOperation::Moved, ActorTreeItem->ToSharedRef());
 		Refresh();
 	}
 
 	void SSceneOutliner::OnLevelActorsRequestRename(const AActor* InActor)
 	{
-		TArray< TSharedPtr < TOutlinerTreeItem > > SelectedItems = OutlinerTreeView->GetSelectedItems();
+		auto SelectedItems = OutlinerTreeView->GetSelectedItems();
 		if( SelectedItems.Num() > 0)
 		{
 			// Ensure that the item we want to rename is visible in the tree
-			TSharedPtr<TOutlinerTreeItem> ItemToRename = SelectedItems[SelectedItems.Num() - 1];
-			ItemToRename->Flags.RenameWhenInView = true;
-			OutlinerTreeView->RequestScrollIntoView(ItemToRename);
+			FTreeItemPtr ItemToRename = SelectedItems[SelectedItems.Num() - 1];
+			if (ItemToRename->CanInteract())
+			{
+				PendingRenameItem = ItemToRename->AsShared();
+				ScrollItemIntoView(ItemToRename);
+			}
 		}
 	}
 
@@ -2752,29 +1885,26 @@ namespace SceneOutliner
 			return;
 		}
 		
-		auto* TreeItem = ActorToTreeItemMap.Find( ChangedActor );
-		if (TreeItem)
-		{
-			OutlinerTreeView->FlashHighlightOnItem(*TreeItem);
-		}
-	}
-
-	void SSceneOutliner::OnSetItemVisibility(FOutlinerTreeItemRef Item, bool bIsVisible)
-	{
-		// We operate on all the selected items if the specified item is selected
-		if (OutlinerTreeView->IsItemSelected(Item))
+		auto TreeItem = TreeItemMap.FindRef(ChangedActor);
+		if (TreeItem.IsValid())
 		{
-			for (auto& TreeItem : OutlinerTreeView->GetSelectedItems())
+			if (SearchBoxFilter->PassesFilter(*TreeItem))
 			{
-				TreeItem->SetIsVisible(bIsVisible);
+				OutlinerTreeView->FlashHighlightOnItem(TreeItem);
+				RequestSort();
+			}
+			else
+			{
+				// Do longer matches the filters, remove it
+				PendingOperations.Emplace(FPendingTreeOperation::Removed, TreeItem.ToSharedRef());
+				Refresh();
 			}
 		}
-		else
+		else if (IsActorDisplayable(ChangedActor))
 		{
-			Item->SetIsVisible(bIsVisible);
+			// Attempt to add the item if we didn't find it - perhaps it now matches the filter?
+			ConstructItemFor<FActorTreeItem>(ChangedActor);
 		}
-
-		GEditor->RedrawAllViewports();
 	}
 
 	void SSceneOutliner::OnFilterTextChanged( const FText& InFilterText )
@@ -2792,28 +1922,19 @@ namespace SceneOutliner
 			// Any text in the filter?  If not, we won't bother doing anything
 			if( !CurrentFilterText.IsEmpty() )
 			{
+				FItemSelection Selection;
+
 				// Gather all of the actors that match the filter text
-				TArray< AActor* > ActorsToSelect;
-				for( FActorToTreeItemMap::TConstIterator ActorIt( ActorToTreeItemMap ); ActorIt; ++ActorIt )
+				for (auto& Pair : TreeItemMap)
 				{
-					AActor* Actor = ActorIt.Key().Get();
-					if( Actor != NULL )
+					if (!Pair.Value->Flags.bIsFilteredOut)
 					{
-						// NOTE: The actor should have already passed this test earlier when we added it to the tree, but
-						// we'll run it again just to be safe, as the data might be stale (e.g. the actor might now be
-						// pending kill), in which case we don't want to try to select it
-						if( IsActorDisplayable( Actor ) )
-						{
-							if( Filters->PassesAllFilters( Actor ) && SearchBoxFilter->PassesFilter( *ActorIt.Value() ) )
-							{
-								ActorsToSelect.Add( Actor );
-							}
-						}
+						Pair.Value->Visit(Selection);
 					}
 				}
 
 				// We only select level actors when in actor browsing mode
-				if( InitOptions.Mode == ESceneOutlinerMode::ActorBrowsing )
+				if( SharedData->Mode == ESceneOutlinerMode::ActorBrowsing )
 				{
 					// Start batching selection changes
 					GEditor->GetSelectedActors()->BeginBatchSelectOperation();
@@ -2823,10 +1944,8 @@ namespace SceneOutliner
 					const bool bDeselectBSPSurfs = false;
 					const bool WarnAboutManyActors = true;
 					GEditor->SelectNone( bNoteSelectionChange, bDeselectBSPSurfs, WarnAboutManyActors );
-					for( TArray< AActor* >::TConstIterator ActorIt( ActorsToSelect ); ActorIt; ++ActorIt )
+					for (auto* Actor : Selection.GetActorPtrs())
 					{
-						AActor* Actor = *ActorIt;
-
 						const bool bShouldSelect = true;
 						const bool bSelectEvenIfHidden = false;
 						GEditor->SelectActor( Actor, bShouldSelect, bNoteSelectionChange, bSelectEvenIfHidden );
@@ -2840,42 +1959,26 @@ namespace SceneOutliner
 
 					// Set keyboard focus to the SceneOutliner, so the user can perform keyboard commands that interact
 					// with selected actors (such as Delete, to delete selected actors.)
-					{
-						// NOTE: Careful, GeneratePathToWidget can be reentrant in that it can call visibility delegates and such
-						FWidgetPath SceneOutlinerWidgetPath;
-						FSlateApplication::Get().GeneratePathToWidgetUnchecked( SharedThis( this ), SceneOutlinerWidgetPath );
-
-						// Set keyboard focus directly
-						FSlateApplication::Get().SetKeyboardFocus( SceneOutlinerWidgetPath, EFocusCause::SetDirectly );
-					}
+					SetKeyboardFocus();
 				}
 
 				// In 'actor picking' mode, we allow the user to commit their selection by pressing enter
 				// in the search window when a single actor is available
-				else if( InitOptions.Mode == ESceneOutlinerMode::ActorPicker )
+				else if( SharedData->Mode == ESceneOutlinerMode::ActorPicker )
 				{
 					// In actor picking mode, we check to see if we have a selected actor, and if so, fire
 					// off the notification to whoever is listening.  This may often cause the widget itself
 					// to be enqueued for destruction
-					if( ActorsToSelect.Num() == 1 )
+					if( Selection.Actors.Num() == 1 )
 					{
-						auto* ActorItem = ActorToTreeItemMap.Find(ActorsToSelect[0]);
-						if( ActorItem != NULL )				
-						{
-							// Signal that an actor was selected. We assume it is valid as it won't have been added to ActorsToSelect if not.
-							OnItemPicked.ExecuteIfBound( *ActorItem );
-						}
+						// Signal that an actor was selected. We assume it is valid as it won't have been added to ActorsToSelect if not.
+						OnItemPicked.ExecuteIfBound( Selection.Actors[0]->AsShared() );
 					}
 				}
 			}
 		}
 	}
 
-	EVisibility SSceneOutliner::GetSearchBoxVisibility() const
-	{
-		return InitOptions.bShowSearchBox ? EVisibility::Visible : EVisibility::Collapsed;
-	}
-
 	EVisibility SSceneOutliner::GetFilterStatusVisibility() const
 	{
 		return IsFilterActive() ? EVisibility::Visible : EVisibility::Collapsed;
@@ -2888,8 +1991,14 @@ namespace SceneOutliner
 
 	FString SSceneOutliner::GetFilterStatusText() const
 	{
-		const int32 SelectedActorCount = OutlinerTreeView->GetNumItemsSelected();
-		const int32 TotalActorCount = ActorToTreeItemMap.Num();
+		const int32 TotalActorCount = ApplicableActors.Num();
+
+		int32 SelectedActorCount = 0;
+		auto Count = [&](const FActorTreeItem&) { ++SelectedActorCount; };
+		for (const auto& Item : OutlinerTreeView->GetSelectedItems())
+		{
+			Item->Visit(FFunctionalVisitor().Actor(Count));
+		}
 
 		if ( !IsFilterActive() )
 		{
@@ -2937,8 +2046,7 @@ namespace SceneOutliner
 
 	bool SSceneOutliner::IsFilterActive() const
 	{
-		const int32 TotalActorCount = ActorToTreeItemMap.Num();
-		return FilterTextBoxWidget->GetText().ToString().Len() > 0 && TotalActorCount != FilteredActorCount;
+		return FilterTextBoxWidget->GetText().ToString().Len() > 0 && ApplicableActors.Num() != FilteredActorCount;
 	}
 
 	const FSlateBrush* SSceneOutliner::GetFilterButtonGlyph() const
@@ -2959,16 +2067,29 @@ namespace SceneOutliner
 
 	}
 
-	FText SSceneOutliner::GetFilterHighlightText() const
+	TAttribute<FText> SSceneOutliner::GetFilterHighlightText() const
 	{
-		return SearchBoxFilter->GetRawFilterText();
+		return TAttribute<FText>::Create(TAttribute<FText>::FGetter::CreateStatic([](TWeakPtr<TreeItemTextFilter> Filter){
+			auto FilterPtr = Filter.Pin();
+			return FilterPtr.IsValid() ? FilterPtr->GetRawFilterText() : FText();
+		}, TWeakPtr<TreeItemTextFilter>(SearchBoxFilter)));
 	}
 
+	void SSceneOutliner::SetKeyboardFocus()
+	{
+		if (SupportsKeyboardFocus())
+		{
+			FWidgetPath OutlinerTreeViewWidgetPath;
+			// NOTE: Careful, GeneratePathToWidget can be reentrant in that it can call visibility delegates and such
+			FSlateApplication::Get().GeneratePathToWidgetUnchecked( OutlinerTreeView.ToSharedRef(), OutlinerTreeViewWidgetPath );
+			FSlateApplication::Get().SetKeyboardFocus( OutlinerTreeViewWidgetPath, EFocusCause::SetDirectly );
+		}
+	}
 
 	bool SSceneOutliner::SupportsKeyboardFocus() const
 	{
 		// We only need to support keyboard focus if we're in actor browsing mode
-		if( InitOptions.Mode == ESceneOutlinerMode::ActorBrowsing )
+		if( SharedData->Mode == ESceneOutlinerMode::ActorBrowsing )
 		{
 			// Scene outliner needs keyboard focus so the user can press keys to activate commands, such as the Delete
 			// key to delete selected actors
@@ -2983,37 +2104,42 @@ namespace SceneOutliner
 		// @todo outliner: Use command system for these for discoverability? (allow bindings?)
 
 		// We only allow these operations in actor browsing mode
-		if( InitOptions.Mode == ESceneOutlinerMode::ActorBrowsing )
+		if( SharedData->Mode == ESceneOutlinerMode::ActorBrowsing )
 		{
 			// Delete key: Delete selected actors (not rebindable, because it doesn't make much sense to bind.)
 			if( InKeyEvent.GetKey() == EKeys::F2 )
 			{
-				auto SelectedItems = OutlinerTreeView->GetSelectedItems();
-				if (SelectedItems.Num() == 1 && SelectedItems[0]->Type == TOutlinerTreeItem::Folder)
+				if (OutlinerTreeView->GetNumItemsSelected() == 1)
 				{
-					SelectedItems[0]->RenameRequestEvent.ExecuteIfBound();
+					FTreeItemPtr ItemToRename = OutlinerTreeView->GetSelectedItems()[0];
+					
+					if (ItemToRename->CanInteract())
+					{
+						PendingRenameItem = ItemToRename->AsShared();
+						ScrollItemIntoView(ItemToRename);
+					}
+
 					return FReply::Handled();
 				}
 			}
-			else if ( InKeyEvent.GetKey() == EKeys::Platform_Delete )
+
+			// F5 forces a full refresh
+			else if ( InKeyEvent.GetKey() == EKeys::F5 )
 			{
-				if( InitOptions.CustomDelete.IsBound() )
-				{
-					TArray< TWeakObjectPtr< AActor > > SelectedActors;
+				FullRefresh();
+				return FReply::Handled();
+			}
 
-					for ( FSelectionIterator It( GEditor->GetSelectedActorIterator() ) ; It ; ++It )
-					{
-						UObject* Object = *It;
-						if( Object->IsA( AActor::StaticClass() ) )
-						{
-							AActor* Actor = static_cast<AActor*>( *It );
-							SelectedActors.Add( Actor );
-						}
-					}
+			// Delete key: Delete selected actors (not rebindable, because it doesn't make much sense to bind.)
+			else if ( InKeyEvent.GetKey() == EKeys::Platform_Delete )
+			{
+				const FItemSelection Selection(*OutlinerTreeView);
 
-					InitOptions.CustomDelete.Execute( SelectedActors );
+				if( SharedData->CustomDelete.IsBound() )
+				{
+					SharedData->CustomDelete.Execute( Selection.GetWeakActors() );
 				}
-				else
+				else if (CheckWorld())
 				{
 					const FScopedTransaction Transaction( LOCTEXT("UndoAction_DeleteSelection", "Delete selection") );
 
@@ -3021,21 +2147,15 @@ namespace SceneOutliner
 					auto SelectedItems = OutlinerTreeView->GetSelectedItems();
 					
 					GEditor->SelectNone(true, true);
-					for (auto ItemIt = SelectedItems.CreateIterator(); ItemIt; ++ItemIt)
+
+					for (auto* Folder : Selection.Folders)
 					{
-						auto Item = (*ItemIt);
-						if (Item->Type == TOutlinerTreeItem::Folder)
-						{
-							DeleteFolder(StaticCastSharedPtr<TOutlinerFolderTreeItem>(Item).ToSharedRef());
-						}
-						else
-						{
-							auto* Actor = StaticCastSharedPtr<TOutlinerActorTreeItem>(Item)->Actor.Get();
-							if (Actor)
-							{
-								GEditor->SelectActor(Actor, true, false);
-							}
-						}
+						Folder->Delete();
+					}
+
+					for (auto* Actor : Selection.GetActorPtrs())
+					{
+						GEditor->SelectActor(Actor, true, false);
 					}
 
 					// Code from FLevelEditorActionCallbacks::Delete_CanExecute()
@@ -3055,10 +2175,9 @@ namespace SceneOutliner
 						}
 					}
 
-					UWorld* World = GWorld;
-					if (GUnrealEd->CanDeleteSelectedActors( World, true, false ))
+					if (GUnrealEd->CanDeleteSelectedActors( SharedData->RepresentingWorld, true, false ))
 					{
-						GEditor->edactDeleteSelected( World );
+						GEditor->edactDeleteSelected( SharedData->RepresentingWorld );
 					}
 				}
 
@@ -3069,12 +2188,46 @@ namespace SceneOutliner
 		return FReply::Unhandled();
 	}
 
+	void SSceneOutliner::SynchronizeActorSelection()
+	{
+		TGuardValue<bool> ReentrantGuard(bIsReentrant, true);
+
+		auto* SelectedActors = GEditor->GetSelectedActors();
+
+		// Deselect actors in the tree that are no longer selected in the world
+		FItemSelection Selection(*OutlinerTreeView);
+		for (auto* ActorItem : Selection.Actors)
+		{
+			if (!SelectedActors->IsSelected(ActorItem->Actor.Get()))
+			{
+				OutlinerTreeView->SetItemSelection(ActorItem->AsShared(), false);
+			}
+		}
+
+		// Ensure that all selected actors in the world are selected in the tree
+		for (FSelectionIterator SelectionIt( *SelectedActors ); SelectionIt; ++SelectionIt)
+		{
+			AActor* Actor = CastChecked< AActor >(*SelectionIt);
+			if (auto* ActorItem = TreeItemMap.Find(Actor))
+			{
+				OutlinerTreeView->SetItemSelection(*ActorItem, true);
+			}
+		}
+
+		// Broadcast selection changed delegate
+		SelectionChanged.Broadcast();
+	}
 
 	void SSceneOutliner::Tick( const FGeometry& AllottedGeometry, const double InCurrentTime, const float InDeltaTime )
 	{
 		// Call parent implementation
 		SCompoundWidget::Tick( AllottedGeometry, InCurrentTime, InDeltaTime );
 
+		for (auto& Pair : Columns)
+		{
+			Pair.Value->Tick(InCurrentTime, InDeltaTime);
+		}
+
 		if ( bPendingFocusNextFrame && FilterTextBoxWidget->GetVisibility() == EVisibility::Visible )
 		{
 			FWidgetPath WidgetToFocusPath;
@@ -3090,27 +2243,51 @@ namespace SceneOutliner
 				Populate();
 			}
 		}
+		SortOutlinerTimer -= InDeltaTime;
 
-		if( SortOutlinerTimer >= SCENE_OUTLINER_RESORT_TIMER )
+		if (bSortDirty && (!SharedData->bRepresentingPlayWorld || SortOutlinerTimer <= 0))
 		{
-			SortTree();
+			SortItems(RootTreeItems);
+			for (const auto& Pair : TreeItemMap)
+			{
+				Pair.Value->Flags.bChildrenRequireSort = true;
+			}
+
+			OutlinerTreeView->RequestTreeRefresh();
+			bSortDirty = false;
+		}
+
+		if (SortOutlinerTimer <= 0)
+		{
+			SortOutlinerTimer = SCENE_OUTLINER_RESORT_TIMER;
+		}
+
+
+		if (bActorSelectionDirty)
+		{
+			SynchronizeActorSelection();
+			bActorSelectionDirty = false;
 		}
-		SortOutlinerTimer -= InDeltaTime;
 	}
 
 	EColumnSortMode::Type SSceneOutliner::GetColumnSortMode( const FName ColumnId ) const
 	{
-		if ( SortByColumn != ColumnId || (SortByColumn == CustomColumn->GetColumnID() && !CustomColumn->SupportsSorting()) )
+		if (SortByColumn == ColumnId)
 		{
-			return EColumnSortMode::None;
+			auto Column = Columns.FindRef(ColumnId);
+			if (Column.IsValid() && Column->SupportsSorting())
+			{
+				return SortMode;
+			}
 		}
 
-		return SortMode;
+		return EColumnSortMode::None;
 	}
 
 	void SSceneOutliner::OnColumnSortModeChanged( const EColumnSortPriority::Type SortPriority, const FName& ColumnId, const EColumnSortMode::Type InSortMode )
 	{
-		if (CustomColumn.IsValid() && ColumnId == CustomColumn->GetColumnID() && !CustomColumn->SupportsSorting())
+		auto Column = Columns.FindRef(ColumnId);
+		if (!Column.IsValid() || !Column->SupportsSorting())
 		{
 			return;
 		}
@@ -3123,105 +2300,40 @@ namespace SceneOutliner
 
 	void SSceneOutliner::RequestSort()
 	{
-		if (bRepresentingPlayWorld)
-		{
-			if (SortOutlinerTimer < 0.0f || OutlinerTreeView->HasKeyboardFocus())
-			{
-				SortOutlinerTimer = SCENE_OUTLINER_RESORT_TIMER;
-			}
-		}
-		else
-		{
-			// Sort the list of root items
-			SortTree();
-		}
-
-		OutlinerTreeView->RequestTreeRefresh();
+		bSortDirty = true;
 	}
 
-	void SSceneOutliner::SortTree()
+	void SSceneOutliner::SortItems(TArray<FTreeItemPtr>& Items) const
 	{
-		if (SortByColumn == ColumnID_ActorLabel)
+		auto Column = Columns.FindRef(SortByColumn);
+		if (Column.IsValid())
 		{
-			if (SortMode == EColumnSortMode::Ascending)
-			{
-				RootTreeItems.Sort(Helpers::FSortByActorAscending());
-			}
-			else if (SortMode == EColumnSortMode::Descending)
-			{
-				RootTreeItems.Sort(Helpers::FSortByActorDescending());
-			}
-		}
-		else if (Gutter.IsValid() && SortByColumn == Gutter->GetColumnID())
-		{
-			Gutter->SortItems(RootTreeItems, SortMode);
-		}
-		else if (CustomColumn.IsValid())
-		{
-			CustomColumn->SortItems(RootTreeItems, SortMode);
+			Column->SortItems(Items, SortMode);
 		}
 	}
-	
-	void NewItemActuator::Empty()
-	{
-		PendingFolderActions.Empty();
-		PendingActorActions.Empty();
-	}
 
-	void NewItemActuator::WhenCreated(FName FolderPath, uint8 InActionMask)
+	void SSceneOutliner::FOnItemAddedToTree::Visit(FActorTreeItem& ActorItem) const
 	{
-		const FolderAction NewAction = { InActionMask, FolderPath };
-		PendingFolderActions.Add(NewAction);
-	}
+		Outliner.FilteredActorCount += ActorItem.Flags.bIsFilteredOut ? 0 : 1;
 
-	void NewItemActuator::WhenCreated(FOutlinerTreeItemRef Item, uint8 InActionMask)
-	{
-		if (Item->Type == TOutlinerTreeItem::Actor)
-		{
-			const ActorAction NewAction = { InActionMask, StaticCastSharedRef<TOutlinerActorTreeItem>(Item) };
-			PendingActorActions.Add(NewAction);
-		}
-		else
+		// Synchronize selection
+		if (Outliner.SharedData->Mode == ESceneOutlinerMode::ActorBrowsing && GEditor->GetSelectedActors()->IsSelected(ActorItem.Actor.Get()))
 		{
-			WhenCreated(StaticCastSharedRef<TOutlinerFolderTreeItem>(Item)->Path, InActionMask);
+			// Have to const cast here as the tree view is templated on non-const ptrs
+			Outliner.OutlinerTreeView->SetItemSelection(ActorItem.AsShared(), true);
 		}
 	}
 
-	void NewItemActuator::ItemHasBeenCreated(TSharedRef<SOutlinerTreeView> OutlinerTreeView, TSharedRef<TOutlinerActorTreeItem> NewItem)
+	void SSceneOutliner::FOnItemAddedToTree::Visit(FFolderTreeItem& Folder) const
 	{
-		for (int32 Index = 0; Index < PendingActorActions.Num(); )
+		if (!Outliner.SharedData->RepresentingWorld)
 		{
-			auto PinnedPendingItem = PendingActorActions[Index].ActorItem.Pin();
-			if (!PinnedPendingItem.IsValid())
-			{
-				PendingActorActions.RemoveAt(Index);
-			}
-			else if(PinnedPendingItem == NewItem)
-			{
-				PerformActionOnItem(OutlinerTreeView, NewItem, PendingActorActions[Index].Actions);
-
-				PendingActorActions.RemoveAt(Index);
-			}
-			else
-			{
-				++Index;
-			}
+			return;
 		}
-	}
 
-	void NewItemActuator::ItemHasBeenCreated(TSharedRef<SOutlinerTreeView> OutlinerTreeView, TSharedRef<TOutlinerFolderTreeItem> NewItem)
-	{
-		for (int32 Index = 0; Index < PendingFolderActions.Num(); )
+		if (FActorFolderProps* Props = FActorFolders::Get().GetFolderProperties(*Outliner.SharedData->RepresentingWorld, Folder.Path))
 		{
-			if (PendingFolderActions[Index].FolderPath == NewItem->Path)
-			{
-				PerformActionOnItem(OutlinerTreeView, NewItem, PendingFolderActions[Index].Actions);
-				PendingFolderActions.RemoveAt(Index);
-			}
-			else
-			{
-				++Index;
-			}
+			Folder.Flags.bIsExpanded = Props->bIsExpanded;
 		}
 	}
 
