Added Worlds and Level Blueprints to the scene outliner

diff --git a/Engine/Source/Editor/SceneOutliner/Private/SSceneOutliner.h b/Engine/Source/Editor/SceneOutliner/Private/SSceneOutliner.h
--- a/Engine/Source/Editor/SceneOutliner/Private/SSceneOutliner.h
+++ b/Engine/Source/Editor/SceneOutliner/Private/SSceneOutliner.h
@@ -2,27 +2,29 @@
 #pragma once
 
 #include "SOutlinerTreeView.h"
+#include "SceneOutlinerStandaloneTypes.h"
 
 namespace SceneOutliner
 {
-	struct TOutlinerActorTreeItem;
-	struct TOutlinerFolderTreeItem;
+	typedef TTextFilter< const ITreeItem& > TreeItemTextFilter;
 
-	typedef TArray< FOutlinerTreeItemPtr > FOutlinerData;
-
-	typedef TTextFilter< const TOutlinerTreeItem& > TreeItemTextFilter;
+	/** Structure that defines an operation that should be applied to the tree */
+	struct FPendingTreeOperation
+	{
+		enum EType { Added, Removed, Moved };
+		FPendingTreeOperation(EType InType, TSharedRef<ITreeItem> InItem) : Type(InType), Item(InItem) { }
 
-	typedef TMap< TWeakObjectPtr< AActor >, TSharedRef<TOutlinerActorTreeItem> > FActorToTreeItemMap;
-	typedef TMap< FName, TSharedRef<TOutlinerFolderTreeItem> > FFolderToTreeItemMap;
-	typedef TSet< const AActor* > FParentActorsSet;
+		/** The type of operation that is to be applied */
+		EType Type;
 
-	FText GetLabelForItem( const TSharedRef<TOutlinerTreeItem> TreeItem );
-	FString GetLabelForItemAsString(const TSharedRef<TOutlinerTreeItem> TreeItem);
+		/** The tree item to which this operation relates */
+		FTreeItemRef Item;
+	};
 
-	/** An array of actions to apply to newly added items in the scene outliner */
-	struct NewItemActuator
+	/** Set of actions to apply to new tree items */
+	namespace ENewItemAction 
 	{
-		enum ActionType
+		enum Type
 		{
 			/** Select the item when it is created */
 			Select			= 1 << 0,
@@ -31,34 +33,7 @@ namespace SceneOutliner
 			/** Interactively rename the item when it is created (implies the above) */
 			Rename			= 1 << 2,
 		};
-
-		/** Empty this actuator (called after every refresh of the tree) */
-		void Empty();
-
-		/** Specify an action to perform on the specified actor or folder when it is created (see ActionType enum) */
-		void WhenCreated(FName FolderPath, uint8 InActionMask);
-		void WhenCreated(FOutlinerTreeItemRef, uint8 InActionMask);
-
-		/** Called by the scene outliner when an item has been added to the tree */
-		void ItemHasBeenCreated(TSharedRef<SOutlinerTreeView> OutlinerTreeView, TSharedRef<TOutlinerActorTreeItem> NewItem);
-		void ItemHasBeenCreated(TSharedRef<SOutlinerTreeView> OutlinerTreeView, TSharedRef<TOutlinerFolderTreeItem> NewItem);
-
-	private:
-
-		struct ActorAction
-		{	
-			uint8 Actions;
-			TWeakPtr<TOutlinerActorTreeItem> ActorItem;
-		};
-		struct FolderAction
-		{	
-			uint8 Actions;
-			FName FolderPath;
-		};
-
-		TArray<ActorAction>		PendingActorActions;
-		TArray<FolderAction>	PendingFolderActions;
-	};
+	}
 
 	/**
 	 * Scene Outliner widget
@@ -80,19 +55,18 @@ namespace SceneOutliner
 		 * @param	InArgs		Declaration used by the SNew() macro to construct this widget
 		 * @param	InitOptions	Programmer-driven initialization options for this widget
 		 */
-		void Construct( const FArguments& InArgs, const FSceneOutlinerInitializationOptions& InitOptions );
+		void Construct( const FArguments& InArgs, const FInitializationOptions& InitOptions );
+
+		/** Default constructor - initializes data that is shared between all tree items */
+		SSceneOutliner() : SharedData(MakeShareable( new FSharedOutlinerData )) {}
 
 		/** SSceneOutliner destructor */
 		~SSceneOutliner();
 
-		/** Called by our list to generate a widget that represents the specified item at the specified column in the tree */
-		TSharedRef< SWidget > GenerateWidgetForItemAndColumn( FOutlinerTreeItemPtr Item, const FName ColumnID, FIsSelected InIsSelected ) const;
-
-		/** Ensure an actor node in the tree is expanded */
-		void ExpandActor(AActor* Actor);
-
 		/** SWidget interface */
 		virtual void Tick( const FGeometry& AllottedGeometry, const double InCurrentTime, const float InDeltaTime ) override;
+		virtual bool SupportsKeyboardFocus() const override;
+		virtual FReply OnKeyDown( const FGeometry& MyGeometry, const FKeyEvent& InKeyEvent ) override;
 
 		/**	Broadcasts whenever the current selection changes */
 		FSimpleMulticastDelegate SelectionChanged;
@@ -105,82 +79,159 @@ namespace SceneOutliner
 		virtual void PostRedo(bool bSuccess) override { PostUndo(bSuccess); }
 		// End of FEditorUndoClient
 
-		/** Get an array of the currently selected items in this tree */
-		TArray<FOutlinerTreeItemPtr> GetSelectedItems();
-
-		/** const public access to this outliner's init options */
-		inline const FSceneOutlinerInitializationOptions& GetInitOptions() const
+		/** @return Returns the common data for this outliner */
+		virtual const FSharedOutlinerData& GetSharedData() const override
 		{
-			return InitOptions;
+			return *SharedData;
 		}
 
-		/** Find a folder by its path */
-		const TSharedPtr<TOutlinerFolderTreeItem> FindFolderByPath(FName Path) const;
+		/** Get a const reference to the actual tree hierarchy */
+		virtual const STreeView<FTreeItemPtr>& GetTree() const override
+		{
+			return *OutlinerTreeView;
+		}
 
-		/** Move the specified folder to a new parent path, returning the new full path */
-		FName MoveFolderTo(TSharedRef<TOutlinerFolderTreeItem> Folder, FName NewParent);
+		/** @return Returns a string to use for highlighting results in the outliner list */
+		virtual TAttribute<FText> GetFilterHighlightText() const override;
 
+		/** Set the keyboard focus to the outliner */
+		virtual void SetKeyboardFocus() override;
+		
 	private:
-
-		/** Populates our data set */
-		void Populate();
+		/** Methods that implement structural modification logic for the tree */
 
 		/** Empty all the tree item containers maintained by this outliner */
 		void EmptyTreeItems();
 
-		/** Populates OutSearchStrings with the strings associated with TreeItem that should be used in searching */
-		void PopulateSearchStrings( const TOutlinerTreeItem& TreeItem, OUT TArray< FString >& OutSearchStrings ) const;
+		/** Apply incremental changes to, or a complete repopulation of the tree  */
+		void Populate();
+
+		/** Repopulates the entire tree */
+		void RepopulateEntireTree();
 
 		/** Tells the scene outliner that it should do a full refresh, which will clear the entire tree and rebuild it from scratch. */
 		void FullRefresh();
 
-		/** Check whether we should be showing folders or not in this scene outliner */
-		bool ShouldShowFolders() const;
+		/** Attempts to add an item to the tree. Will add any parents if required. */
+		bool AddItemToTree(FTreeItemRef InItem);
+		
+		/** Add an item to the tree, even if it doesn't match the filter terms. Used to add parent's that would otherwise be filtered out */
+		void AddUnfilteredItemToTree(FTreeItemRef Item);
+
+		/** Ensure that the specified item's parent is added to the tree, if applicable */
+		FTreeItemPtr EnsureParentForItem(FTreeItemRef Item);
+
+		/** Remove the specified item from the tree */
+		void RemoveItemFromTree(FTreeItemRef InItem);
+
+		/** Called when a child has been removed from the specified parent. Will potentially remove the parent from the tree */
+		void OnChildRemovedFromParent(ITreeItem& Parent);
+
+		/** Called when a child has been moved in the tree hierarchy */
+		void OnItemMoved(const FTreeItemRef& Item);
+
+		/** Adds a new item for the specified type and refreshes the tree, provided it matches the filter terms */
+		template<typename TreeItemType, typename DataType>
+		void ConstructItemFor(const DataType& Data)
+		{
+			// We test the filters with a temporary so we don't allocate on the heap unnecessarily
+			const TreeItemType Temporary(Data);
+			if (Filters->PassesAllFilters(Temporary) && SearchBoxFilter->PassesFilter(Temporary))
+			{
+				FTreeItemRef NewItem = MakeShareable(new TreeItemType(Data));
+				PendingOperations.Emplace(FPendingTreeOperation::Added, NewItem);
+				Refresh();
+			}
+		}
+
+		/** Visitor that is used to set up type-specific data after tree items are added to the tree */
+		struct FOnItemAddedToTree : IMutableTreeItemVisitor
+		{
+			SSceneOutliner& Outliner;
+			FOnItemAddedToTree(SSceneOutliner& InOutliner) : Outliner(InOutliner) {}
+
+			virtual void Visit(FActorTreeItem& Actor) const override;
+			virtual void Visit(FFolderTreeItem& Folder) const override;
+		};
+
+		/** Friendship required so the visitor can access our guts */
+		friend FOnItemAddedToTree;
+
+	public:
+
+		/** Instruct the outliner to perform an action on the specified item when it is created */
+		void OnItemAdded(const FTreeItemID& ItemID, uint8 ActionMask);
+
+		/** Get the columns to be displayed in this outliner */
+		const TMap<FName, TSharedPtr<ISceneOutlinerColumn>>& GetColumns() const
+		{
+			return Columns;
+		}
+
+	private:
+
+		/** Map of columns that are shown on this outliner. */
+		TMap<FName, TSharedPtr<ISceneOutlinerColumn>> Columns;
+
+		/** Set up the columns required for this outliner */
+		void SetupColumns(SHeaderRow& HeaderRow);
+
+		/** Populates OutSearchStrings with the strings associated with TreeItem that should be used in searching */
+		void PopulateSearchStrings( const ITreeItem& TreeItem, OUT TArray< FString >& OutSearchStrings ) const;
+
+
+	public:
+		/** Miscellaneous helper functions */
 
-		/** Synchronize (and scroll into view) the current selection to the specified actor item, if the corresponding actor is selected */
-		void SynchronizeSelectionTo(TSharedRef<TOutlinerActorTreeItem> ActorItem);
+		/** Scroll the specified item into view */
+		void ScrollItemIntoView(FTreeItemPtr Item);
 
-		/** Is actor class name text visible? */
-		EVisibility IsActorClassNameVisible() const;
+		/** Open a rename for the specified tree Item */
+		void InitiateRename(TSharedRef<ITreeItem> TrreeItem);
 
-		/** Gets the class name for a tree item */
-		FString GetClassNameForItem( FOutlinerTreeItemPtr TreeItem ) const;
+	private:
+
+		/** Synchronize the current actor selection in the world, to the tree */
+		void SynchronizeActorSelection();
 
-		/** Gets text for the specified item to display in the custom column for the outliner tree */
-		FString GetCustomColumnTextForItem( FOutlinerTreeItemPtr TreeItem ) const;
+		/** Check that we are reflecting a valid world */
+		bool CheckWorld() const { return SharedData->RepresentingWorld != nullptr; }
 
-		/** Gets the color to draw item labels */
-		FSlateColor GetColorAndOpacityForItem( FOutlinerTreeItemPtr TreeItem ) const;
+		/** Check whether we should be showing folders or not in this scene outliner */
+		bool ShouldShowFolders() const;
 
-		/** Gets the icon for the specified icon */
-		const FSlateBrush* GetIconForItem(FOutlinerTreeItemRef Item) const;
+		/** Get an array of selected folders */
+		TArray<FFolderTreeItem*> GetSelectedFolders() const;
 
-		/** Gets the tool tip text for the item icon */
-		FText GetToolTipTextForItemIcon( FOutlinerTreeItemPtr TreeItem ) const;
+		/** Checks to see if the actor is valid for displaying in the outliner */
+		bool IsActorDisplayable( const AActor* Actor ) const;
 
-		/** Gets the brush to draw the component mobility icon */
-		const FSlateBrush* GetBrushForComponentMobilityIcon( FOutlinerTreeItemPtr TreeItem ) const;
+		/** @return	Returns true if the filter is currently active */
+		bool IsFilterActive() const;
 
-		/** Gets the color to draw custom column text */
-		FLinearColor GetColorAndOpacityForCustomColumn( FOutlinerTreeItemPtr TreeItem ) const;
+	private:
+		/** Tree view event bindings */
 
 		/** Called by STreeView to generate a table row for the specified item */
-		TSharedRef< ITableRow > OnGenerateRowForOutlinerTree( FOutlinerTreeItemPtr Item, const TSharedRef< STableViewBase >& OwnerTable );
+		TSharedRef< ITableRow > OnGenerateRowForOutlinerTree( FTreeItemPtr Item, const TSharedRef< STableViewBase >& OwnerTable );
 
 		/** Called by STreeView to get child items for the specified parent item */
-		void OnGetChildrenForOutlinerTree( FOutlinerTreeItemPtr InParent, TArray< FOutlinerTreeItemPtr >& OutChildren );
+		void OnGetChildrenForOutlinerTree( FTreeItemPtr InParent, TArray< FTreeItemPtr >& OutChildren );
 
 		/** Called by STreeView when the tree's selection has changed */
-		void OnOutlinerTreeSelectionChanged( FOutlinerTreeItemPtr TreeItem, ESelectInfo::Type SelectInfo );
+		void OnOutlinerTreeSelectionChanged( FTreeItemPtr TreeItem, ESelectInfo::Type SelectInfo );
 
 		/** Called by STreeView when the user double-clicks on an item in the tree */
-		void OnOutlinerTreeDoubleClick( FOutlinerTreeItemPtr TreeItem );
+		void OnOutlinerTreeDoubleClick( FTreeItemPtr TreeItem );
 
 		/** Called by STreeView when an item is scrolled into view */
-		void OnOutlinerTreeItemScrolledIntoView( FOutlinerTreeItemPtr TreeItem, const TSharedPtr<ITableRow>& Widget );
+		void OnOutlinerTreeItemScrolledIntoView( FTreeItemPtr TreeItem, const TSharedPtr<ITableRow>& Widget );
 
 		/** Called when an item in the tree has been collapsed or expanded */
-		void OnItemExpansionChanged(FOutlinerTreeItemPtr TreeItem, bool bIsExpanded) const;
+		void OnItemExpansionChanged(FTreeItemPtr TreeItem, bool bIsExpanded) const;
+
+	private:
+		/** Level, editor and other global event hooks required to keep the outliner up to date */
 
 		/** Called by USelection::SelectionChangedEvent delegate when the level's selection changes */
 		void OnLevelSelectionChanged(UObject* Obj);
@@ -206,107 +257,37 @@ namespace SceneOutliner
 		/** Called by the engine when an actor is being requested to be renamed */
 		void OnLevelActorsRequestRename(const AActor* InActor);
 
-		/** Callback when item's label is committed */
-		void OnItemLabelCommitted( const FText& InLabel, ETextCommit::Type InCommitInfo, FOutlinerTreeItemPtr TreeItem );
-
-		/** Callback to verify a item label change */
-		bool OnVerifyItemLabelChanged( const FText& InLabel, FText& OutErrorMessage, FOutlinerTreeItemPtr TreeItem );
-
 		/** Called by the engine when an actor's folder is changed */
 		void OnLevelActorFolderChanged(const AActor* InActor, FName OldPath);
 
+		/** Handler for when a property changes on any object */
+		void OnActorLabelChanged(AActor* ChangedActor);
+
 		/** Called when the map has changed*/
 		void OnMapChange(uint32 MapFlags);
 
 		/** Called when the current level has changed */
 		void OnNewCurrentLevel();
 
-		/** Open a context menu for this scene outliner */
-		TSharedPtr<SWidget> OnOpenContextMenu() const;
-
-		/** Build a context menu for right-clicking an item in the tree */
-		TSharedPtr<SWidget> BuildDefaultContextMenu() const;
-		void FillFoldersSubMenu(FMenuBuilder& MenuBuilder) const;
-		void AddMoveToFolderOutliner(FMenuBuilder& MenuBuilder) const;
-		void FillSelectionSubMenu(FMenuBuilder& MenuBuilder) const;
-
-		/** Select the immediate children of the currently selected folders */
-		void SelectFoldersImmediateChildren() const;
-
-		/** Called to select all the descendants of the currently selected folders */
-		void SelectFoldersDescendants() const;
-
-		/** Move the selected items to the specified parent */
-		void MoveSelectionTo(TSharedRef<TOutlinerTreeItem> NewParent);
-
-		/** Moves the current selection to the specified folder path */
-		void MoveSelectionTo(FName NewParent);
-
-		/** Called when the user has clicked the button to add a new folder */
-		FReply OnCreateFolderClicked();
-
-		/** Create a new folder under the specified parent name (NAME_None for root) */
-		void CreateFolder();
-
-		/** Create a sub folder inside the specified parent */
-		void CreateSubFolder(FName Parent);
-
-		/** Initiate a rename operation for the specified folder */
-		void RenameFolder(TSharedRef<TOutlinerFolderTreeItem> Folder);
-
-		/** Delete the specified folder and all its contents */
-		void DeleteFolder(TSharedRef<TOutlinerFolderTreeItem> Folder);
-
-		/** Add the specified tree item to its parent, attempting to create the parent if possible */
-		bool AddChildToParent(FOutlinerTreeItemRef TreeItem, FName ParentPath, bool bIgnoreSearchFilter = false);
-
 		/** Called when a folder is to be created */
 		void OnBroadcastFolderCreate(UWorld& InWorld, FName NewPath);
 
 		/** Called when a folder is to be moved */
-		void OnBroadcastFolderMove(FName OldPath, FName NewPath);
+		void OnBroadcastFolderMove(UWorld& InWorld, FName OldPath, FName NewPath);
 
 		/** Called when a folder is to be deleted */
 		void OnBroadcastFolderDelete(UWorld& InWorld, FName Path);
 
-		/** Detach the specified item from its parent if specified */
-		void DetachChildFromParent(TSharedRef<TOutlinerTreeItem> Child, FName ParentPath);
-
-		/**
-		 * Checks to see if the actor is valid for displaying in the outliner
-		 *
-		 * @return	True if actor can be displayed
-		 */
-		bool IsActorDisplayable( const AActor* Actor ) const;
-
-		/** @return Returns a string to use for highlighting results in the outliner list */
-		virtual FText GetFilterHighlightText() const override;
-
-		/**
-		 * Handler for when a property changes on any object
-		 *
-		 * @param	ObjectBeingModified
-		 */
-		virtual void OnActorLabelChanged(AActor* ChangedActor);
+	private:
+		/** Miscellaneous bindings required by the UI */
 
-		/**
-		 * Called by the editable text control when the filter text is changed by the user
-		 *
-		 * @param	InFilterText	The new text
-		 */
+		/** Called by the editable text control when the filter text is changed by the user */
 		void OnFilterTextChanged( const FText& InFilterText );
 
 		/** Called by the editable text control when a user presses enter or commits their text change */
 		void OnFilterTextCommitted( const FText& InFilterText, ETextCommit::Type CommitInfo );
 
-		/** Invoked by the scene outliner gutter when we need to set the specified item's visibility */
-		void OnSetItemVisibility(FOutlinerTreeItemRef Item, bool bIsVisible);
-
-		/**
-		 * Called by the filter button to get the image to display in the button
-		 *
-		 * @return	Slate brush for the button to display
-		 */
+		/** Called by the filter button to get the image to display in the button */
 		const FSlateBrush* GetFilterButtonGlyph() const;
 
 		/** @return	The filter button tool-tip text */
@@ -321,21 +302,6 @@ namespace SceneOutliner
 		/** @return Returns color for the filter status text message, based on success of search filter */
 		FSlateColor GetFilterStatusTextColor() const;
 
-		/** @return	Returns true if the filter is currently active */
-		bool IsFilterActive() const;
-
-		/** @return Returns whether the Searchbox widget should be visibility */
-		EVisibility GetSearchBoxVisibility() const;
-
-		/** Overridden from SWidget: Checks to see if this widget supports keyboard focus */
-		virtual bool SupportsKeyboardFocus() const override;
-
-		/** Overridden from SWidget: Called when a key is pressed down */
-		virtual FReply OnKeyDown( const FGeometry& MyGeometry, const FKeyEvent& InKeyEvent ) override;
-
-		/** Function to validate actor list shown in scene outliner widget */
-		void ValidateOutlinerTreeView();
-
 		/**	Returns the current visibility of the Empty label */
 		EVisibility GetEmptyLabelVisibility() const;
 
@@ -354,10 +320,48 @@ namespace SceneOutliner
 		/** @return the foreground color for the view button */
 		FSlateColor GetViewButtonForegroundColor() const;
 
-		/** Get an array of selected folders in this scene outliner */
-		TArray<TSharedRef<TOutlinerFolderTreeItem>> GetSelectedFolders() const;
+		/** The brush to use when in Editor mode */
+		const FSlateBrush* NoBorder;
+		/** The brush to use when in PIE mode */
+		const FSlateBrush* PlayInEditorBorder;
+		/** The brush to use when in SIE mode */
+		const FSlateBrush* SimulateBorder;
+
+	private:
+
+		/** Open a context menu for this scene outliner */
+		TSharedPtr<SWidget> OnOpenContextMenu();
+
+		/** Build a context menu for right-clicking an item in the tree */
+		TSharedPtr<SWidget> BuildDefaultContextMenu();
+		void FillFoldersSubMenu(FMenuBuilder& MenuBuilder) const;
+		void AddMoveToFolderOutliner(FMenuBuilder& MenuBuilder) const;
+		void FillSelectionSubMenu(FMenuBuilder& MenuBuilder) const;
+		TSharedRef<TSet<FName>> GatherInvalidMoveToDestinations() const;
+
+	private:
+
+		/** Select the immediate children of the currently selected folders */
+		void SelectFoldersImmediateChildren() const;
 
+		/** Called to select all the descendants of the currently selected folders */
+		void SelectFoldersDescendants() const;
+
+		/** Move the selected items to the specified parent */
+		void MoveSelectionTo(FTreeItemRef NewParent);
+
+		/** Moves the current selection to the specified folder path */
+		void MoveSelectionTo(FName NewParent);
+
+		/** Called when the user has clicked the button to add a new folder */
+		FReply OnCreateFolderClicked();
+
+		/** Create a new folder under the specified parent name (NAME_None for root) */
+		void CreateFolder();
+
+	private:
 		/** FILTERS */
+
 		/** @return whether we are displaying only selected Actors */
 		bool IsShowingOnlySelected() const;
 		/** Toggles whether we are displaying only selected Actors */
@@ -379,124 +383,63 @@ namespace SceneOutliner
 		/** Enables/Disables whether the ShowOnlyActorsInCurrentLevelFilter is applied */
 		void ApplyShowOnlyCurrentLevelFilter(bool bShowOnlyActorsInCurrentLevel);
 
-	private:
+		/** When applied, only selected Actors are displayed */
+		TSharedPtr< FOutlinerFilter > SelectedActorFilter;
 
-		/** Structure containing information relating to the expansion state of parent items in the tree */
-		struct FParentsExpansionState
-		{
-			/** The set of parent items which are not expanded */
-			TSet<const AActor*> CollapsedItems;
-		};
+		/** When applied, temporary and run-time actors are hidden */
+		TSharedPtr< FOutlinerFilter > HideTemporaryActorsFilter;
 
+		/** When applied, only Actors that are in the current level are displayed */
+		TSharedPtr< FOutlinerFilter > ShowOnlyActorsInCurrentLevelFilter;
 
-		/** Init options, cached */
-		FSceneOutlinerInitializationOptions InitOptions;
+	private:
 
 		/** Context menu opening delegate provided by the client */
 		FOnContextMenuOpening OnContextMenuOpening;
 
-		/** The current world being represented by the Scene Outliner. */
-		class UWorld* RepresentingWorld;
+		/** Callback that's fired when an item is selected while in 'picking' mode */
+		FOnSceneOutlinerItemPicked OnItemPicked;
 
-		/** List for items being added to the tree since last Populate. */
-		TArray<FOutlinerTreeItemRef> AddedItemsList;
-		
-		/** List for items being removed from the tree since last Populate. */
-		TArray<FOutlinerTreeItemRef> RemovedItemsList;
-		
-		/** List for actors being attached in the tree since last Populate. */
-		TArray<TSharedRef<TOutlinerActorTreeItem>> AttachedActorsList;
+		/** Shared data required by the tree and its items */
+		TSharedRef<FSharedOutlinerData> SharedData;
 
-		/** List for items that need to be refreshed in the tree since last Populate. */
-		TArray<TSharedRef<TOutlinerTreeItem>> RefreshItemsList;
+		/** List of pending operations to be applied to the tree */
+		TArray<FPendingTreeOperation> PendingOperations;
 
-		/** Callback that's fired when an item is selected while in 'picking' mode */
-		FOnSceneOutlinerItemPicked OnItemPicked;
+		/** Map of actions to apply to new tree items */
+		TMap<FTreeItemID, uint8> NewItemActions;
 
 		/** Our tree view */
 		TSharedPtr< SOutlinerTreeView > OutlinerTreeView;
 
-		/** Map of actors to list items in our OutlinerData.  Used to quickly find the item for a specified actor. */
-		FActorToTreeItemMap ActorToTreeItemMap;
+		/** A map of all items we have in the tree */
+		FTreeItemMap TreeItemMap;
+
+		/** Root level tree items */
+		TArray<FTreeItemPtr> RootTreeItems;
 
-		/** Map of folder paths to list items in our OutlinerData. Used to quickly find the item for a specified folder. */
-		FFolderToTreeItemMap FolderToTreeItemMap;
+		/** A set of all actors that pass the non-text filters in the representing world */
+		TSet<TWeakObjectPtr<AActor>> ApplicableActors;
 
 		/** The button that displays view options */
 		TSharedPtr<SComboButton> ViewOptionsComboButton;
 
-		/** FILTERS */
-		/** When applied, only selected Actors are displayed */
-		TSharedPtr< TDelegateFilter< const AActor* const > > SelectedActorFilter;
-
-		/** When applied, temporary and run-time actors are hidden */
-		TSharedPtr< TDelegateFilter< const AActor* const > > HideTemporaryActorsFilter;
-
-		/** When applied, only Actors that are in the current level are displayed */
-		TSharedPtr< TDelegateFilter< const AActor* const > > ShowOnlyActorsInCurrentLevelFilter;
-
-		/** The brush to use when in Editor mode */
-		const FSlateBrush* NoBorder;
-		/** The brush to use when in PIE mode */
-		const FSlateBrush* PlayInEditorBorder;
-		/** The brush to use when in SIE mode */
-		const FSlateBrush* SimulateBorder;
-		/** The component mobility brushes */
-		const FSlateBrush* MobilityStaticBrush;
-		const FSlateBrush* MobilityStationaryBrush;
-		const FSlateBrush* MobilityMovableBrush;
-
-		/** Populates the actor map with the specified UWorld */
-		void UpdateActorMapWithWorld(UWorld* World, FParentActorsSet& ParentActors);
-		
-		/** 
-		 * Attempts to add an item to the tree. Will add any parent folders if required
-		 *
-		 * @param InItem			The item to add to the tree (either an actor or folder)
-		 * @param ParentActors		A set of parent actors found that may need to be added to.
-		 *
-		 * @return					returns true if adding the actor to the tree was successful.
-		 */
-		bool AddItemToTree(TSharedRef<TOutlinerTreeItem> InItem, FParentActorsSet& ParentActors);
-		
-		/** 
-		 * Attempts to add an actor to the tree. Will add any parent folders if required
-		 *
-		 * @param InActorItem		The actor item to add to the tree.
-		 * @param ParentActors		A set of parent actors found that may need to be added to.
-		 *
-		 * @return					returns true if adding the actor to the tree was successful.
-		 */
-		bool AddActorToTree(TSharedRef<TOutlinerActorTreeItem> InActorItem, FParentActorsSet& ParentActors);
-
-		/** Add a new folder to the tree, including any of its parents if possible */
-		bool AddFolderToTree(TSharedRef<TOutlinerFolderTreeItem> InFolderItem, bool bIgnoreSearchFilter = false);
-
-		/** Remove the specified item from the tree */
-		void RemoveItemFromTree(TSharedRef<TOutlinerTreeItem> InItem);
+	private:
 
-		/** 
-		 * Attempts to add parent actors that were filtered out to the tree.
-		 *
-		 * @param InActor			The actor to add to the tree.
-		 */
-		void AddFilteredParentActorToTree(const AActor* InActor);
+		/** Structure containing information relating to the expansion state of parent items in the tree */
+		typedef TMap<FTreeItemID, bool> FParentsExpansionState;
 
 		/** Gets the current expansion state of parent items */
-		void GetParentsExpansionState(FParentsExpansionState& ExpansionStateInfo) const;
+		FParentsExpansionState GetParentsExpansionState() const;
 
 		/** Updates the expansion state of parent items after a repopulate, according to the previous state */
-		void SetParentsExpansionState(const FParentActorsSet& ParentActors, const FParentsExpansionState& ExpansionStateInfo) const;
+		void SetParentsExpansionState(const FParentsExpansionState& ExpansionStateInfo) const;
 
-		/** Whether the scene outliner is currently displaying PlayWorld actors */
-		bool bRepresentingPlayWorld;
+	private:
 
 		/** Number of actors that passed the search filter */
 		int32 FilteredActorCount;
 
-		/** Root level tree items */
-		FOutlinerData RootTreeItems;
-
 		/** True if the outliner needs to be repopulated at the next appropriate opportunity, usually because our
 		    actor set has changed in some way. */
 		bool bNeedsRefresh;
@@ -504,11 +447,8 @@ namespace SceneOutliner
 		/** true if the Scene Outliner should do a full refresh. */
 		bool bFullRefresh;
 
-		/** An instance of a class responsible for operating on newly added tree items */
-		NewItemActuator NewItemActions;
-
-		/** Timer for PIE/SIE mode to sort the outliner. */
-		float SortOutlinerTimer;
+		/** true when the actor selection state in the world does not match the selection state of the tree */
+		bool bActorSelectionDirty;
 
 		/** Reentrancy guard */
 		bool bIsReentrant;
@@ -522,16 +462,26 @@ namespace SceneOutliner
 		/** The TextFilter attached to the SearchBox widget of the Scene Outliner */
 		TSharedPtr< TreeItemTextFilter > SearchBoxFilter;
 
-		/** A custom column to show in the Scene Outliners */
-		TSharedPtr< ISceneOutlinerColumn > CustomColumn;
-
-		/** A visibility gutter which allows users to quickly show/hide actors */
-		TSharedPtr< FSceneOutlinerGutter > Gutter;
-
 		/** True if the search box will take keyboard focus next frame */
 		bool bPendingFocusNextFrame;
 
-		/********** Sort functions **********/
+		/** The tree item that is currently pending a rename */
+		TWeakPtr<ITreeItem> PendingRenameItem;
+
+	private:
+		/** Functions relating to sorting */
+
+		/** Timer for PIE/SIE mode to sort the outliner. */
+		float SortOutlinerTimer;
+
+		/** true if the outliner currently needs to be sorted */
+		bool bSortDirty;
+
+		/** Specify which column to sort with */
+		FName SortByColumn;
+
+		/** Currently selected sorting mode */
+		EColumnSortMode::Type SortMode;
 
 		/** Handles column sorting mode change */
 		void OnColumnSortModeChanged( const EColumnSortPriority::Type SortPriority, const FName& ColumnId, const EColumnSortMode::Type InSortMode );
@@ -542,14 +492,8 @@ namespace SceneOutliner
 		/** Request that the tree be sorted at a convenient time */
 		void RequestSort();
 
-		/** Sort the tree based on the current sort column */
-		void SortTree();
-
-		/** Specify which column to sort with */
-		FName SortByColumn;
-
-		/** Currently selected sorting mode */
-		EColumnSortMode::Type SortMode;
+		/** Sort the specified array of items based on the current sort column */
+		void SortItems(TArray<FTreeItemPtr>& Items) const;
 
 	};
 
diff --git a/Engine/Source/Editor/SceneOutliner/Private/SceneOutlinerActorInfoColumn.cpp b/Engine/Source/Editor/SceneOutliner/Private/SceneOutlinerActorInfoColumn.cpp
--- a/Engine/Source/Editor/SceneOutliner/Private/SceneOutlinerActorInfoColumn.cpp
+++ b/Engine/Source/Editor/SceneOutliner/Private/SceneOutlinerActorInfoColumn.cpp
@@ -1,30 +1,130 @@
 // Copyright 1998-2014 Epic Games, Inc. All Rights Reserved.
 
 #include "SceneOutlinerPrivatePCH.h"
-#include "SceneOutlinerTreeItems.h"
+
+#include "ActorTreeItem.h"
 
 #include "EditorClassUtils.h"
+#include "SortHelper.h"
 
 #define LOCTEXT_NAMESPACE "SceneOutlinerActorInfoColumn"
 
-TArray< TSharedPtr< ECustomColumnMode::Type > > FSceneOutlinerActorInfoColumn::ModeOptions;
+namespace SceneOutliner
+{
+
+struct FGetInfo : TTreeItemGetter<FString>
+{
+	ECustomColumnMode::Type CurrentMode;
+
+	FGetInfo(ECustomColumnMode::Type InCurrentMode) : CurrentMode(InCurrentMode) {}
+
+	virtual FString Get(const FActorTreeItem& ActorItem) const override
+	{
+		AActor* Actor = ActorItem.Actor.Get();
+		if (!Actor)
+		{
+			return FString();
+		}
+
+		switch(CurrentMode)
+		{
+		case ECustomColumnMode::Class:
+			return Actor->GetClass()->GetName();
+
+		case ECustomColumnMode::Level:
+			return FPackageName::GetShortName(Actor->GetOutermost()->GetName());
+
+		case ECustomColumnMode::Socket:
+			return Actor->GetAttachParentSocketName().ToString();
+
+		case ECustomColumnMode::InternalName:
+			return Actor->GetFName().ToString();
+
+		case ECustomColumnMode::UncachedLights:
+			return FString::Printf(TEXT("%7d"), Actor->GetNumUncachedStaticLightingInteractions());
+
+		case ECustomColumnMode::Layer:
+			{
+				FString Result;
+				for (const auto& Layer : Actor->Layers)
+				{
+					if (Result.Len())
+					{
+						Result += TEXT(", ");
+					}
+
+					Result += Layer.ToString();
+				}
+				return Result;
+			}
+
+		default:
+			return FString();
+		}
+	}
 
-FSceneOutlinerActorInfoColumn::FSceneOutlinerActorInfoColumn( const TWeakPtr< ISceneOutliner >& InSceneOutlinerWeak, ECustomColumnMode::Type InDefaultCustomColumnMode )
+	virtual FString Get(const FFolderTreeItem&) const override
+	{
+		switch(CurrentMode)
+		{
+		case ECustomColumnMode::Class:
+			return LOCTEXT("FolderTypeName", "Folder").ToString();
+
+		default:
+			return FString();
+		}
+	}
+
+	virtual FString Get(const FWorldTreeItem&) const override
+	{
+		switch(CurrentMode)
+		{
+		case ECustomColumnMode::Class:
+			return LOCTEXT("WorldTypeName", "World").ToString();
+
+		default:
+			return FString();
+		}
+	}
+
+	virtual FString Get(const FLevelBlueprintTreeItem& Item) const override
+	{
+		switch(CurrentMode)
+		{
+		case ECustomColumnMode::Class:
+			return LOCTEXT("LevelBlueprintTypeName", "Level Blueprint").ToString();
+
+		case ECustomColumnMode::Level:
+			if (ULevel* Level = Item.Handle.ParentLevel.Get())
+			{
+				return FPackageName::GetShortName(Level->GetOutermost()->GetName());
+			}
+			return FString();
+
+		default:
+			return FString();
+		}
+	}
+};
+
+
+TArray< TSharedPtr< ECustomColumnMode::Type > > FActorInfoColumn::ModeOptions;
+
+FActorInfoColumn::FActorInfoColumn( ISceneOutliner& Outliner, ECustomColumnMode::Type InDefaultCustomColumnMode )
 	: CurrentMode( InDefaultCustomColumnMode )
-	, SceneOutlinerWeak( InSceneOutlinerWeak )
+	, SceneOutlinerWeak( StaticCastSharedRef<ISceneOutliner>(Outliner.AsShared()) )
 {
 
 }
 
 
-FName FSceneOutlinerActorInfoColumn::GetColumnID()
+FName FActorInfoColumn::GetColumnID()
 {
-	static FName ActorInfo("ActorInfo");
-	return ActorInfo;
+	return GetID();
 }
 
 
-SHeaderRow::FColumn::FArguments FSceneOutlinerActorInfoColumn::ConstructHeaderRowColumn()
+SHeaderRow::FColumn::FArguments FActorInfoColumn::ConstructHeaderRowColumn()
 {
 	if( ModeOptions.Num() == 0 )
 	{
@@ -36,6 +136,7 @@ SHeaderRow::FColumn::FArguments FSceneOutlinerActorInfoColumn::ConstructHeaderRo
 
 	/** Customizable actor data column */
 	return SHeaderRow::Column( GetColumnID() )
+		.FillWidth(2)
 		.HeaderComboVisibility(EHeaderComboVisibility::Ghosted)
 		.MenuContent()
 		[
@@ -45,8 +146,8 @@ SHeaderRow::FColumn::FArguments FSceneOutlinerActorInfoColumn::ConstructHeaderRo
 			[
 				SNew(SListView<TSharedPtr<ECustomColumnMode::Type>>)
 				.ListItemsSource(&ModeOptions)
-				.OnGenerateRow( this, &FSceneOutlinerActorInfoColumn::MakeComboButtonItemWidget )
-				.OnSelectionChanged( this, &FSceneOutlinerActorInfoColumn::OnModeChanged )
+				.OnGenerateRow( this, &FActorInfoColumn::MakeComboButtonItemWidget )
+				.OnSelectionChanged( this, &FActorInfoColumn::OnModeChanged )
 			]
 		]
 		.HeaderContent()
@@ -57,161 +158,112 @@ SHeaderRow::FColumn::FArguments FSceneOutlinerActorInfoColumn::ConstructHeaderRo
 			.VAlign(VAlign_Center)
 			[
 				SNew(STextBlock)
-				.Text( this, &FSceneOutlinerActorInfoColumn::GetSelectedMode )
+				.Text( this, &FActorInfoColumn::GetSelectedMode )
 			]
 		];
 }
 
-
-TSharedRef< SWidget > FSceneOutlinerActorInfoColumn::ConstructClassHyperlink( const TWeakObjectPtr< AActor >&  InActor )
+const TSharedRef< SWidget > FActorInfoColumn::ConstructRowWidget( FTreeItemRef TreeItem, const STableRow<FTreeItemPtr>& Row )
 {
-	const AActor* Actor = InActor.Get();
-	if ( Actor == nullptr )
-	{
-		return SNew( SSpacer );
-	}
+	auto SceneOutliner = SceneOutlinerWeak.Pin();
+	check(SceneOutliner.IsValid());
 
-	UClass* ObjectClass = Actor->GetClass();
+	TSharedRef<SHorizontalBox> HorizontalBox = SNew(SHorizontalBox);
 
-	TSharedPtr<SWidget> SourceLink; 
-	// Only create source link in SO for Blueprints (not native classes, too noisy having header link in SO)
-	if(UBlueprint::GetBlueprintFromClass(ObjectClass) != NULL)
-	{
-		SourceLink = FEditorClassUtils::GetSourceLink(ObjectClass, InActor);
-	}
+	TSharedRef<STextBlock> MainText = SNew( STextBlock )
+		.Text( this, &FActorInfoColumn::GetTextForItem, TWeakPtr<ITreeItem>(TreeItem) ) 
+		.HighlightText( SceneOutliner->GetFilterHighlightText() )
+		.ColorAndOpacity( FSlateColor::UseSubduedForeground() );
 
-	if (SourceLink.IsValid())
-	{
-		SourceLink->SetVisibility( TAttribute<EVisibility>::Create( TAttribute<EVisibility>::FGetter::CreateSP( this, &FSceneOutlinerActorInfoColumn::GetColumnDataVisibility, true ) ) );
+	HorizontalBox->AddSlot()
+	.AutoWidth()
+	.VAlign(VAlign_Center)
+	[
+		MainText
+	];
 
-		return SourceLink.ToSharedRef();
-	}
-	return SNew( STextBlock )
-		.Text( ObjectClass->GetName() )
-		.HighlightText( SceneOutlinerWeak.Pin().Get(), &ISceneOutliner::GetFilterHighlightText )
-		.ColorAndOpacity( FSlateColor::UseSubduedForeground() )
-		.Visibility( this, &FSceneOutlinerActorInfoColumn::GetColumnDataVisibility, true );
-}
-
-const TSharedRef< SWidget > FSceneOutlinerActorInfoColumn::ConstructRowWidget( const TSharedRef<SceneOutliner::TOutlinerTreeItem> TreeItem )
-{
-	if (TreeItem->Type == SceneOutliner::TOutlinerTreeItem::Actor)
+	auto Hyperlink = ConstructClassHyperlink(*TreeItem);
+	if (Hyperlink.IsValid())
 	{
-		const auto Actor = StaticCastSharedRef<SceneOutliner::TOutlinerActorTreeItem>(TreeItem)->Actor;
-		return SNew( SHorizontalBox )
-			+SHorizontalBox::Slot()
-			.AutoWidth()
-			.Padding( 0.0f, 3.0f, 0.0f, 0.0f )
+		// If we got a hyperlink, disable hide default text, and show the hyperlink
+		MainText->SetVisibility(TAttribute<EVisibility>::Create(TAttribute<EVisibility>::FGetter::CreateSP(this, &FActorInfoColumn::GetColumnDataVisibility, false)));
+		Hyperlink->SetVisibility(TAttribute<EVisibility>::Create(TAttribute<EVisibility>::FGetter::CreateSP(this, &FActorInfoColumn::GetColumnDataVisibility, true)));
+
+		HorizontalBox->AddSlot()
+		.VAlign(VAlign_Center)
+		.AutoWidth()
+		[
+			SNew(SBorder)
+			.BorderImage(FEditorStyle::GetBrush("NoBorder"))
+			.ColorAndOpacity_Static([](TWeakPtr<const STableRow<FTreeItemPtr>> WeakRow)->FLinearColor{
+				auto Row = WeakRow.Pin();
+				return Row.IsValid() && Row->IsSelectedExclusively() ? FLinearColor::Black : FLinearColor::White;
+			}, TWeakPtr<const STableRow<FTreeItemPtr>>(StaticCastSharedRef<const STableRow<FTreeItemPtr>>(Row.AsShared())))
 			[
-				SNew( SHorizontalBox )
-
-				+ SHorizontalBox::Slot()
-				.AutoWidth()
-				[
-					SNew( STextBlock )
-
-					// Bind a delegate for custom text for this item's row
-					.Text( this, &FSceneOutlinerActorInfoColumn::GetTextForActor, Actor ) 
-
-					// Bind our filter text as the highlight string for the text block, so that when the user
-					// starts typing search criteria, this text highlights
-					.HighlightText( SceneOutlinerWeak.Pin().Get(), &ISceneOutliner::GetFilterHighlightText )
-
-					.ColorAndOpacity( FSlateColor::UseSubduedForeground() )
-					.Visibility( this, &FSceneOutlinerActorInfoColumn::GetColumnDataVisibility, false )
-				]
-				+ SHorizontalBox::Slot()
-				.AutoWidth()
-				[
-					ConstructClassHyperlink( Actor )
-				]
-			];
-	}
-	else
-	{
-		return SNullWidget::NullWidget;
+				Hyperlink.ToSharedRef()
+			]
+		];
 	}
-}
 
-bool FSceneOutlinerActorInfoColumn::ProvidesSearchStrings()
-{
-	return false;
+	return HorizontalBox;
 }
 
-void FSceneOutlinerActorInfoColumn::PopulateActorSearchStrings( const AActor* const Actor, OUT TArray< FString >& OutSearchStrings ) const
+TSharedPtr<SWidget> FActorInfoColumn::ConstructClassHyperlink( ITreeItem& TreeItem )
 {
-	switch( CurrentMode )
+	struct FConstructHyperlink : TTreeItemGetter<TSharedPtr<SWidget>>
 	{
-	case ECustomColumnMode::None:
-		break;
-
-	case ECustomColumnMode::Layer:
+		TSharedPtr<SWidget> Get(const FActorTreeItem& ActorItem) const override
 		{
-			for( int32 LayerIndex=0; LayerIndex < Actor->Layers.Num(); ++LayerIndex)
+			if (AActor* Actor = ActorItem.Actor.Get())
 			{
-				OutSearchStrings.Add( Actor->Layers[ LayerIndex ].ToString() );
+				UClass* ActorClass = Actor->GetClass();
+
+				// Only create source link in SO for Blueprints (not native classes, too noisy having header link in SO)
+				if (ActorClass && UBlueprint::GetBlueprintFromClass(ActorClass) != NULL)
+				{
+					return FEditorClassUtils::GetSourceLink(ActorClass, Actor);
+				}
 			}
-		}
-		break;
 
-	case ECustomColumnMode::Level:
-		{
-			OutSearchStrings.Add( Actor->GetOutermost()->GetName() );
+			return nullptr;
 		}
-		break;
+	};
 
-	case ECustomColumnMode::Socket:
-		{
-			OutSearchStrings.Add( Actor->GetAttachParentSocketName().ToString() );
-		}
-		break;
+	return TreeItem.Get(FConstructHyperlink());
+}
 
-	case ECustomColumnMode::InternalName:
+void FActorInfoColumn::PopulateSearchStrings( const ITreeItem& Item, TArray< FString >& OutSearchStrings ) const
+{
+	FString String = Item.Get(FGetInfo(CurrentMode));
+	if (String.Len())
+	{
+		OutSearchStrings.Add(String);
+	}
+
+	// We always add the class
+	if (CurrentMode != ECustomColumnMode::Class)
+	{
+		FString String = Item.Get(FGetInfo(ECustomColumnMode::Class));
+		if (String.Len())
 		{
-			OutSearchStrings.Add( Actor->GetFName().ToString() );
+			OutSearchStrings.Add(String);
 		}
-		break;
-
-	case ECustomColumnMode::UncachedLights:
-		break;
-
 	}
-
-	// Covers ECustomColumnMode::ActorClass 
-	// We always want this to be a search string
-	FString ActorClassName;
-	Actor->GetClass()->GetName( ActorClassName );
-	OutSearchStrings.Add( FString::Printf( TEXT( "(%s)" ), *ActorClassName) );
 }
 
-bool FSceneOutlinerActorInfoColumn::SupportsSorting() const
+bool FActorInfoColumn::SupportsSorting() const
 {
 	return CurrentMode != ECustomColumnMode::None;
 }
 
-void FSceneOutlinerActorInfoColumn::SortItems(TArray<TSharedPtr<SceneOutliner::TOutlinerTreeItem>>& RootItems, const EColumnSortMode::Type SortMode) const
+void FActorInfoColumn::SortItems(TArray<FTreeItemPtr>& RootItems, const EColumnSortMode::Type SortMode) const
 {
-	if (SortMode == EColumnSortMode::Ascending)
-	{
-		RootItems.Sort(
-			[=](TSharedPtr<SceneOutliner::TOutlinerTreeItem> A, TSharedPtr<SceneOutliner::TOutlinerTreeItem> B)
-			{
-				return GetTextForItem(A.ToSharedRef()) < GetTextForItem(B.ToSharedRef());
-			}
-		);
-	}
-	else if (SortMode == EColumnSortMode::Descending)
-	{
-		RootItems.Sort(
-			[=](TSharedPtr<SceneOutliner::TOutlinerTreeItem> A, TSharedPtr<SceneOutliner::TOutlinerTreeItem> B)
-			{
-				return GetTextForItem(A.ToSharedRef()) > GetTextForItem(B.ToSharedRef());
-			}
-		);
-	}
+	FSortHelper<FString>()
+		.Primary(FGetInfo(CurrentMode), SortMode)
+		.Sort(RootItems);
 }
 
-void FSceneOutlinerActorInfoColumn::OnModeChanged( TSharedPtr< ECustomColumnMode::Type > NewSelection, ESelectInfo::Type /*SelectInfo*/ )
+void FActorInfoColumn::OnModeChanged( TSharedPtr< ECustomColumnMode::Type > NewSelection, ESelectInfo::Type /*SelectInfo*/ )
 {
 	CurrentMode = *NewSelection;
 
@@ -220,9 +272,9 @@ void FSceneOutlinerActorInfoColumn::OnModeChanged( TSharedPtr< ECustomColumnMode
 	FSlateApplication::Get().DismissAllMenus();
 }
 
-EVisibility FSceneOutlinerActorInfoColumn::GetColumnDataVisibility( bool bIsClassHyperlink ) const
+EVisibility FActorInfoColumn::GetColumnDataVisibility( bool bIsClassHyperlink ) const
 {
-	if ( CurrentMode == ECustomColumnMode::ActorClass )
+	if ( CurrentMode == ECustomColumnMode::Class )
 	{
 		return bIsClassHyperlink ? EVisibility::Visible : EVisibility::Collapsed;
 	}
@@ -232,75 +284,13 @@ EVisibility FSceneOutlinerActorInfoColumn::GetColumnDataVisibility( bool bIsClas
 	}
 }
 
-FString FSceneOutlinerActorInfoColumn::GetTextForItem( const TSharedRef<SceneOutliner::TOutlinerTreeItem> TreeItem ) const
-{
-	if (TreeItem->Type == SceneOutliner::TOutlinerTreeItem::Actor)
-	{
-		return GetTextForActor(StaticCastSharedRef<const SceneOutliner::TOutlinerActorTreeItem>(TreeItem)->Actor);
-	}
-	else
-	{
-		return StaticCastSharedRef<const SceneOutliner::TOutlinerFolderTreeItem>(TreeItem)->LeafName.ToString();
-	}
-}
-
-FString FSceneOutlinerActorInfoColumn::GetTextForActor( const TWeakObjectPtr< AActor > TreeItem ) const
+FString FActorInfoColumn::GetTextForItem( TWeakPtr<ITreeItem> TreeItem ) const
 {
-	FString CustomColumnText;
-	AActor* Actor = TreeItem.Get();
-
-	if( Actor == NULL )
-	{
-		return CustomColumnText;
-	}
-
-	switch( CurrentMode )
-	{
-	case ECustomColumnMode::ActorClass:
-		CustomColumnText = Actor->GetClass()->GetName();
-		break;
-
-	case ECustomColumnMode::Layer:
-		{
-			FString ResultString = TEXT("");
-
-			for( auto LayerIt = Actor->Layers.CreateConstIterator(); LayerIt; ++LayerIt )
-			{
-				if( ResultString.Len() )
-				{
-					ResultString += TEXT(",");
-				}
-
-				ResultString += LayerIt->ToString();
-			}
-
-			CustomColumnText = ResultString;
-		}
-		break;
-
-	case ECustomColumnMode::Level:
-		CustomColumnText = FPackageName::GetShortName( Actor->GetOutermost()->GetName() );
-		break;
-
-	case ECustomColumnMode::Socket:
-		{
-			CustomColumnText = Actor->GetAttachParentSocketName().ToString();
-		}
-		break;
-
-	case ECustomColumnMode::InternalName:
-		CustomColumnText = Actor->GetFName().ToString();
-		break;
-
-	case ECustomColumnMode::UncachedLights:
-		CustomColumnText = FString::Printf(TEXT("%7d"), Actor->GetNumUncachedStaticLightingInteractions());
-		break;
-	};
-
-	return CustomColumnText;
+	auto Item = TreeItem.Pin();
+	return Item.IsValid() ? Item->Get(FGetInfo(CurrentMode)) : FString();
 }
 
-FText FSceneOutlinerActorInfoColumn::GetSelectedMode() const
+FText FActorInfoColumn::GetSelectedMode() const
 {
 	if (CurrentMode == ECustomColumnMode::None)
 	{
@@ -310,7 +300,7 @@ FText FSceneOutlinerActorInfoColumn::GetSelectedMode() const
 	return MakeComboText(CurrentMode);
 }
 
-FText FSceneOutlinerActorInfoColumn::MakeComboText( const ECustomColumnMode::Type& Mode ) const
+FText FActorInfoColumn::MakeComboText( const ECustomColumnMode::Type& Mode ) const
 {
 	FText ModeName;
 
@@ -320,8 +310,8 @@ FText FSceneOutlinerActorInfoColumn::MakeComboText( const ECustomColumnMode::Typ
 		ModeName = LOCTEXT("CustomColumnMode_None", "None");
 		break;
 
-	case ECustomColumnMode::ActorClass:
-		ModeName = LOCTEXT("CustomColumnMode_ActorClass", "Type");
+	case ECustomColumnMode::Class:
+		ModeName = LOCTEXT("CustomColumnMode_Class", "Type");
 		break;
 
 	case ECustomColumnMode::Level:
@@ -353,7 +343,7 @@ FText FSceneOutlinerActorInfoColumn::MakeComboText( const ECustomColumnMode::Typ
 }
 
 
-FText FSceneOutlinerActorInfoColumn::MakeComboToolTipText( const ECustomColumnMode::Type& Mode )
+FText FActorInfoColumn::MakeComboToolTipText( const ECustomColumnMode::Type& Mode )
 {
 	FText ToolTipText;
 
@@ -363,8 +353,8 @@ FText FSceneOutlinerActorInfoColumn::MakeComboToolTipText( const ECustomColumnMo
 		ToolTipText = LOCTEXT("CustomColumnModeToolTip_None", "Hides all extra actor info");
 		break;
 
-	case ECustomColumnMode::ActorClass:
-		ToolTipText = LOCTEXT("CustomColumnModeToolTip_ActorClass", "Displays the name of each actor's type");
+	case ECustomColumnMode::Class:
+		ToolTipText = LOCTEXT("CustomColumnModeToolTip_Class", "Displays the name of each actor's type");
 		break;
 
 	case ECustomColumnMode::Level:
@@ -396,7 +386,7 @@ FText FSceneOutlinerActorInfoColumn::MakeComboToolTipText( const ECustomColumnMo
 }
 
 
-TSharedRef< ITableRow > FSceneOutlinerActorInfoColumn::MakeComboButtonItemWidget( TSharedPtr< ECustomColumnMode::Type > Mode, const TSharedRef<STableViewBase>& Owner )
+TSharedRef< ITableRow > FActorInfoColumn::MakeComboButtonItemWidget( TSharedPtr< ECustomColumnMode::Type > Mode, const TSharedRef<STableViewBase>& Owner )
 {
 	return
 		SNew( STableRow< TSharedPtr<ECustomColumnMode::Type> >, Owner )
@@ -407,4 +397,6 @@ TSharedRef< ITableRow > FSceneOutlinerActorInfoColumn::MakeComboButtonItemWidget
 		];
 }
 
+}		// namespace SceneOutliner
+
 #undef LOCTEXT_NAMESPACE
diff --git a/Engine/Source/Editor/SceneOutliner/Private/SceneOutlinerActorInfoColumn.h b/Engine/Source/Editor/SceneOutliner/Private/SceneOutlinerActorInfoColumn.h
--- a/Engine/Source/Editor/SceneOutliner/Private/SceneOutlinerActorInfoColumn.h
+++ b/Engine/Source/Editor/SceneOutliner/Private/SceneOutlinerActorInfoColumn.h
@@ -4,12 +4,48 @@
 
 #include "ISceneOutliner.h"
 #include "ISceneOutlinerColumn.h"
-#include "SceneOutlinerInitializationOptions.h"
+#include "SceneOutlinerPublicTypes.h"
+
+namespace SceneOutliner
+{
+
+/** Types of actor data we can display in a 'custom' tree column */
+namespace ECustomColumnMode
+{
+	enum Type
+	{
+		/** Empty column -- doesn't display anything */
+		None = 0,
+
+		/** Class */
+		Class,
+
+		/** Level */
+		Level,
+
+		/** Layer */
+		Layer,
+
+		/** The socket the actor is attached to. */
+		Socket,
+
+		/** Actor's internal name (FName) */
+		InternalName,
+
+		/** Actor's number of uncached lights */
+		UncachedLights,
+
+		// ---
+
+		/** Number of options */
+		Count
+	};
+}
 
 /**
  * A custom column for the SceneOutliner which is capable of displaying a variety of Actor details
  */
-class FSceneOutlinerActorInfoColumn : public ISceneOutlinerColumn
+class FActorInfoColumn : public ISceneOutlinerColumn
 {
 
 public:
@@ -17,9 +53,11 @@ public:
 	/**
 	 *	Constructor
 	 */
-	FSceneOutlinerActorInfoColumn( const TWeakPtr< ISceneOutliner >& InSceneOutlinerWeak, ECustomColumnMode::Type InDefaultCustomColumnMode );
+	FActorInfoColumn( ISceneOutliner& Outliner, ECustomColumnMode::Type InDefaultCustomColumnMode = ECustomColumnMode::Class );
+
+	virtual ~FActorInfoColumn() {}
 
-	virtual ~FSceneOutlinerActorInfoColumn() {}
+	static FName GetID() { return FBuiltInColumnTypes::ActorInfo(); }
 
 	//////////////////////////////////////////////////////////////////////////
 	// Begin ISceneOutlinerColumn Implementation
@@ -28,31 +66,27 @@ public:
 
 	virtual SHeaderRow::FColumn::FArguments ConstructHeaderRowColumn() override;
 
-	virtual const TSharedRef< SWidget > ConstructRowWidget( const TSharedRef<SceneOutliner::TOutlinerTreeItem> TreeItem ) override;
+	virtual const TSharedRef< SWidget > ConstructRowWidget( FTreeItemRef TreeItem, const STableRow<FTreeItemPtr>& Row ) override;
 
-	virtual bool ProvidesSearchStrings() override;
-
-	virtual void PopulateActorSearchStrings( const AActor* const Actor, OUT TArray< FString >& OutSearchStrings ) const override;
+	virtual void PopulateSearchStrings( const ITreeItem& Item, TArray< FString >& OutSearchStrings ) const override;
 
 	virtual bool SupportsSorting() const override;
 
-	virtual void SortItems(TArray<TSharedPtr<SceneOutliner::TOutlinerTreeItem>>& RootItems, const EColumnSortMode::Type SortMode) const override;
+	virtual void SortItems(TArray<FTreeItemPtr>& RootItems, const EColumnSortMode::Type SortMode) const override;
 	
 	// End ISceneOutlinerColumn Implementation
 	//////////////////////////////////////////////////////////////////////////
 
+	FString GetTextForItem( const TWeakPtr<ITreeItem> TreeItem ) const;
+
 private:
 
-	TSharedRef< SWidget > ConstructClassHyperlink( const TWeakObjectPtr< AActor >&  Actor );
+	TSharedPtr<SWidget> ConstructClassHyperlink( ITreeItem& TreeItem );
 
 	void OnModeChanged( TSharedPtr< ECustomColumnMode::Type > NewSelection, ESelectInfo::Type SelectInfo );
 
 	EVisibility GetColumnDataVisibility( bool bIsClassHyperlink ) const;
 
-	FString GetTextForItem( const TSharedRef<SceneOutliner::TOutlinerTreeItem> TreeItem ) const;
-
-	FString GetTextForActor( const TWeakObjectPtr< AActor > TreeItem ) const;
-
 	FText MakeComboText( const ECustomColumnMode::Type& Mode ) const;
 
 	FText MakeComboToolTipText( const ECustomColumnMode::Type& Mode );
@@ -75,3 +109,6 @@ private:
 	/** Weak reference to the outliner widget that owns our list */
 	TWeakPtr< ISceneOutliner > SceneOutlinerWeak;
 };
+
+
+}	// namespace SceneOutliner
\ No newline at end of file
diff --git a/Engine/Source/Editor/SceneOutliner/Private/SceneOutlinerDragDrop.cpp b/Engine/Source/Editor/SceneOutliner/Private/SceneOutlinerDragDrop.cpp
new file mode 100644
--- /dev/null
+++ b/Engine/Source/Editor/SceneOutliner/Private/SceneOutlinerDragDrop.cpp
@@ -0,0 +1,167 @@
+// Copyright 1998-2014 Epic Games, Inc. All Rights Reserved.
+
+#include "SceneOutlinerPrivatePCH.h"
+#include "SceneOutlinerDragDrop.h"
+#include "EditorActorFolders.h"
+
+#define LOCTEXT_NAMESPACE "SSceneOutliner"
+
+namespace SceneOutliner
+{
+
+	FSceneOutlinerDragDropOp::FSceneOutlinerDragDropOp(const FDragDropPayload& DraggedObjects)
+		: OverrideText()
+		, OverrideIcon(nullptr)
+	{
+		if (DraggedObjects.Actors)
+		{
+			ActorOp = MakeShareable(new FActorDragDropOp);
+			ActorOp->Init(DraggedObjects.Actors.GetValue());
+		}
+
+		if (DraggedObjects.Folders)
+		{
+			FolderOp = MakeShareable(new FFolderDragDropOp);
+			FolderOp->Init(DraggedObjects.Folders.GetValue());
+		}
+	}
+	
+	EVisibility FSceneOutlinerDragDropOp::GetOverrideVisibility() const
+	{
+		return OverrideText.IsEmpty() && OverrideIcon == nullptr ? EVisibility::Collapsed : EVisibility::Visible;
+	}
+	
+	EVisibility FSceneOutlinerDragDropOp::GetDefaultVisibility() const
+	{
+		return OverrideText.IsEmpty() && OverrideIcon == nullptr ? EVisibility::Visible : EVisibility::Collapsed;
+	}
+
+	TSharedPtr<SWidget> FSceneOutlinerDragDropOp::GetDefaultDecorator() const
+	{
+		TSharedRef<SVerticalBox> VerticalBox = SNew(SVerticalBox);
+
+		VerticalBox->AddSlot()
+		[
+			SNew(SBorder)
+			.BorderImage(FEditorStyle::GetBrush("Graph.ConnectorFeedback.Border"))
+			.Visibility(this, &FSceneOutlinerDragDropOp::GetOverrideVisibility)
+			.Content()
+			[			
+				SNew(SHorizontalBox)
+				
+				+ SHorizontalBox::Slot()
+				.AutoWidth()
+				.Padding( 0.0f, 0.0f, 3.0f, 0.0f )
+				[
+					SNew( SImage )
+					.Image( this, &FSceneOutlinerDragDropOp::GetOverrideIcon )
+				]
+
+				+ SHorizontalBox::Slot()
+				.AutoWidth()
+				.VAlign( VAlign_Center )
+				[
+					SNew(STextBlock) 
+					.Text( this, &FSceneOutlinerDragDropOp::GetOverrideText )
+				]
+			]
+		];
+
+		if (FolderOp.IsValid())
+		{
+			auto Content = FolderOp->GetDefaultDecorator();
+			if (Content.IsValid())
+			{
+				Content->SetVisibility(TAttribute<EVisibility>(this, &FSceneOutlinerDragDropOp::GetDefaultVisibility));
+				VerticalBox->AddSlot()[ Content.ToSharedRef() ];
+			}
+		}
+
+		if (ActorOp.IsValid())
+		{
+			auto Content = ActorOp->GetDefaultDecorator();
+			if (Content.IsValid())
+			{
+				Content->SetVisibility(TAttribute<EVisibility>(this, &FSceneOutlinerDragDropOp::GetDefaultVisibility));
+				VerticalBox->AddSlot()[ Content.ToSharedRef() ];
+			}
+		}
+
+		return VerticalBox;
+	}
+
+	void FFolderDragDropOp::Init(FFolderPaths InFolders)
+	{
+		Folders = MoveTemp(InFolders);
+
+		CurrentIconBrush = FEditorStyle::Get().GetBrush(TEXT("SceneOutliner.FolderClosed"));
+		if (Folders.Num() == 1)
+		{
+			CurrentHoverText = FText::FromName(GetFolderLeafName(Folders[0]));
+		}
+		else
+		{
+			CurrentHoverText = FText::Format(NSLOCTEXT("FFolderDragDropOp", "FormatFolders", "{0} Folders"), FText::AsNumber(Folders.Num()));
+		}
+	}
+
+	TSharedPtr<FDragDropOperation> CreateDragDropOperation(const TArray<FTreeItemPtr>& InTreeItems)
+	{
+		FDragDropPayload DraggedObjects;
+		for (const auto& Item : InTreeItems)
+		{
+			Item->PopulateDragDropPayload(DraggedObjects);
+		}
+
+		if (DraggedObjects.Folders)
+		{
+			TSharedPtr<FSceneOutlinerDragDropOp> OutlinerOp = MakeShareable(new FSceneOutlinerDragDropOp(DraggedObjects));
+			OutlinerOp->Construct();
+			return OutlinerOp;
+		}
+		else if (DraggedObjects.Actors)
+		{
+			return FActorDragDropGraphEdOp::New(DraggedObjects.Actors.GetValue());
+		}
+		
+		return nullptr;
+	}
+
+	FDragDropPayload::FDragDropPayload()
+	{
+
+	}
+
+	bool FDragDropPayload::ParseDrag(const FDragDropOperation& Operation)
+	{
+		bool bApplicable = false;
+
+		if (Operation.IsOfType<FSceneOutlinerDragDropOp>())
+		{
+			const auto& OutlinerOp = static_cast<const FSceneOutlinerDragDropOp&>(Operation);
+			if (OutlinerOp.FolderOp.IsValid())
+			{
+				Folders = OutlinerOp.FolderOp->Folders;
+			}
+			if (OutlinerOp.ActorOp.IsValid())
+			{
+				Actors = OutlinerOp.ActorOp->Actors;
+			}
+			bApplicable = true;
+		}
+		else if (Operation.IsOfType<FActorDragDropOp>())
+		{
+			Actors = static_cast<const FActorDragDropOp&>(Operation).Actors;
+			bApplicable = true;
+		}
+		else if (Operation.IsOfType<FFolderDragDropOp>())
+		{
+			Folders = static_cast<const FFolderDragDropOp&>(Operation).Folders;
+			bApplicable = true;
+		}
+
+		return bApplicable;
+	}
+}
+
+#undef LOCTEXT_NAMESPACE
diff --git a/Engine/Source/Editor/SceneOutliner/Private/SceneOutlinerGutter.cpp b/Engine/Source/Editor/SceneOutliner/Private/SceneOutlinerGutter.cpp
--- a/Engine/Source/Editor/SceneOutliner/Private/SceneOutlinerGutter.cpp
+++ b/Engine/Source/Editor/SceneOutliner/Private/SceneOutlinerGutter.cpp
@@ -2,10 +2,108 @@
 
 
 #include "SceneOutlinerPrivatePCH.h"
-#include "SceneOutlinerTreeItems.h"
+
+#include "SortHelper.h"
 
 #define LOCTEXT_NAMESPACE "SceneOutlinerGutter"
 
+namespace SceneOutliner
+{
+
+bool FGetVisibilityVisitor::RecurseChildren(const ITreeItem& Item) const
+{
+	if (const bool* Info = VisibilityInfo.Find(&Item))
+	{
+		return *Info;
+	}
+	else
+	{
+		bool bIsVisible = false;
+		for (const auto& Child : Item.GetChildren())
+		{
+			if (Child.Pin()->Get(*this))
+			{
+				bIsVisible = true;
+				break;
+			}
+		}
+		VisibilityInfo.Add(&Item, bIsVisible);
+		
+		return bIsVisible;
+	}
+}
+
+bool FGetVisibilityVisitor::Get(const FActorTreeItem& ActorItem) const
+{
+	if (const bool* Info = VisibilityInfo.Find(&ActorItem))
+	{
+		return *Info;
+	}
+	else
+	{
+		const auto* Actor = ActorItem.Actor.Get();
+
+		const bool bIsVisible = Actor && !Actor->IsTemporarilyHiddenInEditor();
+		VisibilityInfo.Add(&ActorItem, bIsVisible);
+
+		return bIsVisible;
+	}
+}
+
+bool FGetVisibilityVisitor::Get(const FWorldTreeItem& WorldItem) const
+{
+	return RecurseChildren(WorldItem);
+}
+
+bool FGetVisibilityVisitor::Get(const FLevelBlueprintTreeItem&) const
+{
+	return false;
+}
+
+bool FGetVisibilityVisitor::Get(const FFolderTreeItem& FolderItem) const
+{
+	return RecurseChildren(FolderItem);
+}
+
+struct FSetVisibilityVisitor : IMutableTreeItemVisitor
+{
+	/** Whether this item should be visible or not */
+	const bool bSetVisibility;
+
+	FSetVisibilityVisitor(bool bInSetVisibility) : bSetVisibility(bInSetVisibility) {}
+
+	virtual void Visit(FActorTreeItem& ActorItem) const override
+	{
+		auto* Actor = ActorItem.Actor.Get();
+		if (Actor)
+		{
+			// Save the actor to the transaction buffer to support undo/redo, but do
+			// not call Modify, as we do not want to dirty the actor's package and
+			// we're only editing temporary, transient values
+			SaveToTransactionBuffer(Actor, false);
+			Actor->SetIsTemporarilyHiddenInEditor( !bSetVisibility );
+		}
+	}
+
+	virtual void Visit(FWorldTreeItem& WorldItem) const override
+	{
+		for (auto& Child : WorldItem.GetChildren())
+		{
+			FSetVisibilityVisitor Visibility(bSetVisibility);
+			Child.Pin()->Visit(Visibility);
+		}
+	}
+
+	virtual void Visit(FFolderTreeItem& FolderItem) const override
+	{
+		for (auto& Child : FolderItem.GetChildren())
+		{
+			FSetVisibilityVisitor Visibility(bSetVisibility);
+			Child.Pin()->Visit(Visibility);
+		}
+	}
+};
+
 class FVisibilityDragDropOp : public FDragDropOperation, public TSharedFromThis<FVisibilityDragDropOp>
 {
 public:
@@ -42,41 +140,23 @@ class SVisibilityWidget : public SImage
 {
 public:
 	SLATE_BEGIN_ARGS(SVisibilityWidget){}
-		SLATE_ARGUMENT(TWeakPtr<SceneOutliner::TOutlinerTreeItem>, TreeItem);
-		SLATE_EVENT(FOnSetItemVisibility, OnSetItemVisibility);
 	SLATE_END_ARGS()
 
 	/** Construct this widget */
-	void Construct(const FArguments& InArgs)
+	void Construct(const FArguments& InArgs, TWeakPtr<FSceneOutlinerGutter> InWeakColumn, TWeakPtr<ISceneOutliner> InWeakOutliner, TWeakPtr<ITreeItem> InWeakTreeItem)
 	{
-		WeakTreeItem = InArgs._TreeItem;
-		OnSetItemVisibility = InArgs._OnSetItemVisibility;
+		WeakTreeItem = InWeakTreeItem;
+		WeakOutliner = InWeakOutliner;
+		WeakColumn = InWeakColumn;
 
 		SImage::Construct(
 			SImage::FArguments()
-				.Image(this, &SVisibilityWidget::GetBrush)
+			.Image(this, &SVisibilityWidget::GetBrush)
 		);
 	}
 
-	/** Check if the specified actor is visible */
-	static bool IsActorVisible(const AActor* Actor)
-	{
-		return Actor && !Actor->IsTemporarilyHiddenInEditor();
-	}
-
 private:
 
-	virtual FReply OnMouseButtonDoubleClick(const FGeometry& InMyGeometry, const FPointerEvent& InMouseEvent)
-	{
-		if (InMouseEvent.GetEffectingButton() == EKeys::LeftMouseButton)
-		{
-			UndoTransaction.Reset(new FScopedTransaction(LOCTEXT("SetActorVisibility", "Set Actor Visibility")));
-			SetIsVisible(!IsVisible());
-		}
-
-		return FReply::Handled();
-	}
-
 	/** Start a new drag/drop operation for this widget */
 	virtual FReply OnDragDetected(const FGeometry& MyGeometry, const FPointerEvent& MouseEvent) override
 	{
@@ -100,10 +180,13 @@ private:
 		}
 	}
 
-	/** Called when the mouse button is pressed down on this widget */
-	virtual FReply OnMouseButtonDown(const FGeometry& MyGeometry, const FPointerEvent& MouseEvent) override
+	FReply HandleClick()
 	{
-		if (MouseEvent.GetEffectingButton() != EKeys::LeftMouseButton)
+		auto Outliner = WeakOutliner.Pin();
+		auto TreeItem = WeakTreeItem.Pin();
+		auto Column = WeakColumn.Pin();
+		
+		if (!Outliner.IsValid() || !TreeItem.IsValid() || !Column.IsValid())
 		{
 			return FReply::Unhandled();
 		}
@@ -111,11 +194,49 @@ private:
 		// Open an undo transaction
 		UndoTransaction.Reset(new FScopedTransaction(LOCTEXT("SetActorVisibility", "Set Actor Visibility")));
 
-		SetIsVisible(!IsVisible());
+		const auto& Tree = Outliner->GetTree();
+
+		const bool bVisible = !IsVisible();
+
+		// We operate on all the selected items if the specified item is selected
+		if (Tree.IsItemSelected(TreeItem.ToSharedRef()))
+		{
+			const FSetVisibilityVisitor Visitor(bVisible);
+
+			for (auto& SelectedItem : Tree.GetSelectedItems())
+			{
+				if (IsVisible(SelectedItem, Column) != bVisible)
+				{
+					SelectedItem->Visit(Visitor);
+				}		
+			}
+
+			GEditor->RedrawAllViewports();
+		}
+		else
+		{
+			SetIsVisible(bVisible);
+		}
 
 		return FReply::Handled().DetectDrag(SharedThis(this), EKeys::LeftMouseButton);
 	}
 
+	virtual FReply OnMouseButtonDoubleClick(const FGeometry& InMyGeometry, const FPointerEvent& InMouseEvent)
+	{
+		return HandleClick();
+	}
+
+	/** Called when the mouse button is pressed down on this widget */
+	virtual FReply OnMouseButtonDown(const FGeometry& MyGeometry, const FPointerEvent& MouseEvent) override
+	{
+		if (MouseEvent.GetEffectingButton() != EKeys::LeftMouseButton)
+		{
+			return FReply::Unhandled();
+		}
+
+		return HandleClick();
+	}
+
 	/** Process a mouse up message */
 	virtual FReply OnMouseButtonUp( const FGeometry& MyGeometry, const FPointerEvent& MouseEvent ) override
 	{
@@ -149,112 +270,100 @@ private:
 		}
 	}
 
-	/** Check if the specified actor is visible */
+	/** Check if the specified item is visible */
+	static bool IsVisible(const FTreeItemPtr& Item, const TSharedPtr<FSceneOutlinerGutter>& Column)
+	{
+		return Column.IsValid() && Item.IsValid() ? Column->IsItemVisible(*Item) : false;
+	}
+
+	/** Check if our wrapped tree item is visible */
 	bool IsVisible() const
 	{
-		auto TreeItem = WeakTreeItem.Pin();
-		if (TreeItem.IsValid())
-		{
-			return TreeItem->IsVisible();
-		}
-		return false;
+		return IsVisible(WeakTreeItem.Pin(), WeakColumn.Pin());
 	}
 
 	/** Set the actor this widget is responsible for to be hidden or shown */
 	void SetIsVisible(const bool bVisible)
 	{
-		if (IsVisible() != bVisible)
+		auto TreeItem = WeakTreeItem.Pin();
+		if (TreeItem.IsValid() && IsVisible() != bVisible)
 		{
-			auto ClickedOnItem = WeakTreeItem.Pin();
-			if (ClickedOnItem.IsValid())
-			{
-				OnSetItemVisibility.ExecuteIfBound(ClickedOnItem.ToSharedRef(), bVisible);
-			}
+			FSetVisibilityVisitor Visitor(bVisible);
+			TreeItem->Visit(Visitor);
+
+			GEditor->RedrawAllViewports();
 		}
 	}
 
-	/** Tree item that we relate to */
-	TWeakPtr<SceneOutliner::TOutlinerTreeItem> WeakTreeItem;
+	/** The tree item we relate to */
+	TWeakPtr<ITreeItem> WeakTreeItem;
+	
+	/** Reference back to the outliner so we can set visibility of a whole selection */
+	TWeakPtr<ISceneOutliner> WeakOutliner;
 
-	/** A delegate to execute when we need to set the visibility of an item */
-	FOnSetItemVisibility OnSetItemVisibility;
+	/** Weak pointer back to the column */
+	TWeakPtr<FSceneOutlinerGutter> WeakColumn;
 
 	/** Scoped undo transaction */
 	TUniquePtr<FScopedTransaction> UndoTransaction;
 };
 
-FSceneOutlinerGutter::FSceneOutlinerGutter(FOnSetItemVisibility InOnSetItemVisibility)
-	: OnSetItemVisibility(InOnSetItemVisibility)
+FSceneOutlinerGutter::FSceneOutlinerGutter(ISceneOutliner& Outliner)
 {
+	WeakOutliner = StaticCastSharedRef<ISceneOutliner>(Outliner.AsShared());
+}
 
+void FSceneOutlinerGutter::Tick(double InCurrentTime, float InDeltaTime)
+{
+	VisibilityCache.VisibilityInfo.Empty();
 }
 
 FName FSceneOutlinerGutter::GetColumnID()
 {
-	return FName("Gutter");
+	return GetID();
 }
 
 SHeaderRow::FColumn::FArguments FSceneOutlinerGutter::ConstructHeaderRowColumn()
 {
-	return SHeaderRow::Column(GetColumnID())[ SNew(SSpacer) ];
+	return SHeaderRow::Column(GetColumnID())
+		.FixedWidth(16.f)
+		[
+			SNew(SSpacer)
+		];
 }
 
-const TSharedRef<SWidget> FSceneOutlinerGutter::ConstructRowWidget(const TSharedRef<SceneOutliner::TOutlinerTreeItem> TreeItem)
+const TSharedRef<SWidget> FSceneOutlinerGutter::ConstructRowWidget(FTreeItemRef TreeItem, const STableRow<FTreeItemPtr>& Row)
 {
+	bool bGutterIsApplicable = true;
+	TreeItem->Visit(FFunctionalVisitor().LevelBlueprint([&](const FLevelBlueprintTreeItem&){
+		bGutterIsApplicable = false;
+	}));
+
+	if (!bGutterIsApplicable)
+	{
+		return SNullWidget::NullWidget;
+	}
+	
 	return SNew(SHorizontalBox)
 		+SHorizontalBox::Slot()
 		.AutoWidth()
 		.VAlign(VAlign_Center)
 		[
-			SNew(SVisibilityWidget)
-			.OnSetItemVisibility(OnSetItemVisibility)
-			.TreeItem(TreeItem)
+			SNew(SVisibilityWidget, SharedThis(this), WeakOutliner, TreeItem)
 		];
 }
 
-void FSceneOutlinerGutter::SortItems(TArray<TSharedPtr<SceneOutliner::TOutlinerTreeItem>>& RootItems, const EColumnSortMode::Type SortMode) const
+void FSceneOutlinerGutter::SortItems(TArray<FTreeItemPtr>& RootItems, const EColumnSortMode::Type SortMode) const
 {
-	using namespace SceneOutliner;
-	struct SortPredicate
-	{
-		/** Constructor */
-		SortPredicate(const bool InSence) : bSense(InSence) {}
-
-		/** The sense of the sort (identical items are always sorted by name, Ascending) */
-		bool bSense;
-
-		/** Sort predicate operator */
-		bool operator ()(TSharedPtr<TOutlinerTreeItem> A, TSharedPtr<TOutlinerTreeItem> B) const
-		{
-			if ((A->Type == TOutlinerTreeItem::Folder) != (B->Type == TOutlinerTreeItem::Folder))
-			{
-				// Folders appear first
-				return A->Type == TOutlinerTreeItem::Folder;
-			}
-
-			const bool VisibleA = A->IsVisible();
-			const bool VisibleB = A->IsVisible();
-
-			if (VisibleA == VisibleB)
-			{
-				// Always sort identical items by name
-				return GetLabelForItemAsString(A.ToSharedRef()) < GetLabelForItemAsString(B.ToSharedRef());
-			}
-			else
-			{
-				return bSense == VisibleA;
-			}
-		}
-	};
-
-	if (SortMode == EColumnSortMode::Ascending)
-	{
-		RootItems.Sort(SortPredicate(true));
-	}
-	else if (SortMode == EColumnSortMode::Descending)
-	{
-		RootItems.Sort(SortPredicate(false));
-	}
+	FSortHelper<int32, bool>()
+		/** Sort by type first */
+		.Primary([](const ITreeItem& Item){ return Item.GetTypeSortPriority(); }, SortMode)
+		/** Then by visibility */
+		.Secondary(FGetVisibilityVisitor(), SortMode)
+		.Sort(RootItems);
 }
 
+}		// namespace SceneOutliner
+
 #undef LOCTEXT_NAMESPACE
+
diff --git a/Engine/Source/Editor/SceneOutliner/Private/SceneOutlinerGutter.h b/Engine/Source/Editor/SceneOutliner/Private/SceneOutlinerGutter.h
--- a/Engine/Source/Editor/SceneOutliner/Private/SceneOutlinerGutter.h
+++ b/Engine/Source/Editor/SceneOutliner/Private/SceneOutlinerGutter.h
@@ -2,11 +2,32 @@
 
 #pragma once
 
-#include "ISceneOutliner.h"
 #include "ISceneOutlinerColumn.h"
-#include "SceneOutlinerInitializationOptions.h"
 
-DECLARE_DELEGATE_TwoParams(FOnSetItemVisibility, TSharedRef<SceneOutliner::TOutlinerTreeItem>, bool)
+namespace SceneOutliner
+{
+
+/** A 'getter' visitor that gets, and caches the visibility of a tree item */
+struct FGetVisibilityVisitor : TTreeItemGetter<bool>
+{
+	/** Map of tree item to visibility */
+	mutable TMap<const ITreeItem*, bool> VisibilityInfo;
+
+	/** Get an item's visibility based on its children */
+	bool RecurseChildren(const ITreeItem& Item) const;
+
+	/** Get an actor's visibility */
+	virtual bool Get(const FActorTreeItem& ActorItem) const override;
+
+	/** Get a World's visibility */
+	virtual bool Get(const FWorldTreeItem& WorldItem) const override;
+
+	/** Get a Level Blueprint's visibility */
+	virtual bool Get(const FLevelBlueprintTreeItem& LevelBlueprintItem) const override;
+
+	/** Get a folder's visibility */
+	virtual bool Get(const FFolderTreeItem& FolderItem) const override;
+};
 
 /**
  * A gutter for the SceneOutliner which is capable of displaying a variety of Actor details
@@ -17,29 +38,40 @@ class FSceneOutlinerGutter : public ISceneOutlinerColumn
 public:
 
 	/**	Constructor */
-	FSceneOutlinerGutter(FOnSetItemVisibility InOnSetItemVisibility);
+	FSceneOutlinerGutter(ISceneOutliner& Outliner);
 
 	virtual ~FSceneOutlinerGutter() {}
 
+	static FName GetID() { return FBuiltInColumnTypes::Gutter(); }
+	
 	// -----------------------------------------
 	// ISceneOutlinerColumn Implementation
 	virtual FName GetColumnID() override;
 
 	virtual SHeaderRow::FColumn::FArguments ConstructHeaderRowColumn() override;
 
-	virtual const TSharedRef< SWidget > ConstructRowWidget( const TSharedRef<SceneOutliner::TOutlinerTreeItem> TreeItem ) override;
-
-	virtual bool ProvidesSearchStrings() override { return false; }
-
-	virtual void PopulateActorSearchStrings(const AActor* const Actor, OUT TArray< FString >& OutSearchStrings) const override {}
+	virtual const TSharedRef< SWidget > ConstructRowWidget( FTreeItemRef TreeItem, const STableRow<FTreeItemPtr>& Row ) override;
+	
+	virtual void Tick(double InCurrentTime, float InDeltaTime) override;
 
 	virtual bool SupportsSorting() const override { return true; }
 
-	virtual void SortItems(TArray<TSharedPtr<SceneOutliner::TOutlinerTreeItem>>& RootItems, const EColumnSortMode::Type SortMode) const override;
+	virtual void SortItems(TArray<FTreeItemPtr>& RootItems, const EColumnSortMode::Type SortMode) const override;
 	// -----------------------------------------
 
+	/** Check whether the specified item is visible */
+	FORCEINLINE bool IsItemVisible(const ITreeItem& Item)
+	{
+		return Item.Get(VisibilityCache);
+	}
+
 private:
 
-	/** A delegate to execute when we need to set the visibility of an item */
-	FOnSetItemVisibility OnSetItemVisibility;
-};
\ No newline at end of file
+	/** Weak pointer back to the scene outliner - required for setting visibility on current selection. */
+	TWeakPtr<ISceneOutliner> WeakOutliner;
+
+	/** Visitor used to get (and cache) visibilty for items. Cahced per-frame to avoid expensive recursion. */
+	FGetVisibilityVisitor VisibilityCache;
+};
+
+}	// namespace SceneOutliner
\ No newline at end of file
diff --git a/Engine/Source/Editor/SceneOutliner/Private/SceneOutlinerItemLabelColumn.cpp b/Engine/Source/Editor/SceneOutliner/Private/SceneOutlinerItemLabelColumn.cpp
new file mode 100644
--- /dev/null
+++ b/Engine/Source/Editor/SceneOutliner/Private/SceneOutlinerItemLabelColumn.cpp
@@ -0,0 +1,702 @@
+// Copyright 1998-2014 Epic Games, Inc. All Rights Reserved.
+
+#include "SceneOutlinerPrivatePCH.h"
+#include "SceneOutlinerStandaloneTypes.h"
+#include "SceneOutlinerDragDrop.h"
+
+#include "SceneOutlinerItemLabelColumn.h"
+
+#include "EditorClassUtils.h"
+#include "SortHelper.h"
+
+#include "ActorTreeItem.h"
+#include "FolderTreeItem.h"
+#include "EditorActorFolders.h"
+
+#include "SInlineEditableTextBlock.h"
+
+#define LOCTEXT_NAMESPACE "SceneOutlinerItemLabelColumn"
+
+namespace SceneOutliner
+{
+
+struct FCommonLabelData
+{
+	TWeakPtr<ISceneOutliner> WeakSceneOutliner;
+	static const FLinearColor DarkColor;
+
+	TOptional<FLinearColor> GetForegroundColor(const FTreeItemPtr& TreeItem) const
+	{
+		if (!TreeItem.IsValid())
+		{
+			return DarkColor;
+		}
+
+		UWorld* World = TreeItem->GetSharedData().RepresentingWorld;
+
+		// Darken items that aren't suitable targets for an active drag and drop action
+		if (FSlateApplication::Get().IsDragDropping() && World)
+		{
+			TSharedPtr<FDragDropOperation> DragDropOp = FSlateApplication::Get().GetDragDroppingContent();
+
+			FDragDropPayload DraggedObjects;
+			if (DraggedObjects.ParseDrag(*DragDropOp) && !TreeItem->ValidateDrop(DraggedObjects, *World).IsValid())
+			{
+				return DarkColor;
+			}
+		}
+
+		if (!TreeItem->CanInteract())
+		{
+			return DarkColor;
+		}
+
+		return TOptional<FLinearColor>();
+	}
+};
+
+const FLinearColor FCommonLabelData::DarkColor(0.3f, 0.3f, 0.3f);
+
+struct SActorTreeLabel : FCommonLabelData, public SCompoundWidget
+{
+	SLATE_BEGIN_ARGS(SActorTreeLabel) {}
+	SLATE_END_ARGS()
+
+	void Construct(const FArguments& InArgs, FActorTreeItem& ActorItem, ISceneOutliner& SceneOutliner, const STableRow<FTreeItemPtr>& InRow)
+	{
+		WeakSceneOutliner = StaticCastSharedRef<ISceneOutliner>(SceneOutliner.AsShared());
+
+		TreeItemPtr = StaticCastSharedRef<FActorTreeItem>(ActorItem.AsShared());
+		ActorPtr = ActorItem.Actor;
+
+		HighlightText = SceneOutliner.GetFilterHighlightText();
+
+		MobilityStaticBrush = FEditorStyle::GetBrush( "ClassIcon.ComponentMobilityStaticPip" );
+		MobilityStationaryBrush = FEditorStyle::GetBrush( "ClassIcon.ComponentMobilityStationaryPip" );
+		MobilityMovableBrush = FEditorStyle::GetBrush( "ClassIcon.ComponentMobilityMovablePip" );
+
+		TSharedPtr<SInlineEditableTextBlock> InlineTextBlock;
+
+		auto MainContent = SNew(SHorizontalBox)
+
+			// Main actor label
+			+ SHorizontalBox::Slot()
+			[
+				SAssignNew(InlineTextBlock, SInlineEditableTextBlock)
+				.Text(this, &SActorTreeLabel::GetDisplayText)
+				.ToolTipText(this, &SActorTreeLabel::GetTooltipText)
+				.HighlightText(HighlightText)
+				.ColorAndOpacity(this, &SActorTreeLabel::GetForegroundColor)
+				.OnTextCommitted(this, &SActorTreeLabel::OnLabelCommitted)
+				.OnVerifyTextChanged(this, &SActorTreeLabel::OnVerifyItemLabelChanged)
+				.IsSelected(FIsSelected::CreateSP(&InRow, &STableRow<FTreeItemPtr>::IsSelectedExclusively))
+			]
+
+			+ SHorizontalBox::Slot()
+			.AutoWidth()
+			.Padding( 0.0f, 0.f, 3.0f, 0.0f )
+			[
+				SNew( STextBlock )
+				.Text(this, &SActorTreeLabel::GetTypeText)
+				.Visibility(this, &SActorTreeLabel::GetTypeTextVisibility)
+				.HighlightText(HighlightText)
+			];
+
+		TSharedRef<SOverlay> IconContent = SNew(SOverlay)
+			+ SOverlay::Slot()
+			.HAlign(HAlign_Right)
+			.VAlign(VAlign_Center)
+			[
+				SNew(SImage)
+				.Image(this, &SActorTreeLabel::GetIcon)
+				.ToolTipText(this, &SActorTreeLabel::GetIconTooltip)
+			];
+
+		if (ActorItem.GetSharedData().Mode == ESceneOutlinerMode::ActorBrowsing)
+		{
+			// Add the component mobility icon
+			IconContent->AddSlot()
+			.HAlign(HAlign_Left)
+			[
+				SNew(SImage)
+				.Image(this, &SActorTreeLabel::GetBrushForComponentMobilityIcon)
+			];
+
+			ActorItem.RenameRequestEvent.BindSP( InlineTextBlock.Get(), &SInlineEditableTextBlock::EnterEditingMode );
+		}
+		
+		ChildSlot
+		[
+			SNew(SHorizontalBox)
+
+			+ SHorizontalBox::Slot()
+			.AutoWidth()
+			.VAlign(VAlign_Center)
+			.Padding(FDefaultTreeItemMetrics::IconPadding())
+			[
+				SNew(SBox)
+				.WidthOverride(FDefaultTreeItemMetrics::IconSize())
+				.HeightOverride(FDefaultTreeItemMetrics::IconSize())
+				[
+					IconContent
+				]
+			]
+
+			+ SHorizontalBox::Slot()
+			.FillWidth(1.0f)
+			.VAlign(VAlign_Center)
+			.Padding(0.0f, 2.0f)
+			[
+				MainContent
+			]
+		];
+	}
+
+private:
+	TWeakPtr<FActorTreeItem> TreeItemPtr;
+	TWeakObjectPtr<AActor> ActorPtr;
+	TAttribute<FText> HighlightText;
+
+	/** The component mobility brushes */
+	const FSlateBrush* MobilityStaticBrush;
+	const FSlateBrush* MobilityStationaryBrush;
+	const FSlateBrush* MobilityMovableBrush;
+
+	FText GetDisplayText() const
+	{
+		const AActor* Actor = ActorPtr.Get();
+		return Actor ? FText::FromString(Actor->GetActorLabel()) : LOCTEXT("ActorLabelForMissingActor", "(Deleted Actor)");
+	}
+
+	FText GetTooltipText() const
+	{
+		if (const AActor* Actor = ActorPtr.Get())
+		{
+			FFormatNamedArguments Args;
+			Args.Add(TEXT("ID_Name"), LOCTEXT("CustomColumnMode_InternalName", "ID Name"));
+			Args.Add(TEXT("Name"), FText::FromString(Actor->GetName()));
+			return FText::Format(LOCTEXT("ActorNameTooltip", "{ID_Name}: {Name}"), Args);
+		}
+
+		return FText();
+	}
+
+	FText GetTypeText() const
+	{
+		if (const AActor* Actor = ActorPtr.Get())
+		{
+			return FText::FromName(Actor->GetClass()->GetFName());
+		}			
+
+		return FText();
+	}
+
+	EVisibility GetTypeTextVisibility() const
+	{
+		return HighlightText.Get().IsEmpty() ? EVisibility::Collapsed : EVisibility::Visible;
+	}
+
+	const FSlateBrush* GetIcon() const
+	{
+		if (const AActor* Actor = ActorPtr.Get())
+		{
+			return FClassIconFinder::FindIconForActor(Actor);
+		}
+		else
+		{
+			return nullptr;
+		}
+	}
+
+	FText GetIconTooltip() const
+	{
+		auto TreeItem = TreeItemPtr.Pin();
+		if (!TreeItem.IsValid())
+		{
+			return FText();
+		}
+
+		FText ToolTipText;
+		if (AActor* Actor = ActorPtr.Get())
+		{
+			ToolTipText = FText::FromString( Actor->GetClass()->GetName() );
+			if (TreeItem->GetSharedData().Mode == ESceneOutlinerMode::ActorBrowsing)
+			{
+				USceneComponent* RootComponent = Actor->GetRootComponent();
+				if (RootComponent)
+				{
+					FFormatNamedArguments Args;
+					Args.Add( TEXT("ActorClassName"), ToolTipText );
+
+					if( RootComponent->Mobility == EComponentMobility::Static )
+					{
+						ToolTipText = FText::Format( LOCTEXT( "ComponentMobility_Static", "{ActorClassName} with static mobility" ), Args );
+					}
+					else if( RootComponent->Mobility == EComponentMobility::Stationary )
+					{
+						ToolTipText = FText::Format( LOCTEXT( "ComponentMobility_Stationary", "{ActorClassName} with stationary mobility" ), Args );
+					}
+					else if( RootComponent->Mobility == EComponentMobility::Movable )
+					{
+						ToolTipText = FText::Format( LOCTEXT( "ComponentMobility_Movable", "{ActorClassName} with movable mobility" ), Args );
+					}
+				}
+			}
+		}
+
+		return ToolTipText;
+	}
+
+	FSlateColor GetForegroundColor() const
+	{
+		auto TreeItem = TreeItemPtr.Pin();
+		if (auto BaseColor = FCommonLabelData::GetForegroundColor(TreeItem))
+		{
+			return BaseColor.GetValue();
+		}
+
+		AActor* Actor = ActorPtr.Get();
+		if (!Actor)
+		{
+			// Deleted actor!
+			return FLinearColor(0.2f, 0.2f, 0.25f);
+		}
+
+		if (TreeItem->GetSharedData().bRepresentingPlayWorld && !TreeItem->bExistsInCurrentWorldAndPIE)
+		{
+			// Highlight actors that are exclusive to PlayWorld
+			return FLinearColor(0.12f, 0.56f, 1.0f);
+		}
+
+		// also darken items that are non selectable in the active mode(s)
+		const bool bInSelected = true;
+		const bool bSelectEvenIfHidden = true;		// @todo outliner: Is this actually OK?
+		if (!GEditor->CanSelectActor(Actor, bInSelected, bSelectEvenIfHidden))
+		{
+			return FCommonLabelData::DarkColor;
+		}
+
+		return FSlateColor::UseForeground();
+	}
+
+	bool OnVerifyItemLabelChanged(const FText& InLabel, FText& OutErrorMessage)
+	{
+		FText TrimmedLabel = FText::TrimPrecedingAndTrailing(InLabel);
+
+		if (TrimmedLabel.IsEmpty())
+		{
+			OutErrorMessage = LOCTEXT( "RenameFailed_LeftBlank", "Names cannot be left blank" );
+			return false;
+		}
+
+		if (TrimmedLabel.ToString().Len() >= NAME_SIZE)
+		{
+			FFormatNamedArguments Arguments;
+			Arguments.Add( TEXT("CharCount"), NAME_SIZE );
+			OutErrorMessage = FText::Format(LOCTEXT("RenameFailed_TooLong", "Names must be less than {CharCount} characters long."), Arguments);
+			return false;
+		}
+
+		return true;
+	}
+
+	void OnLabelCommitted(const FText& InLabel, ETextCommit::Type InCommitInfo)
+	{
+		auto* Actor = ActorPtr.Get();
+		if (Actor && Actor->IsActorLabelEditable() && !InLabel.ToString().Equals(Actor->GetActorLabel(), ESearchCase::CaseSensitive))
+		{
+			const FScopedTransaction Transaction( LOCTEXT( "SceneOutlinerRenameActorTransaction", "Rename Actor" ) );
+			Actor->SetActorLabel( InLabel.ToString() );
+
+			auto Outliner = WeakSceneOutliner.Pin();
+			if (Outliner.IsValid())
+			{
+				Outliner->SetKeyboardFocus();
+			}
+		}
+	}
+
+	const FSlateBrush* GetBrushForComponentMobilityIcon() const
+	{
+		const FSlateBrush* IconBrush = MobilityStaticBrush;
+
+		if (auto* Actor = ActorPtr.Get())
+		{
+			USceneComponent* RootComponent = ActorPtr->GetRootComponent();
+			if (RootComponent)
+			{
+				if (RootComponent->Mobility == EComponentMobility::Stationary)
+				{
+					IconBrush = MobilityStationaryBrush;
+				}
+				else if (RootComponent->Mobility == EComponentMobility::Movable)
+				{
+					IconBrush = MobilityMovableBrush;
+				}
+			}
+		}
+
+		return IconBrush;
+	}
+};
+
+struct SWorldTreeLabel : FCommonLabelData, public SCompoundWidget
+{
+	SLATE_BEGIN_ARGS(SWorldTreeLabel){}
+	SLATE_END_ARGS()
+
+	void Construct(const FArguments& InArgs, FWorldTreeItem& WorldItem, ISceneOutliner& SceneOutliner, const STableRow<FTreeItemPtr>& InRow)
+	{
+		TreeItemPtr = StaticCastSharedRef<FWorldTreeItem>(WorldItem.AsShared());
+		WeakSceneOutliner = StaticCastSharedRef<ISceneOutliner>(SceneOutliner.AsShared());
+
+		ChildSlot
+		[
+			SNew(SHorizontalBox)
+
+			+ SHorizontalBox::Slot()
+			.AutoWidth()
+			.VAlign(VAlign_Center)
+			.Padding(FDefaultTreeItemMetrics::IconPadding())
+			[
+				SNew(SBox)
+				.WidthOverride(FDefaultTreeItemMetrics::IconSize())
+				.HeightOverride(FDefaultTreeItemMetrics::IconSize())
+				[
+					SNew(SImage)
+					.Image(FClassIconFinder::FindIconForClass(UWorld::StaticClass()))
+					.ToolTipText(LOCTEXT("WorldIcon_Tooltip", "World"))
+				]
+			]
+
+			+ SHorizontalBox::Slot()
+			.FillWidth(1.0f)
+			.VAlign(VAlign_Center)
+			.Padding(0.0f, 2.0f)
+			[
+				SNew(STextBlock)
+				.Text(this, &SWorldTreeLabel::GetDisplayText)
+				.HighlightText(SceneOutliner.GetFilterHighlightText())
+				.ColorAndOpacity(this, &SWorldTreeLabel::GetForegroundColor)
+			]
+		];
+	}
+
+private:
+	TWeakPtr<FWorldTreeItem> TreeItemPtr;
+
+	FText GetDisplayText() const
+	{
+		auto Item = TreeItemPtr.Pin();
+		return Item.IsValid() ? FText::FromString(Item->GetDisplayString()) : FText();
+	}
+
+	FSlateColor GetForegroundColor() const
+	{
+		if (auto BaseColor = FCommonLabelData::GetForegroundColor(TreeItemPtr.Pin()))
+		{
+			return BaseColor.GetValue();
+		}
+		
+		return FSlateColor::UseForeground();
+	}
+};
+
+struct SLevelBlueprintTreeLabel : FCommonLabelData, public SCompoundWidget
+{
+	SLATE_BEGIN_ARGS(SLevelBlueprintTreeLabel){}
+	SLATE_END_ARGS()
+
+	void Construct(const FArguments& InArgs, FLevelBlueprintTreeItem& LevelBlueprintItem, ISceneOutliner& SceneOutliner, const STableRow<FTreeItemPtr>& InRow)
+	{
+		TreeItemPtr = StaticCastSharedRef<FLevelBlueprintTreeItem>(LevelBlueprintItem.AsShared());
+		WeakSceneOutliner = StaticCastSharedRef<ISceneOutliner>(SceneOutliner.AsShared());
+
+		ChildSlot
+		[
+			SNew(SHorizontalBox)
+
+			+ SHorizontalBox::Slot()
+			.AutoWidth()
+			.VAlign(VAlign_Center)
+			.Padding(FDefaultTreeItemMetrics::IconPadding())
+			[
+				SNew(SBox)
+				.WidthOverride(FDefaultTreeItemMetrics::IconSize())
+				.HeightOverride(FDefaultTreeItemMetrics::IconSize())
+				[
+					SNew(SImage)
+					.Image(FEditorStyle::GetBrush("LevelEditor.OpenLevelBlueprint"))
+					.ToolTipText(LOCTEXT("LevelBlueprint_Tooltip", "Level Blueprint"))
+				]
+			]
+
+			+ SHorizontalBox::Slot()
+			.FillWidth(1.0f)
+			.VAlign(VAlign_Center)
+			.Padding(0.0f, 2.0f)
+			[
+				SNew(STextBlock)
+				.Text(this, &SLevelBlueprintTreeLabel::GetDisplayText)
+				.HighlightText(SceneOutliner.GetFilterHighlightText())
+				.ColorAndOpacity(this, &SLevelBlueprintTreeLabel::GetForegroundColor)
+			]
+		];
+	}
+
+private:
+	TWeakPtr<FLevelBlueprintTreeItem> TreeItemPtr;
+
+	FText GetDisplayText() const
+	{
+		auto Item = TreeItemPtr.Pin();
+		return Item.IsValid() ? FText::FromString(Item->GetDisplayString()) : FText();
+	}
+
+	FSlateColor GetForegroundColor() const
+	{
+		if (auto BaseColor = FCommonLabelData::GetForegroundColor(TreeItemPtr.Pin()))
+		{
+			return BaseColor.GetValue();
+		}
+		
+		return FSlateColor::UseForeground();
+	}
+};
+
+struct SFolderTreeLabel : FCommonLabelData, public SCompoundWidget
+{
+	SLATE_BEGIN_ARGS(SFolderTreeLabel){}
+	SLATE_END_ARGS()
+
+	void Construct(const FArguments& InArgs, FFolderTreeItem& FolderItem, ISceneOutliner& SceneOutliner, const STableRow<FTreeItemPtr>& InRow)
+	{
+		TreeItemPtr = StaticCastSharedRef<FFolderTreeItem>(FolderItem.AsShared());
+		WeakSceneOutliner = StaticCastSharedRef<ISceneOutliner>(SceneOutliner.AsShared());
+
+		TSharedPtr<SInlineEditableTextBlock> InlineTextBlock = SNew(SInlineEditableTextBlock)
+			.Text(this, &SFolderTreeLabel::GetDisplayText)
+			.HighlightText(SceneOutliner.GetFilterHighlightText())
+			.ColorAndOpacity(this, &SFolderTreeLabel::GetForegroundColor)
+			.OnTextCommitted(this, &SFolderTreeLabel::OnLabelCommitted)
+			.OnVerifyTextChanged(this, &SFolderTreeLabel::OnVerifyItemLabelChanged)
+			.IsSelected(FIsSelected::CreateSP(&InRow, &STableRow<FTreeItemPtr>::IsSelectedExclusively));
+
+		if (SceneOutliner.GetSharedData().Mode == ESceneOutlinerMode::ActorBrowsing)
+		{
+			FolderItem.RenameRequestEvent.BindSP( InlineTextBlock.Get(), &SInlineEditableTextBlock::EnterEditingMode );
+		}
+
+		ChildSlot
+		[
+			SNew(SHorizontalBox)
+
+			+ SHorizontalBox::Slot()
+			.AutoWidth()
+			.VAlign(VAlign_Center)
+			.Padding(FDefaultTreeItemMetrics::IconPadding())
+			[
+				SNew(SBox)
+				.WidthOverride(FDefaultTreeItemMetrics::IconSize())
+				[
+					SNew(SImage)
+					.Image(this, &SFolderTreeLabel::GetIcon)
+				]
+			]
+
+			+ SHorizontalBox::Slot()
+			.FillWidth(1.0f)
+			.VAlign(VAlign_Center)
+			.Padding(0.0f, 2.0f)
+			[
+				InlineTextBlock.ToSharedRef()
+			]
+		];
+	}
+
+private:
+	TWeakPtr<FFolderTreeItem> TreeItemPtr;
+
+	FText GetDisplayText() const
+	{
+		auto Folder = TreeItemPtr.Pin();
+		return Folder.IsValid() ? FText::FromName(Folder->LeafName) : FText();
+	}
+
+	const FSlateBrush* GetIcon() const
+	{	
+		auto TreeItem = TreeItemPtr.Pin();
+		if (!TreeItem.IsValid())
+		{
+			return FEditorStyle::Get().GetBrush(TEXT("SceneOutliner.FolderClosed"));
+		}
+
+		if (TreeItem->Flags.bIsExpanded && TreeItem->GetChildren().Num())
+		{
+			return FEditorStyle::Get().GetBrush(TEXT("SceneOutliner.FolderOpen"));
+		}
+		else
+		{
+			return FEditorStyle::Get().GetBrush(TEXT("SceneOutliner.FolderClosed"));
+		}
+	}
+
+	FSlateColor GetForegroundColor() const
+	{
+		if (auto BaseColor = FCommonLabelData::GetForegroundColor(TreeItemPtr.Pin()))
+		{
+			return BaseColor.GetValue();
+		}
+
+		return FSlateColor::UseForeground();
+	}
+
+	bool OnVerifyItemLabelChanged( const FText& InLabel, FText& OutErrorMessage)
+	{
+		auto TreeItem = TreeItemPtr.Pin();
+		if (!TreeItem.IsValid())
+		{
+			OutErrorMessage = LOCTEXT( "RenameFailed_TreeItemDeleted", "Tree item no longer exists" );
+			return false;
+		}
+
+		FText TrimmedLabel = FText::TrimPrecedingAndTrailing(InLabel);
+
+		if (TrimmedLabel.IsEmpty())
+		{
+			OutErrorMessage = LOCTEXT( "RenameFailed_LeftBlank", "Names cannot be left blank" );
+			return false;
+		}
+
+		if (TrimmedLabel.ToString().Len() >= NAME_SIZE)
+		{
+			FFormatNamedArguments Arguments;
+			Arguments.Add( TEXT("CharCount"), NAME_SIZE );
+			OutErrorMessage = FText::Format(LOCTEXT("RenameFailed_TooLong", "Names must be less than {CharCount} characters long."), Arguments);
+			return false;
+		}
+
+		FString LabelString = TrimmedLabel.ToString();
+		if (TreeItem->LeafName.ToString() == LabelString)
+		{
+			return true;
+		}
+
+		int32 Dummy = 0;
+		if (LabelString.FindChar('/', Dummy) || LabelString.FindChar('\\', Dummy))
+		{
+			OutErrorMessage = LOCTEXT("RenameFailed_InvalidChar", "Folder names cannot contain / or \\.");
+			return false;
+		}
+
+		// Validate that this folder doesn't exist already
+		FName NewPath = GetParentPath(TreeItem->Path);
+		if (NewPath.IsNone())
+		{
+			NewPath = FName(*LabelString);
+		}
+		else
+		{
+			NewPath = FName(*(NewPath.ToString() / LabelString));
+		}
+
+		if (FActorFolders::Get().GetFolderProperties(*GWorld, NewPath))
+		{
+			OutErrorMessage = LOCTEXT("RenameFailed_AlreadyExists", "A folder with this name already exists at this level");
+			return false;
+		}
+
+		return true;
+	}
+
+	void OnLabelCommitted(const FText& InLabel, ETextCommit::Type InCommitInfo)
+	{
+		auto TreeItem = TreeItemPtr.Pin();
+		if (TreeItem.IsValid() && !InLabel.ToString().Equals(TreeItem->LeafName.ToString(), ESearchCase::CaseSensitive))
+		{
+			// Rename the item
+			FName NewPath = GetParentPath(TreeItem->Path);
+			if (NewPath.IsNone())
+			{
+				NewPath = FName(*InLabel.ToString());
+			}
+			else
+			{
+				NewPath = FName(*(NewPath.ToString() / InLabel.ToString()));
+			}
+
+			FActorFolders::Get().RenameFolderInWorld(*GWorld, TreeItem->Path, NewPath);
+
+			auto Outliner = WeakSceneOutliner.Pin();
+			if (Outliner.IsValid())
+			{
+				Outliner->SetKeyboardFocus();
+			}
+		}
+	}
+};
+
+FName FItemLabelColumn::GetColumnID()
+{
+	return GetID();
+}
+
+SHeaderRow::FColumn::FArguments FItemLabelColumn::ConstructHeaderRowColumn()
+{
+	return SHeaderRow::Column(GetColumnID())
+		.DefaultLabel(LOCTEXT("ItemLabel_HeaderText", "Label").ToString())
+		.FillWidth( 5.0f );
+}
+const TSharedRef<SWidget> FItemLabelColumn::ConstructRowWidget(FTreeItemRef TreeItem, const STableRow<FTreeItemPtr>& Row)
+{
+	FColumnWidgetGenerator Generator(*this, Row);
+	TreeItem->Visit(Generator);
+	return Generator.Widget.ToSharedRef();
+}
+
+void FItemLabelColumn::PopulateSearchStrings( const ITreeItem& Item, TArray< FString >& OutSearchStrings ) const
+{
+	OutSearchStrings.Add(Item.GetDisplayString());
+}
+
+void FItemLabelColumn::SortItems(TArray<FTreeItemPtr>& OutItems, const EColumnSortMode::Type SortMode) const
+{
+	typedef FSortHelper<int32, FNumericStringWrapper> FSort;
+
+	FSort()
+		.Primary([](const ITreeItem& Item){ return Item.GetTypeSortPriority(); }, 	SortMode)
+		.Secondary([](const ITreeItem& Item){ return FNumericStringWrapper(Item.GetDisplayString()); }, 	SortMode)
+		.Sort(OutItems);
+}
+
+TSharedRef<SWidget> FItemLabelColumn::GenerateWidget( FActorTreeItem& TreeItem, const STableRow<FTreeItemPtr>& InRow )
+{
+	ISceneOutliner* Outliner = WeakSceneOutliner.Pin().Get();
+	check(Outliner);
+	return SNew(SActorTreeLabel, TreeItem, *Outliner, InRow);
+}
+
+TSharedRef<SWidget> FItemLabelColumn::GenerateWidget( FWorldTreeItem& TreeItem, const STableRow<FTreeItemPtr>& InRow )
+{
+	ISceneOutliner* Outliner = WeakSceneOutliner.Pin().Get();
+	check(Outliner);
+	return SNew(SWorldTreeLabel, TreeItem, *Outliner, InRow);
+}
+
+TSharedRef<SWidget> FItemLabelColumn::GenerateWidget( FLevelBlueprintTreeItem& TreeItem, const STableRow<FTreeItemPtr>& InRow )
+{
+	ISceneOutliner* Outliner = WeakSceneOutliner.Pin().Get();
+	check(Outliner);
+	return SNew(SLevelBlueprintTreeLabel, TreeItem, *Outliner, InRow);
+}
+
+TSharedRef<SWidget> FItemLabelColumn::GenerateWidget( FFolderTreeItem& TreeItem, const STableRow<FTreeItemPtr>& InRow )
+{
+	ISceneOutliner* Outliner = WeakSceneOutliner.Pin().Get();
+	check(Outliner);
+	return SNew(SFolderTreeLabel, TreeItem, *Outliner, InRow);
+}
+
+}	// namespace SceneOutliner
+
+#undef LOCTEXT_NAMESPACE
diff --git a/Engine/Source/Editor/SceneOutliner/Private/SceneOutlinerItemLabelColumn.h b/Engine/Source/Editor/SceneOutliner/Private/SceneOutlinerItemLabelColumn.h
new file mode 100644
--- /dev/null
+++ b/Engine/Source/Editor/SceneOutliner/Private/SceneOutlinerItemLabelColumn.h
@@ -0,0 +1,67 @@
+// Copyright 1998-2014 Epic Games, Inc. All Rights Reserved.
+
+#pragma once
+
+#include "SceneOutlinerFwd.h"
+
+#include "ISceneOutliner.h"
+#include "ISceneOutlinerColumn.h"
+
+namespace SceneOutliner
+{
+
+/** A column for the SceneOutliner that displays the item label */
+class FItemLabelColumn : public ISceneOutlinerColumn
+{
+
+public:
+	FItemLabelColumn(ISceneOutliner& SceneOutliner) : WeakSceneOutliner(StaticCastSharedRef<ISceneOutliner>(SceneOutliner.AsShared())) {}
+	
+	virtual ~FItemLabelColumn() {}
+
+	static FName GetID() { return FBuiltInColumnTypes::Label(); }
+
+	//////////////////////////////////////////////////////////////////////////
+	// Begin ISceneOutlinerColumn Implementation
+
+	virtual FName GetColumnID() override;
+
+	virtual SHeaderRow::FColumn::FArguments ConstructHeaderRowColumn() override;
+	
+	virtual const TSharedRef< SWidget > ConstructRowWidget( FTreeItemRef TreeItem, const STableRow<FTreeItemPtr>& Row ) override;
+
+	virtual void PopulateSearchStrings( const ITreeItem& Item, TArray< FString >& OutSearchStrings ) const override;
+
+	virtual bool SupportsSorting() const override { return true; }
+
+	virtual void SortItems(TArray<FTreeItemPtr>& RootItems, const EColumnSortMode::Type SortMode) const override;
+	
+	// End ISceneOutlinerColumn Implementation
+	//////////////////////////////////////////////////////////////////////////
+
+private:
+	TWeakPtr<ISceneOutliner> WeakSceneOutliner;
+
+private:
+
+	TSharedRef<SWidget> GenerateWidget( FActorTreeItem& TreeItem, const STableRow<FTreeItemPtr>& InRow );
+	TSharedRef<SWidget> GenerateWidget( FWorldTreeItem& TreeItem, const STableRow<FTreeItemPtr>& InRow );
+	TSharedRef<SWidget> GenerateWidget( FLevelBlueprintTreeItem& TreeItem, const STableRow<FTreeItemPtr>& InRow );
+	TSharedRef<SWidget> GenerateWidget( FFolderTreeItem& TreeItem, const STableRow<FTreeItemPtr>& InRow );
+
+	struct FColumnWidgetGenerator : FColumnGenerator
+	{
+		FItemLabelColumn& Column;
+		const STableRow<FTreeItemPtr>& Row;
+
+		FColumnWidgetGenerator(FItemLabelColumn& InColumn, const STableRow<FTreeItemPtr>& InRow) : Column(InColumn), Row(InRow) {}
+
+		virtual TSharedRef<SWidget> GenerateWidget(FFolderTreeItem& Item) const override { return Column.GenerateWidget(Item, Row); }
+		virtual TSharedRef<SWidget> GenerateWidget(FWorldTreeItem& Item) const override { return Column.GenerateWidget(Item, Row); }
+		virtual TSharedRef<SWidget> GenerateWidget(FLevelBlueprintTreeItem& Item) const override { return Column.GenerateWidget(Item, Row); }
+		virtual TSharedRef<SWidget> GenerateWidget(FActorTreeItem& Item) const override { return Column.GenerateWidget(Item, Row); }
+	};
+	friend FColumnWidgetGenerator;
+};
+
+}	// namespace SceneOutliner
\ No newline at end of file
diff --git a/Engine/Source/Editor/SceneOutliner/Private/SceneOutlinerModule.cpp b/Engine/Source/Editor/SceneOutliner/Private/SceneOutlinerModule.cpp
--- a/Engine/Source/Editor/SceneOutliner/Private/SceneOutlinerModule.cpp
+++ b/Engine/Source/Editor/SceneOutliner/Private/SceneOutlinerModule.cpp
@@ -1,37 +1,56 @@
 // Copyright 1998-2014 Epic Games, Inc. All Rights Reserved.
 
 #include "SceneOutlinerPrivatePCH.h"
-#include "SceneOutlinerTreeItems.h"
 #include "ModuleManager.h"
 
+#include "SceneOutlinerActorInfoColumn.h"
+#include "SceneOutlinerGutter.h"
+#include "SceneOutlinerItemLabelColumn.h"
 
 /* FSceneOutlinerModule interface
  *****************************************************************************/
 
 namespace SceneOutliner
 {
-	void OnSceneOutlinerItemClicked(TSharedRef<TOutlinerTreeItem> Item, FOnActorPicked OnActorPicked)
+	void OnSceneOutlinerItemClicked(TSharedRef<ITreeItem> Item, FOnActorPicked OnActorPicked)
 	{
-		if (Item->Type == TOutlinerTreeItem::Actor)
-		{
-			AActor* Actor = StaticCastSharedRef<TOutlinerActorTreeItem>(Item)->Actor.Get();
-			if (Actor)
-			{
-				OnActorPicked.ExecuteIfBound(Actor);
-			}
-		}
+		Item->Visit(
+			FFunctionalVisitor()
+			.Actor([&](const FActorTreeItem& ActorItem){
+				if (AActor* Actor = ActorItem.Actor.Get())
+				{
+					OnActorPicked.ExecuteIfBound(Actor);
+				}
+			})
+		);
 	}
 }
 
 
-TSharedRef< ISceneOutliner > FSceneOutlinerModule::CreateSceneOutliner( const FSceneOutlinerInitializationOptions& InitOptions, const FOnActorPicked& OnActorPickedDelegate ) const
+void FSceneOutlinerModule::StartupModule()
+{
+	RegisterColumnType< SceneOutliner::FSceneOutlinerGutter >();
+	RegisterColumnType< SceneOutliner::FItemLabelColumn >();
+	RegisterColumnType< SceneOutliner::FActorInfoColumn >();
+}
+
+
+void FSceneOutlinerModule::ShutdownModule()
+{
+	UnRegisterColumnType< SceneOutliner::FSceneOutlinerGutter >();
+	UnRegisterColumnType< SceneOutliner::FItemLabelColumn >();
+	UnRegisterColumnType< SceneOutliner::FActorInfoColumn >();
+}
+
+
+TSharedRef< ISceneOutliner > FSceneOutlinerModule::CreateSceneOutliner( const SceneOutliner::FInitializationOptions& InitOptions, const FOnActorPicked& OnActorPickedDelegate ) const
 {
 	auto OnItemPicked = FOnSceneOutlinerItemPicked::CreateStatic( &SceneOutliner::OnSceneOutlinerItemClicked, OnActorPickedDelegate );
 	return CreateSceneOutliner(InitOptions, OnItemPicked);
 }
 
 
-TSharedRef< ISceneOutliner > FSceneOutlinerModule::CreateSceneOutliner( const FSceneOutlinerInitializationOptions& InitOptions, const FOnSceneOutlinerItemPicked& OnItemPickedDelegate ) const
+TSharedRef< ISceneOutliner > FSceneOutlinerModule::CreateSceneOutliner( const SceneOutliner::FInitializationOptions& InitOptions, const FOnSceneOutlinerItemPicked& OnItemPickedDelegate ) const
 {
 	return SNew( SceneOutliner::SSceneOutliner, InitOptions )
 		.IsEnabled( FSlateApplication::Get().GetNormalExecutionAttribute() )
diff --git a/Engine/Source/Editor/SceneOutliner/Private/SceneOutlinerPrivatePCH.h b/Engine/Source/Editor/SceneOutliner/Private/SceneOutlinerPrivatePCH.h
--- a/Engine/Source/Editor/SceneOutliner/Private/SceneOutlinerPrivatePCH.h
+++ b/Engine/Source/Editor/SceneOutliner/Private/SceneOutlinerPrivatePCH.h
@@ -14,5 +14,4 @@
 
 #include "SceneOutlinerModule.h"
 #include "SSceneOutliner.h"
-#include "SOutlinerTreeView.h"
-#include "SSceneOutlinerItemWidget.h"
\ No newline at end of file
+#include "SOutlinerTreeView.h"
\ No newline at end of file
diff --git a/Engine/Source/Editor/SceneOutliner/Private/SceneOutlinerStandaloneTypes.cpp b/Engine/Source/Editor/SceneOutliner/Private/SceneOutlinerStandaloneTypes.cpp
new file mode 100644
--- /dev/null
+++ b/Engine/Source/Editor/SceneOutliner/Private/SceneOutlinerStandaloneTypes.cpp
@@ -0,0 +1,47 @@
+// Copyright 1998-2014 Epic Games, Inc. All Rights Reserved.
+
+#include "SceneOutlinerPrivatePCH.h"
+#include "SSceneOutliner.h"
+
+#include "EditorActorFolders.h"
+
+
+#define LOCTEXT_NAMESPACE "SceneOutlinerStandaloneTypes"
+
+namespace SceneOutliner
+{
+	/** Parse a new path (including leaf-name) into this tree item. Does not do any notification */
+	FName GetFolderLeafName(FName InPath)
+	{
+		FString PathString = InPath.ToString();
+		int32 LeafIndex = 0;
+		if (PathString.FindLastChar('/', LeafIndex))
+		{
+			return FName(*PathString.RightChop(LeafIndex + 1));
+		}
+		else
+		{
+			return InPath;
+		}
+	}
+
+	FName MoveFolderTo(FName InPath, FName NewParent, UWorld& World)
+	{
+		FName NewPath = GetFolderLeafName(InPath);
+
+		if (!NewParent.IsNone())
+		{
+			NewPath = FName(*(NewParent.ToString() / NewPath.ToString()));
+		}
+
+		if (FActorFolders::Get().RenameFolderInWorld(World, InPath, NewPath))
+		{
+			return NewPath;
+		}
+
+		return FName();
+	}
+
+}	// namespace SceneOutliner
+
+#undef LOCTEXT_NAMESPACE
