Added Worlds and Level Blueprints to the scene outliner

diff --git a/Engine/Source/Editor/SceneOutliner/Private/SceneOutlinerStandaloneTypes.h b/Engine/Source/Editor/SceneOutliner/Private/SceneOutlinerStandaloneTypes.h
new file mode 100644
--- /dev/null
+++ b/Engine/Source/Editor/SceneOutliner/Private/SceneOutlinerStandaloneTypes.h
@@ -0,0 +1,195 @@
+// Copyright 1998-2014 Epic Games, Inc. All Rights Reserved.
+
+#pragma once
+
+#include "SceneOutlinerFwd.h"
+
+#include "SceneOutlinerDragDrop.h"
+#include "SceneOutlinerFilters.h"
+
+namespace SceneOutliner
+{
+	#if PLATFORM_COMPILER_HAS_VARIADIC_TEMPLATES
+
+		template <typename...> struct TMaxSizeof;
+		template <> struct TMaxSizeof<> { static const uint32 Value = 0; };
+		template <typename T, typename... TRest> struct TMaxSizeof<T, TRest...> { static const uint32 Value = sizeof(T) > TMaxSizeof<TRest...>::Value ? sizeof(T) : TMaxSizeof<TRest...>::Value; };
+
+	#else
+
+		template <typename T1, typename T2=void, typename T3=void, typename T4=void, typename T5=void, typename T6=void>
+		struct TMaxSizeof
+		{
+			static const uint32 Value = sizeof(T1) > TMaxSizeof<T2, T3, T4, T5, T6>::Value ? sizeof(T1) : TMaxSizeof<T2, T3, T4, T5, T6>::Value;
+		};
+		template <> struct TMaxSizeof<void,void,void,void,void,void> { static const uint32 Value = 0; };
+		
+	#endif
+
+	/** Variant type that defines an identifier for a tree item. Assumes 'trivial relocatability' as with many unreal containers. */
+	struct FTreeItemID
+	{
+	public:
+		/** Default constructed null item ID */
+		FTreeItemID() : Type(EType::Null), CachedHash(0) {}
+
+		/** ID representing an actor */
+		FTreeItemID(const TWeakObjectPtr<AActor>& InActor) : Type(EType::Actor)
+		{
+			new (Data) TWeakObjectPtr<AActor>(InActor);
+			CachedHash = CalculateTypeHash();
+		}
+		FTreeItemID(const AActor* InActor) : Type(EType::Actor)
+		{
+			new (Data) TWeakObjectPtr<AActor>(InActor);
+			CachedHash = CalculateTypeHash();
+		}
+
+		/** ID representing a world */
+		FTreeItemID(const TWeakObjectPtr<UWorld>& InWorld) : Type(EType::World)
+		{
+			new (Data) TWeakObjectPtr<UWorld>(InWorld);
+			CachedHash = CalculateTypeHash();
+		}
+		FTreeItemID(const UWorld* InWorld) : Type(EType::World)
+		{
+			new (Data) TWeakObjectPtr<UWorld>(InWorld);
+			CachedHash = CalculateTypeHash();
+		}
+
+		/** ID representing a level blueprint */
+		explicit FTreeItemID(const FLevelBlueprintHandle& InLevelBlueprint) : Type(EType::LevelBlueprint)
+		{
+			new (Data) FLevelBlueprintHandle(InLevelBlueprint);
+			CachedHash = CalculateTypeHash();
+		}
+
+		/** ID representing a folder */
+		FTreeItemID(const FName& InFolder) : Type(EType::Folder)
+		{
+			new (Data) FName(InFolder);
+			CachedHash = CalculateTypeHash();
+		}
+
+		/** Copy construction / assignment */
+		FTreeItemID(const FTreeItemID& Other)
+		{
+			*this = Other;
+		}
+		FTreeItemID& operator=(const FTreeItemID& Other)
+		{
+			Type = Other.Type;
+			switch(Type)
+			{
+				case EType::Actor:			new (Data) TWeakObjectPtr<AActor>(Other.GetAsActorRef());								break;
+				case EType::World:			new (Data) TWeakObjectPtr<UWorld>(Other.GetAsWorldRef());								break;
+				case EType::LevelBlueprint:	new (Data) FLevelBlueprintHandle(Other.GetAsLevelBlueprintRef());						break;
+				case EType::Folder:			new (Data) FName(Other.GetAsFolderRef());												break;
+				default:																											break;
+			}
+
+			CachedHash = CalculateTypeHash();
+			return *this;
+		}
+
+		/** Move construction / assignment */
+		FTreeItemID(FTreeItemID&& Other)
+		{
+			*this = MoveTemp(Other);
+		}
+		FTreeItemID& operator=(FTreeItemID&& Other)
+		{
+			FMemory::Memswap(this, &Other, sizeof(FTreeItemID));
+			return *this;
+		}
+
+		~FTreeItemID()
+		{
+			switch(Type)
+			{
+				case EType::Actor:			GetAsActorRef().~TWeakObjectPtr();						break;
+				case EType::World:			GetAsWorldRef().~TWeakObjectPtr();						break;
+				case EType::LevelBlueprint:	GetAsLevelBlueprintRef().~FLevelBlueprintHandle();		break;
+				case EType::Folder:			GetAsFolderRef().~FName();								break;
+				default:																			break;
+			}
+		}
+
+		friend bool operator==(const FTreeItemID& One, const FTreeItemID& Other)
+		{
+			return One.Type == Other.Type && One.CachedHash == Other.CachedHash && One.Compare(Other);
+		}
+		friend bool operator!=(const FTreeItemID& One, const FTreeItemID& Other)
+		{
+			return One.Type != Other.Type || One.CachedHash != Other.CachedHash || !One.Compare(Other);
+		}
+
+		uint32 CalculateTypeHash() const
+		{
+			uint32 Hash = 0;
+			switch(Type)
+			{
+				case EType::Actor:			Hash = GetTypeHash(GetAsActorRef());				break;
+				case EType::World:			Hash = GetTypeHash(GetAsWorldRef());				break;
+				case EType::LevelBlueprint:	Hash = GetTypeHash(GetAsLevelBlueprintRef());		break;
+				case EType::Folder:			Hash = GetTypeHash(GetAsFolderRef());				break;
+				default:																		break;
+			}
+
+			return HashCombine((uint8)Type, Hash);
+		}
+
+		friend uint32 GetTypeHash(const FTreeItemID& ItemID)
+		{
+			return ItemID.CachedHash;
+		}
+
+	private:
+
+		TWeakObjectPtr<AActor>& 				GetAsActorRef() const 			{ return *reinterpret_cast<TWeakObjectPtr<AActor>*>(Data); }
+		TWeakObjectPtr<UWorld>& 				GetAsWorldRef() const 			{ return *reinterpret_cast<TWeakObjectPtr<UWorld>*>(Data); }
+		FLevelBlueprintHandle& 					GetAsLevelBlueprintRef() const 	{ return *reinterpret_cast<FLevelBlueprintHandle*>(Data); }
+		FName& 									GetAsFolderRef() const			{ return *reinterpret_cast<FName*>(Data); }
+
+		/** Compares the specified ID with this one - assumes matching types */
+		bool Compare(const FTreeItemID& Other) const
+		{
+			switch(Type)
+			{
+				case EType::Actor:			return GetAsActorRef() == Other.GetAsActorRef();
+				case EType::World:			return GetAsWorldRef() == Other.GetAsWorldRef();
+				case EType::LevelBlueprint:	return GetAsLevelBlueprintRef() == Other.GetAsLevelBlueprintRef();
+				case EType::Folder:			return GetAsFolderRef() == Other.GetAsFolderRef();
+				case EType::Null:			return true;
+				default: check(false);		return false;
+			}
+		}
+
+		enum class EType : uint8 { Actor, World, LevelBlueprint, Folder, Null };
+		EType Type;
+
+		uint32 CachedHash;
+		static const uint32 MaxSize = TMaxSizeof<TWeakObjectPtr<UWorld>, TWeakObjectPtr<ULevelScriptBlueprint>, TWeakObjectPtr<AActor>, FName>::Value;
+		mutable uint8 Data[MaxSize];
+	};
+
+	struct ETreeItemSortOrder
+	{
+		enum Type { World = 0, Folder = 10, Actor = 20 };
+	};
+
+	/** Folder-specific helper functions */
+
+	/** Parse a new path (including leaf-name) into this tree item. Does not do any notification */
+	FName GetFolderLeafName(FName InPath);
+
+	/** Get the parent path for the specified folder path */
+	FORCEINLINE FName GetParentPath(FName Path)
+	{
+		return FName(*FPaths::GetPath(Path.ToString()));
+	}
+
+	/** Move the specified folder path to the specified new parent */
+	FName MoveFolderTo(FName Path, FName NewParent, UWorld& World);
+
+}	// namespace SceneOutliner
\ No newline at end of file
diff --git a/Engine/Source/Editor/SceneOutliner/Private/WorldTreeItem.cpp b/Engine/Source/Editor/SceneOutliner/Private/WorldTreeItem.cpp
new file mode 100644
--- /dev/null
+++ b/Engine/Source/Editor/SceneOutliner/Private/WorldTreeItem.cpp
@@ -0,0 +1,83 @@
+// Copyright 1998-2014 Epic Games, Inc. All Rights Reserved.
+
+#include "SceneOutlinerPrivatePCH.h"
+
+#include "WorldTreeItem.h"
+#include "LevelEditor.h"
+
+#define LOCTEXT_NAMESPACE "SceneOutliner_WorldTreeItem"
+
+namespace SceneOutliner
+{
+
+FWorldTreeItem::FWorldTreeItem(UWorld* InWorld)
+	: World(InWorld)
+{
+
+}
+
+FTreeItemPtr FWorldTreeItem::FindParent(const FTreeItemMap& ExistingItems) const
+{
+	return nullptr;
+}
+
+FTreeItemPtr FWorldTreeItem::CreateParent() const
+{
+	return nullptr;
+}
+
+void FWorldTreeItem::Visit(const ITreeItemVisitor& Visitor) const
+{
+	Visitor.Visit(*this);
+}
+
+void FWorldTreeItem::Visit(const IMutableTreeItemVisitor& Visitor)
+{
+	Visitor.Visit(*this);
+}
+
+FTreeItemID FWorldTreeItem::GetID() const
+{
+	return FTreeItemID(World);
+}
+
+FString FWorldTreeItem::GetDisplayString() const
+{
+	if (UWorld* WorldPtr = World.Get())
+	{
+		return World->GetName();
+	}
+	return FString();
+}
+
+int32 FWorldTreeItem::GetTypeSortPriority() const
+{
+	return ETreeItemSortOrder::World;
+}
+
+bool FWorldTreeItem::CanInteract() const
+{
+	return !Flags.bIsFilteredOut;
+}
+
+void FWorldTreeItem::GenerateContextMenu(FMenuBuilder& MenuBuilder, SSceneOutliner& Outliner)
+{
+	const FSlateIcon WorldSettingsIcon(FEditorStyle::GetStyleSetName(), "LevelEditor.WorldProperties.Tab");
+
+	MenuBuilder.AddMenuEntry(LOCTEXT("OpenWorldSettings", "World Settings"), FText(), WorldSettingsIcon, FExecuteAction::CreateSP(this, &FWorldTreeItem::OpenWorldSettings));
+}
+
+FDragValidationInfo FWorldTreeItem::ValidateDrop(FDragDropPayload& DraggedObjects, UWorld& World) const
+{
+	return FDragValidationInfo::Invalid();
+}
+
+void FWorldTreeItem::OpenWorldSettings() const
+{
+	FLevelEditorModule& LevelEditorModule = FModuleManager::GetModuleChecked<FLevelEditorModule>( TEXT("LevelEditor") );
+	LevelEditorModule.GetLevelEditorTabManager()->InvokeTab(FName("WorldSettingsTab"));	
+}
+
+}		// namespace SceneOutliner
+
+#undef LOCTEXT_NAMESPACE
\ No newline at end of file
diff --git a/Engine/Source/Editor/SceneOutliner/Public/ActorTreeItem.h b/Engine/Source/Editor/SceneOutliner/Public/ActorTreeItem.h
new file mode 100644
--- /dev/null
+++ b/Engine/Source/Editor/SceneOutliner/Public/ActorTreeItem.h
@@ -0,0 +1,87 @@
+// Copyright 1998-2014 Epic Games, Inc. All Rights Reserved.
+
+#pragma once
+
+#include "ITreeItem.h"
+
+namespace SceneOutliner
+{
+	/** Helper class to manage moving abritrary data onto an actor */
+	struct FActorDropTarget : IDropTarget
+	{
+		/** The actor this tree item is associated with. */
+		TWeakObjectPtr<AActor> Actor;
+		
+		/** Construct this object out of an Actor */
+		FActorDropTarget(AActor* InActor) : Actor(InActor) {}
+
+	public:
+
+		/** Called to test whether the specified payload can be dropped onto this tree item */
+		virtual FDragValidationInfo ValidateDrop(FDragDropPayload& DraggedObjects, UWorld& World) const override;
+
+		/** Called to drop the specified objects on this item. Only called if ValidateDrop() allows. */
+		virtual void OnDrop(FDragDropPayload& DraggedObjects, UWorld& World, const FDragValidationInfo& ValidationInfo, TSharedRef<SWidget> DroppedOnWidget) override;
+
+	protected:
+
+		/** Attach our actor to the specified parent and socket */
+		static void PerformAttachment(FName SocketName, TWeakObjectPtr<AActor> Parent, const FActorArray NewAttachments);
+
+		/** Detach the specified actor from this actor */
+		void DetachActorFromParent(AActor* ChildActor);
+
+	};
+
+	/** A tree item that represents an actor in the world */
+	struct FActorTreeItem : ITreeItem
+	{
+		/** The actor this tree item is associated with. */
+		mutable TWeakObjectPtr<AActor> Actor;
+
+		/** Construct this item from an actor */
+		FActorTreeItem(AActor* InActor);
+
+		/** Get this item's parent item. It is valid to return nullptr if this item has no parent */
+		virtual FTreeItemPtr FindParent(const FTreeItemMap& ExistingItems) const override;
+
+		/** Create this item's parent. It is valid to return nullptr if this item has no parent */
+		virtual FTreeItemPtr CreateParent() const override;
+
+ 		/** Visit this tree item */
+ 		virtual void Visit(const ITreeItemVisitor& Visitor) const override;
+ 		virtual void Visit(const IMutableTreeItemVisitor& Visitor) override;
+
+	public:
+
+		/** Get the ID that represents this tree item. Used to reference this item in a map */
+		virtual FTreeItemID GetID() const override;
+
+		/** Get the raw string to display for this tree item - used for sorting */
+		virtual FString GetDisplayString() const override;
+
+		/** Get the sort priority given to this item's type */
+		virtual int32 GetTypeSortPriority() const override;
+
+		/** Check whether it should be possible to interact with this tree item */
+		virtual bool CanInteract() const override;
+
+	public:
+
+		/** Populate the specified drag/drop payload with any relevant information for this type */
+		virtual void PopulateDragDropPayload(FDragDropPayload& Payload) const override;
+
+		/** Called to test whether the specified payload can be dropped onto this tree item */
+		virtual FDragValidationInfo ValidateDrop(FDragDropPayload& DraggedObjects, UWorld& World) const override;
+
+		/** Called to drop the specified objects on this item. Only called if ValidateDrop() allows. */
+		virtual void OnDrop(FDragDropPayload& DraggedObjects, UWorld& World, const FDragValidationInfo& ValidationInfo, TSharedRef<SWidget> DroppedOnWidget) override;
+
+	public:
+
+		/** true if this item exists in both the current world and PIE. */
+		bool bExistsInCurrentWorldAndPIE;
+
+	};
+
+}		// namespace SceneOutliner
\ No newline at end of file
diff --git a/Engine/Source/Editor/SceneOutliner/Public/FolderTreeItem.h b/Engine/Source/Editor/SceneOutliner/Public/FolderTreeItem.h
new file mode 100644
--- /dev/null
+++ b/Engine/Source/Editor/SceneOutliner/Public/FolderTreeItem.h
@@ -0,0 +1,89 @@
+// Copyright 1998-2014 Epic Games, Inc. All Rights Reserved.
+
+#pragma once
+
+#include "ITreeItem.h"
+
+namespace SceneOutliner
+{
+	/** Helper class to manage moving abritrary data onto a folder */
+	struct FFolderDropTarget : IDropTarget
+	{
+		/** The path that we are dropping on */
+		FName DestinationPath;
+
+		/** Constructor that takes a path to this folder (including leaf-name) */
+		FFolderDropTarget(FName InDestinationPath) : DestinationPath(InDestinationPath) {}
+
+	public:
+
+		/** Called to test whether the specified payload can be dropped onto this tree item */
+		virtual FDragValidationInfo ValidateDrop(FDragDropPayload& DraggedObjects, UWorld& World) const override;
+
+		/** Called to drop the specified objects on this item. Only called if ValidateDrop() allows. */
+		virtual void OnDrop(FDragDropPayload& DraggedObjects, UWorld& World, const FDragValidationInfo& ValidationInfo, TSharedRef<SWidget> DroppedOnWidget) override;
+	};
+
+	/** A tree item that represents a folder in the world */
+	struct FFolderTreeItem : ITreeItem
+	{
+		/** The path of this folder. / separated. */
+		FName Path;
+
+		/** The leaf name of this folder */
+		FName LeafName;
+
+		/** Constructor that takes a path to this folder (including leaf-name) */
+		FFolderTreeItem(FName InPath);
+
+		/** Get this item's parent item. It is valid to return nullptr if this item has no parent */
+		virtual FTreeItemPtr FindParent(const FTreeItemMap& ExistingItems) const override;
+
+		/** Create this item's parent. It is valid to return nullptr if this item has no parent */
+		FTreeItemPtr CreateParent() const override;
+
+ 		/** Visit this tree item */
+ 		virtual void Visit(const ITreeItemVisitor& Visitor) const override;
+ 		virtual void Visit(const IMutableTreeItemVisitor& Visitor) override;
+ 		
+	public:
+
+		/** Called when the expansion changes on this folder */
+		virtual void OnExpansionChanged() override;
+
+		/** Get the ID that represents this tree item. Used to reference this item in a map */
+		virtual FTreeItemID GetID() const override;
+
+		/** Get the raw string to display for this tree item - used for sorting */
+		virtual FString GetDisplayString() const override;
+
+		/** Get the sort priority given to this item's type */
+		virtual int32 GetTypeSortPriority() const override;
+		
+		/** Check whether it should be possible to interact with this tree item */
+		virtual bool CanInteract() const override;
+		
+		/** Generate a context menu for this item. Only called if *only* this item is selected. */
+		virtual void GenerateContextMenu(FMenuBuilder& MenuBuilder, SSceneOutliner& Outliner) override;
+
+		/** Populate the specified drag/drop payload with any relevant information for this type */
+		virtual void PopulateDragDropPayload(FDragDropPayload& Payload) const override;
+
+	public:
+
+		/** Called to test whether the specified payload can be dropped onto this tree item */
+		virtual FDragValidationInfo ValidateDrop(FDragDropPayload& DraggedObjects, UWorld& World) const override;
+
+		/** Called to drop the specified objects on this item. Only called if ValidateDrop() allows. */
+		virtual void OnDrop(FDragDropPayload& DraggedObjects, UWorld& World, const FDragValidationInfo& ValidationInfo, TSharedRef<SWidget> DroppedOnWidget) override;
+
+		/** Delete this folder */
+		void Delete();
+
+	private:
+
+		/** Create a new folder as a child of this one */
+		void CreateSubFolder(TWeakPtr<SSceneOutliner> WeakOutliner);
+	};
+
+}		// namespace SceneOutliner
\ No newline at end of file
diff --git a/Engine/Source/Editor/SceneOutliner/Public/ISceneOutliner.h b/Engine/Source/Editor/SceneOutliner/Public/ISceneOutliner.h
--- a/Engine/Source/Editor/SceneOutliner/Public/ISceneOutliner.h
+++ b/Engine/Source/Editor/SceneOutliner/Public/ISceneOutliner.h
@@ -2,18 +2,27 @@
 
 #pragma once
 
+#include "SceneOutlinerFwd.h"
 
 /**
  * The public interface for the Scene Outliner widget
  */
 class ISceneOutliner : public SCompoundWidget
 {
-
 public:
 
 	/** Sends a requests to the Scene Outliner to refresh itself the next chance it gets */
 	virtual void Refresh() = 0;
 
 	/** @return Returns a string to use for highlighting results in the outliner list */
-	virtual FText GetFilterHighlightText() const = 0;
+	virtual TAttribute<FText> GetFilterHighlightText() const = 0;
+
+	/** @return Returns the common data for this outliner */
+	virtual const SceneOutliner::FSharedOutlinerData& GetSharedData() const = 0;
+
+	/** Get a const reference to the actual tree hierarchy */
+	virtual const STreeView<SceneOutliner::FTreeItemPtr>& GetTree() const = 0;
+
+	/** Set the keyboard focus to the outliner */
+	virtual void SetKeyboardFocus() = 0;
 };
\ No newline at end of file
diff --git a/Engine/Source/Editor/SceneOutliner/Public/ISceneOutlinerColumn.h b/Engine/Source/Editor/SceneOutliner/Public/ISceneOutlinerColumn.h
--- a/Engine/Source/Editor/SceneOutliner/Public/ISceneOutlinerColumn.h
+++ b/Engine/Source/Editor/SceneOutliner/Public/ISceneOutlinerColumn.h
@@ -2,29 +2,29 @@
 
 #pragma once
 
-namespace SceneOutliner
-{
-	struct TOutlinerTreeItem;
-}
+#include "SceneOutlinerFwd.h"
 
 /**
- *	
+ *	Interface for a scene outliner column
  */
 class ISceneOutlinerColumn : public TSharedFromThis< ISceneOutlinerColumn >
 {
-
 public:
+
 	virtual FName GetColumnID() = 0;
 
 	virtual SHeaderRow::FColumn::FArguments ConstructHeaderRowColumn() = 0;
 
-	virtual const TSharedRef< SWidget > ConstructRowWidget( const TSharedRef<SceneOutliner::TOutlinerTreeItem> TreeItem ) = 0;
+	virtual const TSharedRef< SWidget > ConstructRowWidget( SceneOutliner::FTreeItemRef TreeItem, const STableRow<SceneOutliner::FTreeItemPtr>& Row ) = 0;
 
-	virtual bool ProvidesSearchStrings() = 0;
+	virtual void Tick(double InCurrentTime, float InDeltaTime) {}
+
+public:
+	/** Optionally overridden interface methods */
+	virtual void PopulateSearchStrings( const SceneOutliner::ITreeItem& Item, TArray< FString >& OutSearchStrings ) const {}
 
-	virtual void PopulateActorSearchStrings( const AActor* const Actor, OUT TArray< FString >& OutSearchStrings ) const = 0;
+	virtual bool SupportsSorting() const { return false; }
 
-	virtual bool SupportsSorting() const = 0;
+	virtual void SortItems(TArray<SceneOutliner::FTreeItemPtr>& OutItems, const EColumnSortMode::Type SortMode) const {}
+};
 
-	virtual void SortItems(TArray<TSharedPtr<SceneOutliner::TOutlinerTreeItem>>& RootItems, const EColumnSortMode::Type SortMode) const = 0;
-};
\ No newline at end of file
diff --git a/Engine/Source/Editor/SceneOutliner/Public/ITreeItem.h b/Engine/Source/Editor/SceneOutliner/Public/ITreeItem.h
new file mode 100644
--- /dev/null
+++ b/Engine/Source/Editor/SceneOutliner/Public/ITreeItem.h
@@ -0,0 +1,150 @@
+// Copyright 1998-2014 Epic Games, Inc. All Rights Reserved.
+
+#pragma once
+
+#include "SceneOutlinerFwd.h"
+#include "SceneOutlinerVisitorTypes.h"
+
+namespace SceneOutliner
+{
+	/** Interface used for validating movement (ie, drag/dropping) operations */
+	struct IDropTarget
+	{
+		/** Called to test whether the specified payload can be dropped onto this tree item */
+		virtual FDragValidationInfo ValidateDrop(FDragDropPayload& DraggedObjects, UWorld& World) const = 0;
+
+		/** Called to drop the specified objects on this item. Only called if ValidateDrop() allows. */
+		virtual void OnDrop(FDragDropPayload& DraggedObjects, UWorld& World, const FDragValidationInfo& ValidationInfo, TSharedRef<SWidget> DroppedOnWidget) = 0;
+	};
+
+	/** Base tree item interface  */
+	struct ITreeItem : IDropTarget, TSharedFromThis<ITreeItem>
+	{
+		/** Friendship required for access to various internals */
+		friend SSceneOutliner;
+
+		/* Flags structure */
+		struct FlagsType
+		{
+			/** Whether this item is expanded or not */
+			bool bIsExpanded : 1;
+
+			/* true if this item is filtered out */
+			bool bIsFilteredOut : 1;
+
+			/** true if this item's children need to be sorted */
+			bool bChildrenRequireSort : 1;
+
+			/** Default constructor */		
+			FlagsType() : bIsExpanded(1), bIsFilteredOut(0) {}
+		};
+
+	public:
+		
+		/** Flags for this item */
+		FlagsType Flags;
+
+		/** Delegate for hooking up an inline editable text block to be notified that a rename is requested. */
+		DECLARE_DELEGATE( FOnRenameRequest );
+
+		/** Broadcasts whenever a rename is requested */
+		FOnRenameRequest RenameRequestEvent;
+
+	protected:
+
+		/** Default constructor */
+		ITreeItem() : SharedData(nullptr), Parent(nullptr) {}
+		virtual ~ITreeItem() {}
+
+		/** Data that is common between all outliner items - owned by the Outliner itself */
+		TSharedPtr<FSharedOutlinerData> SharedData;
+
+		/** This item's parent, if any. */
+		mutable TWeakPtr<ITreeItem> Parent;
+
+		/** Array of children contained underneath this item */
+		mutable TArray<TWeakPtr<ITreeItem>> Children;
+
+	public:
+
+		const FSharedOutlinerData& GetSharedData() const
+		{
+			checkSlow(SharedData.IsValid());
+			return *SharedData;
+		}
+
+	public:
+
+		/** Get this item's parent. Can be nullptr. */
+		FTreeItemPtr GetParent() const
+		{
+			return Parent.Pin();
+		}
+
+		/** Add a child to this item */
+		void AddChild(FTreeItemRef Child)
+		{
+			Child->Parent = AsShared();
+			Children.Add(MoveTemp(Child));
+		}
+
+		/** Remove a child from this item */
+		void RemoveChild(const FTreeItemRef& Child)
+		{
+			if (Children.Remove(Child))
+			{
+				Child->Parent = nullptr;	
+			}
+		}
+
+		/** Get this item's children, if any. Although we store as weak pointers, they are guaranteed to be valid. */
+		FORCEINLINE const TArray<TWeakPtr<ITreeItem>>& GetChildren() const
+		{
+			return Children;
+		}
+
+		/** Find this item's parent in the specified map. It is valid to return nullptr if this item has no parent */
+		virtual FTreeItemPtr FindParent(const FTreeItemMap& ExistingItems) const = 0;
+
+		/** Create this item's parent. It is valid to return nullptr if this item has no parent */
+		virtual FTreeItemPtr CreateParent() const = 0;
+
+	public:
+
+	 	/** Visit this tree item */
+	 	virtual void Visit(const ITreeItemVisitor& Visitor) const = 0;
+	 	virtual void Visit(const IMutableTreeItemVisitor& Visitor) = 0;
+
+	 	/** Get some data from this tree item using a 'getter' visitor. */
+	 	template<typename T>
+	 	T Get(const TTreeItemGetter<T>& Getter) const
+	 	{
+	 		Visit(Getter);
+	 		return Getter.Result();
+	 	}
+
+	public:
+
+		/** Get the ID that represents this tree item. Used to reference this item in a map */
+		virtual FTreeItemID GetID() const = 0;
+
+		/** Get the raw string to display for this tree item - used for sorting */
+		virtual FString GetDisplayString() const = 0;
+
+		/** Get the sort priority given to this item's type */
+		virtual int32 GetTypeSortPriority() const = 0;
+
+		/** Check whether it should be possible to interact with this tree item */
+		virtual bool CanInteract() const = 0;
+		
+		/** Called when this item is expanded or collapsed */
+		virtual void OnExpansionChanged() {};
+
+		/** Generate a context menu for this item. Only called if *only* this item is selected. */
+		virtual void GenerateContextMenu(FMenuBuilder& MenuBuilder, SSceneOutliner& Outliner) {}
+
+		/** Populate the specified drag/drop payload with any relevant information for this type */
+		virtual void PopulateDragDropPayload(FDragDropPayload& Payload) const = 0;
+	};
+
+}	// namespace SceneOutliner
\ No newline at end of file
diff --git a/Engine/Source/Editor/SceneOutliner/Public/LevelBlueprintTreeItem.h b/Engine/Source/Editor/SceneOutliner/Public/LevelBlueprintTreeItem.h
new file mode 100644
--- /dev/null
+++ b/Engine/Source/Editor/SceneOutliner/Public/LevelBlueprintTreeItem.h
@@ -0,0 +1,101 @@
+// Copyright 1998-2014 Epic Games, Inc. All Rights Reserved.
+
+#pragma once
+
+#include "ITreeItem.h"
+
+namespace SceneOutliner
+{
+	/** Custom type to distinguish levels from level blueprints */
+	struct FLevelBlueprintHandle
+	{
+		/** We reference level blueprints by their parent level, as they don't always exist */
+		TWeakObjectPtr<ULevel> ParentLevel;
+
+		/** Construct this handle from a level */
+		explicit FLevelBlueprintHandle(ULevel* InLevel) : ParentLevel(InLevel) {}
+		explicit FLevelBlueprintHandle(const TWeakObjectPtr<ULevel>& InLevel) : ParentLevel(InLevel) {}
+
+		/** Get the level blueprint for the parent level (if it exists) */
+		ULevelScriptBlueprint* GetLevelBlueprint()
+		{
+			if (auto* Level = ParentLevel.Get())
+			{
+				const bool bCreateIfNotFound = false;
+				return Level->GetLevelScriptBlueprint(!bCreateIfNotFound);
+			}
+			return nullptr;
+		}
+
+		/** Create the level blueprint for the parent level (if it doesn't exist) */
+		ULevelScriptBlueprint* GetOrCreateLevelBlueprint()
+		{
+			if (auto* Level = ParentLevel.Get())
+			{
+				const bool bCreateIfNotFound = true;
+				return Level->GetLevelScriptBlueprint(!bCreateIfNotFound);
+			}
+			return nullptr;
+		}
+
+		friend uint32 GetTypeHash(const FLevelBlueprintHandle& Handle)
+		{
+			return GetTypeHash(Handle.ParentLevel);
+		}
+
+		bool operator==(const FLevelBlueprintHandle& Other) const { return ParentLevel == Other.ParentLevel; }
+		bool operator!=(const FLevelBlueprintHandle& Other) const { return ParentLevel != Other.ParentLevel; }
+	};
+
+	/** A tree item that represents a level blueprint */
+	struct FLevelBlueprintTreeItem : ITreeItem
+	{
+		/** The level containing the blueprint to which this tree item relates. */
+		mutable FLevelBlueprintHandle Handle;
+
+		/** Construct this item from a level or handle */
+		FLevelBlueprintTreeItem(const FLevelBlueprintHandle& InHandle);
+
+		/** Get this item's parent item. Always returns nullptr. */
+		virtual FTreeItemPtr FindParent(const FTreeItemMap& ExistingItems) const override;
+
+		/** Create this item's parent. Always returns nullptr. */
+		virtual FTreeItemPtr CreateParent() const override;
+
+ 		/** Visit this tree item */
+ 		virtual void Visit(const ITreeItemVisitor& Visitor) const override;
+ 		virtual void Visit(const IMutableTreeItemVisitor& Visitor) override;
+
+	public:
+
+		/** Get the ID that represents this tree item. Used to reference this item in a map */
+		virtual FTreeItemID GetID() const override;
+
+		/** Get the raw string to display for this tree item - used for sorting */
+		virtual FString GetDisplayString() const override;
+
+		/** Get the sort priority given to this item's type */
+		virtual int32 GetTypeSortPriority() const override;
+
+		/** Check whether it should be possible to interact with this tree item */
+		virtual bool CanInteract() const override;
+
+		/** Generate a context menu for this item. Only called if *only* this item is selected. */
+		virtual void GenerateContextMenu(FMenuBuilder& MenuBuilder, SSceneOutliner& Outliner) override;
+
+		/** Populate the specified drag/drop payload with any relevant information for this type */
+		virtual void PopulateDragDropPayload(FDragDropPayload& Payload) const override { }
+
+		/** Called to test whether the specified payload can be dropped onto this tree item */
+		virtual FDragValidationInfo ValidateDrop(FDragDropPayload& DraggedObjects, UWorld& World) const override;
+
+		/** Called to drop the specified objects on this item. Only called if ValidateDrop() allows. */
+		virtual void OnDrop(FDragDropPayload& DraggedObjects, UWorld& World, const FDragValidationInfo& ValidationInfo, TSharedRef<SWidget> DroppedOnWidget) override { }
+
+	public:
+
+		/** Open the level blueprint contained by this tree item */
+		void Open() const;
+	};
+
+}		// namespace SceneOutliner
\ No newline at end of file
diff --git a/Engine/Source/Editor/SceneOutliner/Public/SceneOutliner.h b/Engine/Source/Editor/SceneOutliner/Public/SceneOutliner.h
new file mode 100644
--- /dev/null
+++ b/Engine/Source/Editor/SceneOutliner/Public/SceneOutliner.h
@@ -0,0 +1,22 @@
+// Copyright 1998-2014 Epic Games, Inc. All Rights Reserved.
+
+#pragma once
+
+#include "SceneOutlinerFwd.h"
+
+#include "ISceneOutliner.h"
+#include "ISceneOutlinerColumn.h"
+
+#include "SceneOutlinerPublicTypes.h"
+
+#include "SceneOutlinerModule.h"
+
+#include "SceneOutlinerVisitorTypes.h"
+#include "ITreeItem.h"
+
+#include "ActorTreeItem.h"
+#include "FolderTreeItem.h"
+#include "WorldTreeItem.h"
+#include "LevelBlueprintTreeItem.h"
+
+#include "SceneOutlinerFilters.h"
\ No newline at end of file
diff --git a/Engine/Source/Editor/SceneOutliner/Public/SceneOutlinerDragDrop.h b/Engine/Source/Editor/SceneOutliner/Public/SceneOutlinerDragDrop.h
new file mode 100644
--- /dev/null
+++ b/Engine/Source/Editor/SceneOutliner/Public/SceneOutlinerDragDrop.h
@@ -0,0 +1,133 @@
+// Copyright 1998-2014 Epic Games, Inc. All Rights Reserved.
+
+#pragma once
+
+#include "SceneOutlinerFwd.h"
+#include "ActorDragDropGraphEdOp.h"
+
+namespace SceneOutliner
+{
+	/** Consilidated drag/drop information parsed for the scene outliner */
+	struct FDragDropPayload
+	{
+		/** Default constructor, resulting in unset contents */
+		FDragDropPayload();
+
+		/** Populate this payload from an array of tree items */
+		template<typename TItem>
+		FDragDropPayload(const TArray<TItem>& InDraggedItems)
+		{
+			for (const auto& Item : InDraggedItems)
+			{
+				Item->PopulateDragDropPayload(*this);
+			}
+		}
+
+		/** Optional array of dragged folders */
+		TOptional<FFolderPaths> Folders;
+
+		/** OPtional array of dragged actors */
+		TOptional<FActorArray> Actors;
+
+		/**
+		 *	Parse a drag operation into our list of actors and folders
+		 *	@return true if the operation is viable for the sceneoutliner to process, false otherwise
+		 */
+		bool ParseDrag(const FDragDropOperation& Operation);
+	};
+	
+	/** Construct a new Drag and drop operation for a scene outliner */
+	TSharedPtr<FDragDropOperation> CreateDragDropOperation(const TArray<FTreeItemPtr>& InTreeItems);
+
+	/** Struct used for validation of a drag/drop operation in the scene outliner */
+	struct FDragValidationInfo
+	{
+		/** The tooltip type to display on the operation */
+		FActorDragDropGraphEdOp::ToolTipTextType TooltipType;
+
+		/** The tooltip text to display on the operation */
+		FText ValidationText;
+
+		/** Construct this validation information out of a tootip type and some text */
+		FDragValidationInfo(const FActorDragDropGraphEdOp::ToolTipTextType InTooltipType, const FText InValidationText)
+			: TooltipType(InTooltipType)
+			, ValidationText(InValidationText)
+		{}
+
+		/** Return a generic invalid result */
+		static FDragValidationInfo Invalid()
+		{
+			return FDragValidationInfo(FActorDragDropGraphEdOp::ToolTip_IncompatibleGeneric, FText());
+		}
+		
+		/** @return true if this operation is valid, false otheriwse */ 
+		bool IsValid() const
+		{
+			switch(TooltipType)
+			{
+			case FActorDragDropGraphEdOp::ToolTip_Compatible:
+			case FActorDragDropGraphEdOp::ToolTip_CompatibleAttach:
+			case FActorDragDropGraphEdOp::ToolTip_CompatibleGeneric:
+			case FActorDragDropGraphEdOp::ToolTip_CompatibleMultipleAttach:
+			case FActorDragDropGraphEdOp::ToolTip_CompatibleDetach:
+			case FActorDragDropGraphEdOp::ToolTip_CompatibleMultipleDetach:
+				return true;
+			default:
+				return false;
+			}
+		}
+	};
+
+	/* A drag/drop operation when dragging folders in the scene outliner */
+	struct FFolderDragDropOp: public FDecoratedDragDropOp
+	{
+		DRAG_DROP_OPERATOR_TYPE(FFolderDragDropOp, FDecoratedDragDropOp)
+
+		/** Array of folders that we are dragging */
+		FFolderPaths Folders;
+
+		void Init(FFolderPaths InFolders);
+	};
+
+	/** A drag/drop operation that was started from the scene outliner */
+	struct FSceneOutlinerDragDropOp : public FDragDropOperation
+	{
+		DRAG_DROP_OPERATOR_TYPE(FSceneOutlinerDragDropOp, FDragDropOperation);
+
+		FSceneOutlinerDragDropOp(const FDragDropPayload& DraggedObjects);
+
+		using FDragDropOperation::Construct;
+
+		/** Actor drag operation */
+		TSharedPtr<FActorDragDropOp>	ActorOp;
+
+		/** Actor drag operation */
+		TSharedPtr<FFolderDragDropOp>	FolderOp;
+
+		void ResetTooltip()
+		{
+			OverrideText = FText();
+			OverrideIcon = nullptr;
+		}
+
+		void SetTooltip(FText InOverrideText, const FSlateBrush* InOverrideIcon)
+		{
+			OverrideText = InOverrideText;
+			OverrideIcon = InOverrideIcon;
+		}
+
+		virtual TSharedPtr<SWidget> GetDefaultDecorator() const override;
+
+	private:
+
+		EVisibility GetOverrideVisibility() const;
+		EVisibility GetDefaultVisibility() const;
+
+		FText OverrideText;
+		FText GetOverrideText() const { return OverrideText; }
+
+		const FSlateBrush* OverrideIcon;
+		const FSlateBrush* GetOverrideIcon() const { return OverrideIcon; }
+	};
+
+}
\ No newline at end of file
diff --git a/Engine/Source/Editor/SceneOutliner/Public/SceneOutlinerFilters.h b/Engine/Source/Editor/SceneOutliner/Public/SceneOutlinerFilters.h
--- a/Engine/Source/Editor/SceneOutliner/Public/SceneOutlinerFilters.h
+++ b/Engine/Source/Editor/SceneOutliner/Public/SceneOutlinerFilters.h
@@ -4,147 +4,164 @@
 
 #include "IFilter.h"
 #include "FilterCollection.h"
-#include "DelegateFilter.h"
+#include "ITreeItem.h"
+
+#include "FolderTreeItem.h"
+#include "ActorTreeItem.h"
+#include "WorldTreeItem.h"
+#include "LevelBlueprintTreeItem.h"
 
 namespace SceneOutliner
 {
-	/** Helper typedefs to aid in creation of filter predicates */
-	typedef TDelegateFilter< const AActor* const >::FPredicate FActorFilterPredicate;
-	typedef TDelegateFilter< FName >::FPredicate FFolderFilterPredicate;
 
-	/** Proxy type that is passed to the filters to check whether something should be shown or not */
-	struct FOutlinerFilterProxy
+	/** Enum to specify how items that are not explicitly handled by this filter should be managed */
+	enum class EDefaultFilterBehaviour : uint8 { Pass, Fail };
+
+	/** A filter that can be applied to any type in the tree */
+	class FOutlinerFilter : public ITreeItemVisitor, public IFilter<const ITreeItem&>
 	{
-		/** The actor to filter, null for folders */
-		const AActor* Actor;
+	public:
 
-		/** The folder to filter */
-		FName FolderPath;
+		/** Event that is fired if this filter changes */
+		DECLARE_DERIVED_EVENT(FOutlinerFilter, IFilter<const ITreeItem&>::FChangedEvent, FChangedEvent);
+		virtual FChangedEvent& OnChanged() override { return ChangedEvent; }
 
-		/** Constructor for filtering the specified actor */
-		FOutlinerFilterProxy(const AActor* InActor)
-			: Actor(InActor)
-			, FolderPath(Actor->GetFolderPath())
-		{}
+	protected:
 
-		/** Constructor for filtering the specified folder path */
-		FOutlinerFilterProxy(FName InFolder)
-			: Actor(nullptr), FolderPath(InFolder)
-		{}
-	};
+		/**	The event that broadcasts whenever a change occurs to the filter */
+		FChangedEvent ChangedEvent;
 
-	/** Scene outliner filter class. This class abstracts the filtering of both actors and folders and allows for filtering on both types */
-	struct FOutlinerFilters : public TFilterCollection<const FOutlinerFilterProxy&>
-	{
-		/** Check if the specified object matches the specified folder filter */
-		static bool ProxyPassesFilter(const FOutlinerFilterProxy& Proxy, TSharedPtr<IFilter<FName>> Filter)
-		{
-			return Filter->PassesFilter(Proxy.FolderPath);
-		}
+		/** Default result of the filter when not overridden in derived classes */
+		const EDefaultFilterBehaviour DefaultBehaviour;
+		
+		/** Constructor to specify the default result of a filter */
+		FOutlinerFilter(EDefaultFilterBehaviour InDefaultBehaviour = EDefaultFilterBehaviour::Pass) : DefaultBehaviour(InDefaultBehaviour) {}
 
-		/** Check if the specified object matches the specified actor filter */
-		static bool ProxyPassesFilter(const FOutlinerFilterProxy& Proxy, TSharedPtr<IFilter<const AActor* const>> Filter)
-		{
-			return !Proxy.Actor || Filter->PassesFilter(Proxy.Actor);
-		}
+		/** Overridden in derived types to filter actors */
+		virtual bool PassesFilter(const AActor* Actor) const { return DefaultBehaviour == EDefaultFilterBehaviour::Pass; }
+		
+		/** Overridden in derived types to filter worlds */
+		virtual bool PassesFilter(const UWorld* World) const { return DefaultBehaviour == EDefaultFilterBehaviour::Pass; }
+		
+		/** Overridden in derived types to filter level blueprints */
+		virtual bool PassesFilter(const FLevelBlueprintHandle& LevelBlueprint) const { return DefaultBehaviour == EDefaultFilterBehaviour::Pass; }
 
-		/** Base class for proxy delegates */
-		struct ProxyBase : IFilter<const FOutlinerFilterProxy&>
-		{
-			/** Virtual functions to test whether the wrapped filter matches the specified filter. Used for removal */
-			virtual bool MatchesWrappedFilter(TSharedPtr<IFilter<FName>>) const					{ return false; }
-			virtual bool MatchesWrappedFilter(TSharedPtr<IFilter<const AActor* const>>) const	{ return false; }
-		};
-
-		/** Derived proxy wrapper. This proxy executes the wrapped delegate with the correct type.
-			Currently supported TFilterTypes are FName and const AActor* const.
-		 */
-		template<class TFilterType>
-		struct ProxyFilter : ProxyBase
-		{
-			typedef TSharedPtr<IFilter<TFilterType>> TFilterPtr;
+		/** Overridden in derived types to filter folders */
+		virtual bool PassesFilter(FName Folder) const { return DefaultBehaviour == EDefaultFilterBehaviour::Pass; }
 
-			/** Constructor */
-			ProxyFilter(TFilterPtr InFilter) : Filter(InFilter) {}
+	private:
 
-			/** Returns whether the specified Item passes the Filter's restrictions */
-			virtual bool PassesFilter(const FOutlinerFilterProxy& InItem) const override
+		/** Transient result from the filter operation. Only valid until the next invocation of the filter. */
+		mutable bool bTransientFilterResult;
+
+		virtual void Visit(const FActorTreeItem& ActorItem) const override
+		{
+			if (const AActor* Actor = ActorItem.Actor.Get())
 			{
-				return ProxyPassesFilter(InItem, Filter);
+				bTransientFilterResult = PassesFilter(Actor);
 			}
-
-			/** Get the on changed event - does nothing for this wrapper */
-			virtual FChangedEvent& OnChanged() override
+			else
 			{
-				return ChangedEvent;
+				bTransientFilterResult = false;
 			}
+		}
 
-			/** Override the relevant function that matches our type */
-			virtual bool MatchesWrappedFilter(TFilterPtr WrappedFilter) const override
+		virtual void Visit(const FWorldTreeItem& WorldItem) const override
+		{
+			if (const UWorld* World = WorldItem.World.Get())
 			{
-				return WrappedFilter == Filter;
+				bTransientFilterResult = PassesFilter(World);
 			}
+			else
+			{
+				bTransientFilterResult = false;
+			}
+		}
 
-			TFilterPtr Filter;
-			FChangedEvent ChangedEvent;
-		};
-
-		/** Meta template used to extract the filter type from a pointer or raw ptr proxy (to support MakeShareable) */
-		template <class> struct TGetFilterType;
-		template <class X, template<class, ESPMode> class TPtr, ESPMode E> struct TGetFilterType<TPtr<X, E>> { typedef typename X::ItemType Type; };
-		template <class X, template<class> class TPtr> struct TGetFilterType<TPtr<X>> { typedef typename X::ItemType Type; };
-
-		/** Add a custom filter to this collection.
-			Example usage:
-				struct FMyFilterType : public IFilter<const AActor* const> {...}
-				Filters.Add(MakeSareable(new FMyFilterType));
-		*/
-		template<class T>
-		void Add(T InFilter)
+		virtual void Visit(const FLevelBlueprintTreeItem& LevelBlueprintItem) const override
 		{
-			typedef ProxyFilter<typename TGetFilterType<T>::Type> TProxyType;
+			bTransientFilterResult = PassesFilter(LevelBlueprintItem.Handle);
+		}
 
-			TSharedPtr<TProxyType> NewProxy(new TProxyType(InFilter));
-			NewProxy->Filter->OnChanged().AddSP(this, &FOutlinerFilters::OnChildFilterChanged);
-			TFilterCollection::Add(NewProxy);
+		virtual void Visit(const FFolderTreeItem& FolderItem) const override
+		{
+			bTransientFilterResult = PassesFilter(FolderItem.Path);
 		}
-		
-		/** Add an actor filter predicate to this filter collection */
-		void AddFilterPredicate(SceneOutliner::FActorFilterPredicate Predicate)
+
+		/** Check whether the specified item passes our filter */
+		virtual bool PassesFilter( const ITreeItem& InItem ) const override
+		{
+			InItem.Visit(*this);
+			return bTransientFilterResult;
+		}
+	};
+
+	DECLARE_DELEGATE_RetVal_OneParam( bool, FActorFilterPredicate, const AActor* );
+	DECLARE_DELEGATE_RetVal_OneParam( bool, FWorldFilterPredicate, const UWorld* );
+	DECLARE_DELEGATE_RetVal_OneParam( bool, FLevelBlueprintFilterPredicate, const FLevelBlueprintHandle& );
+	DECLARE_DELEGATE_RetVal_OneParam( bool, FFolderFilterPredicate, FName );
+
+	/** Predicate based filter for the outliner */
+	struct FOutlinerPredicateFilter : public FOutlinerFilter
+	{
+		/** Predicate used to filter actors */
+		mutable FActorFilterPredicate	ActorPred;
+		/** Predicate used to filter worlds */
+		mutable FWorldFilterPredicate	WorldPred;
+		/** Predicate used to filter Folders */
+		mutable FFolderFilterPredicate	FolderPred;
+		/** Predicate used to filter level blueprints */
+		mutable FLevelBlueprintFilterPredicate LevelBlueprintPred;
+
+		FOutlinerPredicateFilter(FActorFilterPredicate InActorPred, EDefaultFilterBehaviour InDefaultBehaviour)
+			: FOutlinerFilter(InDefaultBehaviour)
+			, ActorPred(InActorPred)
+		{}
+
+		FOutlinerPredicateFilter(FWorldFilterPredicate InWorldPred, EDefaultFilterBehaviour InDefaultBehaviour)
+			: FOutlinerFilter(InDefaultBehaviour)
+			, WorldPred(InWorldPred)
+		{}
+
+		FOutlinerPredicateFilter(FFolderFilterPredicate InFolderPred, EDefaultFilterBehaviour InDefaultBehaviour)
+			: FOutlinerFilter(InDefaultBehaviour)
+			, FolderPred(InFolderPred)
+		{}
+
+		FOutlinerPredicateFilter(FLevelBlueprintFilterPredicate InLevelBpPred, EDefaultFilterBehaviour InDefaultBehaviour)
+			: FOutlinerFilter(InDefaultBehaviour)
+			, LevelBlueprintPred(InLevelBpPred)
+		{}
+
+		virtual bool PassesFilter(const AActor* Actor) const override
 		{
-			Add(MakeShareable(new TDelegateFilter<const AActor* const>(Predicate)));
+			return ActorPred.IsBound() ? ActorPred.Execute(Actor) : DefaultBehaviour == EDefaultFilterBehaviour::Pass;
 		}
 
-		/** Add a folder filter predicate to this filter collection */
-		void AddFilterPredicate(SceneOutliner::FFolderFilterPredicate Predicate)
+		virtual bool PassesFilter(const UWorld* World) const override
 		{
-			Add(MakeShareable(new TDelegateFilter<FName>(Predicate)));
+			return WorldPred.IsBound() ? WorldPred.Execute(World) : DefaultBehaviour == EDefaultFilterBehaviour::Pass;
 		}
 
-		/** Remove the specified wrapped filter from this collection */
-		template<class T>
-		void Remove(TSharedPtr<T> InFilter)
+		virtual bool PassesFilter(const FLevelBlueprintHandle& Handle) const override
 		{
-			TSharedPtr<IFilter<typename T::ItemType>> InBaseFilter = InFilter;
+			return LevelBlueprintPred.IsBound() ? LevelBlueprintPred.Execute(Handle) : DefaultBehaviour == EDefaultFilterBehaviour::Pass;
+		}
 
-			InFilter->OnChanged().RemoveAll(this);
+		virtual bool PassesFilter(FName Folder) const override
+		{
+			return FolderPred.IsBound() ? FolderPred.Execute(Folder) : DefaultBehaviour == EDefaultFilterBehaviour::Pass;
+		}
+	};
 
-			TArray<TSharedPtr<ProxyBase>> FiltersToRemove;
-			for (const auto& Child : ChildFilters)
-			{
-				const TSharedPtr<ProxyBase> ProxyChild = StaticCastSharedPtr<ProxyBase>(Child);
-				if (ProxyChild->MatchesWrappedFilter(InBaseFilter))
-				{
-					FiltersToRemove.Add(ProxyChild);
-				}
-			}
-			if (FiltersToRemove.Num())
-			{
-				for (const auto& Filter : FiltersToRemove)
-				{
-					TFilterCollection::Remove(Filter);
-				}
-			}
+	/** Scene outliner filter class. This class abstracts the filtering of both actors and folders and allows for filtering on both types */
+	struct FOutlinerFilters : public TFilterCollection<const ITreeItem&>
+	{
+		/** Add a filter predicate to this filter collection */
+		template<typename T>
+		void AddFilterPredicate(T Predicate, EDefaultFilterBehaviour InDefaultBehaviour = EDefaultFilterBehaviour::Fail)
+		{
+			Add(MakeShareable(new FOutlinerPredicateFilter(Predicate, InDefaultBehaviour)));
 		}
 	};
 }
\ No newline at end of file
diff --git a/Engine/Source/Editor/SceneOutliner/Public/SceneOutlinerFwd.h b/Engine/Source/Editor/SceneOutliner/Public/SceneOutlinerFwd.h
new file mode 100644
--- /dev/null
+++ b/Engine/Source/Editor/SceneOutliner/Public/SceneOutlinerFwd.h
@@ -0,0 +1,52 @@
+// Copyright 1998-2014 Epic Games, Inc. All Rights Reserved.
+
+#pragma once
+
+class ISceneOutliner;
+class ISceneOutlinerColumn;
+
+namespace SceneOutliner
+{
+	struct FTreeItemID;
+
+	struct FInitializationOptions;
+	struct FSharedOutlinerData;
+	
+	struct ITreeItem;
+	struct FActorTreeItem;
+	struct FWorldTreeItem;
+	struct FLevelBlueprintTreeItem;
+	struct FFolderTreeItem;
+
+	struct ITreeItemVisitor;
+	struct IMutableTreeItemVisitor;
+	
+	typedef TSharedPtr<ITreeItem> FTreeItemPtr;
+	typedef TSharedRef<ITreeItem> FTreeItemRef;
+
+	typedef TMap<FTreeItemID, FTreeItemPtr> FTreeItemMap;
+
+	class SSceneOutliner;
+
+	class FOutlinerFilter;
+	struct FOutlinerFilters;
+
+	struct FDragDropPayload;
+	struct FDragValidationInfo;
+
+	/** Typedef to define an array of actors, used during dragging */
+	typedef TArray<TWeakObjectPtr<AActor>> FActorArray;
+
+	/** Typedef to define an array of folder names, used during dragging */
+	typedef TArray<FName> FFolderPaths;
+}
+
+/** Delegate used with the Scene Outliner in 'actor picking' mode.  You'll bind a delegate when the
+    outliner widget is created, which will be fired off when an actor is selected in the list */
+DECLARE_DELEGATE_OneParam( FOnActorPicked, AActor* );
+DECLARE_DELEGATE_OneParam( FOnSceneOutlinerItemPicked, TSharedRef<SceneOutliner::ITreeItem> );
+
+DECLARE_DELEGATE_OneParam( FCustomSceneOutlinerDeleteDelegate, const TArray< TWeakObjectPtr< AActor > >&  )
+
+/** A delegate used to factory a new column type */
+DECLARE_DELEGATE_RetVal_OneParam( TSharedRef< ISceneOutlinerColumn >, FCreateSceneOutlinerColumn, ISceneOutliner& );
diff --git a/Engine/Source/Editor/SceneOutliner/Public/SceneOutlinerModule.h b/Engine/Source/Editor/SceneOutliner/Public/SceneOutlinerModule.h
--- a/Engine/Source/Editor/SceneOutliner/Public/SceneOutlinerModule.h
+++ b/Engine/Source/Editor/SceneOutliner/Public/SceneOutlinerModule.h
@@ -3,9 +3,7 @@
 #pragma once
 
 #include "ModuleInterface.h"
-#include "ISceneOutliner.h"
-#include "SceneOutlinerInitializationOptions.h"
-
+#include "SceneOutlinerFwd.h"
 
 /**
  * Implements the Scene Outliner module.
@@ -24,8 +22,8 @@ public:
 	 *
 	 * @return	New scene outliner widget
 	 */
-	virtual TSharedRef< class ISceneOutliner > CreateSceneOutliner(
-		const FSceneOutlinerInitializationOptions& InitOptions,
+	virtual TSharedRef< ISceneOutliner > CreateSceneOutliner(
+		const SceneOutliner::FInitializationOptions& InitOptions,
 		const FOnActorPicked& OnActorPickedDelegate ) const;
 
 	/**
@@ -37,14 +35,53 @@ public:
 	 *
 	 * @return	New scene outliner widget
 	 */
-	virtual TSharedRef< class ISceneOutliner > CreateSceneOutliner(
-		const FSceneOutlinerInitializationOptions& InitOptions,
+	virtual TSharedRef< ISceneOutliner > CreateSceneOutliner(
+		const SceneOutliner::FInitializationOptions& InitOptions,
 		const FOnSceneOutlinerItemPicked& OnItemPickedDelegate ) const;
 
+public:
+	/** Register a new type of column available to all scene outliners */
+	template< typename T >
+	void RegisterColumnType()
+	{
+		auto ID = T::GetID();
+		if ( !ColumnMap.Contains( ID ) )
+		{
+			auto CreateColumn = []( ISceneOutliner& Outliner ){
+				return TSharedRef< ISceneOutlinerColumn >( MakeShareable( new T(Outliner) ) );
+			};
+
+			ColumnMap.Add( ID, FCreateSceneOutlinerColumn::CreateStatic( CreateColumn ) );
+		}
+	}
+
+	/** Unregister a previously registered column type */
+	template< typename T >
+	void UnRegisterColumnType()
+	{
+		ColumnMap.Remove( T::GetID() );
+	}
+
+	/** Factory a new column from the specified name. Returns null if no type has been registered under that name. */
+	TSharedPtr< ISceneOutlinerColumn > FactoryColumn( FName ID, ISceneOutliner& Outliner ) const
+	{
+		if ( auto* Factory = ColumnMap.Find( ID ) )
+		{
+			return Factory->Execute(Outliner);
+		}
+		
+		return nullptr;
+	}
+
+private:
+
+	/** Map of column type name -> factory delegate */
+	TMap< FName, FCreateSceneOutlinerColumn > ColumnMap;
+
 public:
 
 	// IModuleInterface interface
 
-	virtual void StartupModule() { }
-	virtual void ShutdownModule() { }
+	virtual void StartupModule() override;
+	virtual void ShutdownModule() override;
 };
diff --git a/Engine/Source/Editor/SceneOutliner/Public/SceneOutlinerPublicTypes.h b/Engine/Source/Editor/SceneOutliner/Public/SceneOutlinerPublicTypes.h
new file mode 100644
--- /dev/null
+++ b/Engine/Source/Editor/SceneOutliner/Public/SceneOutlinerPublicTypes.h
@@ -0,0 +1,170 @@
+// Copyright 1998-2014 Epic Games, Inc. All Rights Reserved.
+
+#pragma once
+
+#include "SceneOutlinerFwd.h"
+
+#include "SceneOutlinerFilters.h"
+#include "ISceneOutlinerColumn.h"
+#include "ISceneOutliner.h"
+
+enum class ESceneOutlinerMode : uint8
+{
+	/** Allows all actors to be browsed and selected; syncs selection with the editor; drag and drop attachment, stc. */
+	ActorBrowsing,
+
+	/** Sets the outliner to operate as an actor 'picker'. */
+	ActorPicker,
+};
+
+namespace SceneOutliner
+{
+	/** Container for built in column types. Function-static so they are available without linking */
+	struct FBuiltInColumnTypes
+	{
+		/** The gutter column */
+		static const FName& Gutter()
+		{
+			static FName Gutter("Gutter");
+			return Gutter;
+		}
+
+		/** The item label column */
+		static const FName& Label()
+		{
+			static FName Label("ItemLabel");
+			return Label;
+		}
+
+		/** Generic actor info column */
+		static FName& ActorInfo()
+		{
+			static FName ActorInfo("ActorInfo");
+			return ActorInfo;
+		}
+	};
+
+	/** Visibility enum for scene outliner columns */
+	enum class EColumnVisibility : uint8
+	{
+		/** This column defaults to being visible on the scene outliner */
+		Visible,
+
+		/** This column defaults to being invisible, yet still available on the scene outliner */
+		Invisible,
+	};
+
+	/** Column information for the scene outliner */
+	struct FColumnInfo
+	{
+		FColumnInfo(EColumnVisibility InVisibility, int32 InPriorityIndex, const FCreateSceneOutlinerColumn& InFactory = FCreateSceneOutlinerColumn())
+			: Visibility(InVisibility), PriorityIndex(InPriorityIndex), Factory(InFactory)
+		{}
+
+		EColumnVisibility 	Visibility;
+		uint8				PriorityIndex;
+
+		FCreateSceneOutlinerColumn	Factory;
+	};
+
+	struct FSharedDataBase
+	{
+		/** Mode to operate in */
+		ESceneOutlinerMode Mode;
+
+		/**	Invoked whenever the user attempts to delete an actor from within the Scene Outliner */
+		FCustomSceneOutlinerDeleteDelegate CustomDelete;
+
+		/** Override default context menu handling */
+		FOnContextMenuOpening ContextMenuOverride;
+
+		/** Extend default context menu handling */
+		TSharedPtr<FExtender> DefaultMenuExtender;
+
+		/** Map of column types available to the scene outliner, along with default ordering */
+		TMap<FName, FColumnInfo> ColumnMap;
+		
+		/** Whether the Scene Outliner should display parenta actors in a Tree */
+		bool bShowParentTree : 1;
+
+		/** True to only show folders in this outliner */
+		bool bOnlyShowFolders : 1;
+
+	public:
+
+		/** Constructor */
+		FSharedDataBase()
+			: Mode( ESceneOutlinerMode::ActorPicker )
+			, bShowParentTree( true )
+			, bOnlyShowFolders( false )
+		{}
+
+		/** Set up a default array of columns for this outliner */
+		void UseDefaultColumns()
+		{
+			if (Mode == ESceneOutlinerMode::ActorBrowsing)
+			{
+				ColumnMap.Add( FBuiltInColumnTypes::Gutter(), FColumnInfo(EColumnVisibility::Visible, 0) );
+			}
+
+			ColumnMap.Add( FBuiltInColumnTypes::Label(), FColumnInfo(EColumnVisibility::Visible, 10) );
+			ColumnMap.Add( FBuiltInColumnTypes::ActorInfo(), FColumnInfo(EColumnVisibility::Visible, 20) );
+		}
+	};
+
+	/**
+	 * Settings for the Scene Outliner set by the programmer before spawning an instance of the widget.  This
+	 * is used to modify the outliner's behavior in various ways, such as filtering in or out specific classes
+	 * of actors.
+	 */
+	struct FInitializationOptions : FSharedDataBase
+	{
+		/** True if we should draw the header row above the tree view */
+		bool bShowHeaderRow : 1;
+
+		/** Whether the Scene Outliner should expose its searchbox */
+		bool bShowSearchBox : 1;
+
+		/** If true, the search box will gain focus when the scene outliner is created */
+		bool bFocusSearchBoxWhenOpened : 1;
+
+		/** Optional collection of filters to use when filtering in the Scene Outliner */
+		TSharedPtr<FOutlinerFilters> Filters;
+
+		/** Broadcasts whenever the Scene Outliners selection changes */
+		FSimpleMulticastDelegate::FDelegate OnSelectionChanged;
+
+	public:
+
+		/** Constructor */
+		FInitializationOptions()
+			: bShowHeaderRow( true )
+			, bShowSearchBox( true )
+			, bFocusSearchBoxWhenOpened( false )
+			, Filters( new FOutlinerFilters )
+		{}
+	};
+
+	/** Outliner data that is shared between a scene outliner and its items */
+	struct FSharedOutlinerData : FSharedDataBase, TSharedFromThis<FSharedOutlinerData>
+	{
+		/** Whether the scene outliner is currently displaying PlayWorld actors */
+		bool bRepresentingPlayWorld;
+
+		/** The world that we are representing */
+		UWorld* RepresentingWorld;
+
+		FSharedOutlinerData()
+			: bRepresentingPlayWorld(false)
+			, RepresentingWorld(nullptr)
+		{}
+	};
+
+	/** Default metrics for outliner tree items */
+	struct FDefaultTreeItemMetrics
+	{
+		static int32	IconSize() { return 18; };
+		static FMargin	IconPadding() { return FMargin(0.f, 0.f, 6.f, 0.f); };
+	};
+
+}	// namespace SceneOutliner
\ No newline at end of file
diff --git a/Engine/Source/Editor/SceneOutliner/Public/SceneOutlinerVisitorTypes.h b/Engine/Source/Editor/SceneOutliner/Public/SceneOutlinerVisitorTypes.h
new file mode 100644
--- /dev/null
+++ b/Engine/Source/Editor/SceneOutliner/Public/SceneOutlinerVisitorTypes.h
@@ -0,0 +1,91 @@
+// Copyright 1998-2014 Epic Games, Inc. All Rights Reserved.
+
+#pragma once
+
+#include "SceneOutlinerFwd.h"
+
+namespace SceneOutliner
+{
+	/** A const tree item visitor. Derive to implement type-specific behaviour for tree items. */
+	struct ITreeItemVisitor
+	{
+		virtual void Visit(const FActorTreeItem& Actor) const {}
+		virtual void Visit(const FWorldTreeItem& World) const {}
+		virtual void Visit(const FLevelBlueprintTreeItem& LevelBlueprint) const {}
+		virtual void Visit(const FFolderTreeItem& Folder) const {}
+	};
+
+	/** A non-const tree item visitor. Derive to implement type-specific behaviour for tree items. */
+	struct IMutableTreeItemVisitor
+	{
+		virtual void Visit(FActorTreeItem& Actor) const {}
+		virtual void Visit(FWorldTreeItem& World) const {}
+		virtual void Visit(FLevelBlueprintTreeItem& LevelBlueprint) const {}
+		virtual void Visit(FFolderTreeItem& Folder) const {}
+	};
+
+	/** A functional-based visitor. Allows for visitor-pattern behaviour without creating a custom type. */
+	struct FFunctionalVisitor : ITreeItemVisitor
+	{
+		typedef TFunctionRef<void(const FActorTreeItem&)> 				FActorFunction;
+		typedef TFunctionRef<void(const FWorldTreeItem&)> 				FWorldFunction;
+		typedef TFunctionRef<void(const FLevelBlueprintTreeItem&)> 		FLevelBlueprintFunction;
+		typedef TFunctionRef<void(const FFolderTreeItem&)> 				FFolderFunction;
+
+		FFunctionalVisitor& Actor(const FActorFunction& InFunction) 					{ ActorFunction = InFunction; return *this; }
+		FFunctionalVisitor& World(const FWorldFunction& InFunction) 					{ WorldFunction = InFunction; return *this; }
+		FFunctionalVisitor& LevelBlueprint(const FLevelBlueprintFunction& InFunction) 	{ LevelBlueprintFunction = InFunction; return *this; }
+		FFunctionalVisitor& Folder(const FFolderFunction& InFunction) 					{ FolderFunction = InFunction; return *this; }
+
+	private:
+		TOptional<FActorFunction> ActorFunction;
+		TOptional<FWorldFunction> WorldFunction;
+		TOptional<FLevelBlueprintFunction> LevelBlueprintFunction;
+		TOptional<FFolderFunction> FolderFunction;
+
+		virtual void Visit(const FActorTreeItem& Item) const override 					{ if (ActorFunction) { ActorFunction.GetValue()(Item); } }
+		virtual void Visit(const FWorldTreeItem& Item) const override 					{ if (WorldFunction) { WorldFunction.GetValue()(Item); } }
+		virtual void Visit(const FLevelBlueprintTreeItem& Item) const override 			{ if (LevelBlueprintFunction) { LevelBlueprintFunction.GetValue()(Item); } }
+		virtual void Visit(const FFolderTreeItem& Item) const override 					{ if (FolderFunction) { FolderFunction.GetValue()(Item); } }
+	};
+
+	/** A visitor specialized for getting/extracting a value from a tree item. */
+	template<typename TDataType>
+	struct TTreeItemGetter : ITreeItemVisitor
+	{
+		mutable TDataType Data;
+
+		/** Override to extract the data from specific tree item types */
+		virtual TDataType Get(const FActorTreeItem& ActorItem) const 					{ return TDataType(); }
+		virtual TDataType Get(const FWorldTreeItem& WorldItem) const 					{ return TDataType(); }
+		virtual TDataType Get(const FLevelBlueprintTreeItem& Item) const 				{ return TDataType(); }
+		virtual TDataType Get(const FFolderTreeItem& FolderItem) const 					{ return TDataType(); }
+
+		/** Return the result returned from Get() */
+		FORCEINLINE const TDataType& Result() const 									{ return Data; }
+
+	private:
+		virtual void Visit(const FActorTreeItem& ActorItem) const override 				{ Data = Get(ActorItem); }
+		virtual void Visit(const FWorldTreeItem& WorldItem) const override 				{ Data = Get(WorldItem); }
+		virtual void Visit(const FLevelBlueprintTreeItem& Item) const override 			{ Data = Get(Item); }
+		virtual void Visit(const FFolderTreeItem& FolderItem) const override 			{ Data = Get(FolderItem); }
+	};
+
+	/** A visitor class used to generate column cells for specific tree item types */
+	struct FColumnGenerator : IMutableTreeItemVisitor
+	{
+		mutable TSharedPtr<SWidget> Widget;
+
+		virtual TSharedRef<SWidget> GenerateWidget(FActorTreeItem& Item) const 			{ return SNullWidget::NullWidget; }
+		virtual TSharedRef<SWidget> GenerateWidget(FWorldTreeItem& Item) const 			{ return SNullWidget::NullWidget; }
+		virtual TSharedRef<SWidget> GenerateWidget(FLevelBlueprintTreeItem& Item) const { return SNullWidget::NullWidget; }
+		virtual TSharedRef<SWidget> GenerateWidget(FFolderTreeItem& Item) const 		{ return SNullWidget::NullWidget; }
+
+	private:
+		virtual void Visit(FActorTreeItem& Item) const override 						{ Widget = GenerateWidget(Item); }
+		virtual void Visit(FWorldTreeItem& Item) const override 						{ Widget = GenerateWidget(Item); }
+		virtual void Visit(FLevelBlueprintTreeItem& Item) const override 				{ Widget = GenerateWidget(Item); }
+		virtual void Visit(FFolderTreeItem& Item) const override 						{ Widget = GenerateWidget(Item); }
+	};
+	
+}	// namespace SceneOutliner
\ No newline at end of file
diff --git a/Engine/Source/Editor/SceneOutliner/Public/SortHelper.h b/Engine/Source/Editor/SceneOutliner/Public/SortHelper.h
new file mode 100644
--- /dev/null
+++ b/Engine/Source/Editor/SceneOutliner/Public/SortHelper.h
@@ -0,0 +1,268 @@
+// Copyright 1998-2014 Epic Games, Inc. All Rights Reserved.
+
+#pragma once
+
+namespace SceneOutliner
+{
+	/**
+	 * Templated helper to alleviate performance problems with sorting based on complex predicates.
+	 * Example Usage:
+	 * 		FSortHelper<FString>().Primary([](const ITreeItem& Item){ return Item->GetString(); }).SortItems(Array);
+	 *
+	 * Or:
+	 * 		FSortHelper<FString, FString>()
+	 *			.Primary(FGetPrimaryStringVisitor())
+	 *			.Secondary(FGetSecondaryStringVisitor())
+	 *			.SortItems(Array);
+	 */
+	template<typename PrimaryKeyType, typename SecondaryKeyType = int32>
+	struct FSortHelper
+	{
+		typedef TFunctionRef<PrimaryKeyType(const ITreeItem&)> FPrimaryFunction;
+		typedef TFunctionRef<SecondaryKeyType(const ITreeItem&)> FSecondaryFunction;
+
+		FSortHelper()
+			: PrimarySortMode(EColumnSortMode::None), SecondarySortMode(EColumnSortMode::None)
+			, PrimaryVisitor(nullptr), SecondaryVisitor(nullptr)
+		{}
+
+		/** Sort primarily by the specified function and mode. Beware the function is a reference, so must be valid for the lifetime of this instance. */
+		FSortHelper& Primary(FPrimaryFunction&& Function, EColumnSortMode::Type SortMode)
+		{
+			PrimarySortMode = SortMode;
+			PrimaryFunction = MoveTemp(Function);
+			return *this;
+		}
+		/** Sort primarily using the specified 'getter' visitor and mode */
+		FSortHelper& Primary(const TTreeItemGetter<PrimaryKeyType>& Visitor, EColumnSortMode::Type SortMode)
+		{
+			PrimarySortMode = SortMode;
+			PrimaryVisitor = &Visitor;
+			return *this;
+		}
+
+		/** Sort secondarily by the specified function and mode. Beware the function is a reference, so must be valid for the lifetime of this instance. */
+		FSortHelper& Secondary(FSecondaryFunction&& Function, EColumnSortMode::Type SortMode)
+		{
+			SecondarySortMode = SortMode;
+			SecondaryFunction = MoveTemp(Function);
+			return *this;
+		}
+		/** Sort secondarily using the specified 'getter' visitor and mode */
+		FSortHelper& Secondary(const TTreeItemGetter<SecondaryKeyType>& Visitor, EColumnSortMode::Type SortMode)
+		{
+			PrimarySortMode = SortMode;
+			SecondaryVisitor = &Visitor;
+			return *this;
+		}
+
+		/** Sort the specified array using the current sort settings */
+		void Sort(TArray<FTreeItemPtr>& Array)
+		{
+			TArray<FSortPayload> SortData;
+			const auto End = Array.Num();
+			for (int32 Index = 0; Index < End; ++Index)
+			{
+				const auto& Element = Array[Index];
+
+				PrimaryKeyType PrimaryKey;
+				if (PrimaryVisitor)
+				{
+					Element->Visit(*PrimaryVisitor);
+					PrimaryKey = MoveTemp(PrimaryVisitor->Data);
+				}
+				else if (PrimaryFunction)
+				{
+					PrimaryKey = PrimaryFunction.GetValue()(*Element);
+				}
+
+				SecondaryKeyType SecondaryKey;
+				if (SecondarySortMode != EColumnSortMode::None)
+				{
+					if (SecondaryVisitor)
+					{
+						Element->Visit(*SecondaryVisitor);
+						SecondaryKey = MoveTemp(SecondaryVisitor->Data);
+					}
+					else if (SecondaryFunction)
+					{
+						SecondaryKey = SecondaryFunction.GetValue()(*Element);
+					}
+				}
+
+				SortData.Emplace(Index, MoveTemp(PrimaryKey), MoveTemp(SecondaryKey));
+			}
+
+			SortData.Sort([&](const FSortPayload& One, const FSortPayload& Two){
+				if (PrimarySortMode == EColumnSortMode::Ascending && One.PrimaryKey != Two.PrimaryKey)
+					return One.PrimaryKey < Two.PrimaryKey;
+				else if (PrimarySortMode == EColumnSortMode::Descending && One.PrimaryKey != Two.PrimaryKey)
+					return One.PrimaryKey > Two.PrimaryKey;
+
+				if (SecondarySortMode == EColumnSortMode::Ascending)
+					return One.SecondaryKey < Two.SecondaryKey;
+				else if (SecondarySortMode == EColumnSortMode::Descending)
+					return One.SecondaryKey > Two.SecondaryKey;
+
+				return false;
+			});
+
+			TArray<FTreeItemPtr> NewArray;
+			NewArray.Reserve(Array.Num());
+
+			for (const auto& Element : SortData)
+			{
+				NewArray.Add(Array[Element.OriginalIndex]);
+			}
+			Array = MoveTemp(NewArray);
+		}
+
+	private:
+
+		EColumnSortMode::Type 	PrimarySortMode;
+		EColumnSortMode::Type 	SecondarySortMode;
+
+		TOptional<FPrimaryFunction>		PrimaryFunction;
+		TOptional<FSecondaryFunction>	SecondaryFunction;
+
+		const TTreeItemGetter<PrimaryKeyType>*		PrimaryVisitor;
+		const TTreeItemGetter<SecondaryKeyType>*	SecondaryVisitor;
+
+		/** Aggregated data from the sort methods. We extract the sort data from all elements first, then sort based on the extracted data. */
+		struct FSortPayload
+		{
+			int32 OriginalIndex;
+			
+			PrimaryKeyType PrimaryKey;
+			SecondaryKeyType SecondaryKey;
+
+			FSortPayload(int32 InOriginalIndex, PrimaryKeyType&& InPrimaryKey, SecondaryKeyType&& InSecondaryKey)
+				: OriginalIndex(InOriginalIndex), PrimaryKey(MoveTemp(InPrimaryKey)), SecondaryKey(MoveTemp(InSecondaryKey))
+			{}
+
+			FSortPayload(FSortPayload&& Other)
+			{
+				(*this) = MoveTemp(Other);
+			}
+			FSortPayload& operator=(FSortPayload&& rhs)
+			{
+				OriginalIndex = rhs.OriginalIndex;
+				PrimaryKey = MoveTemp(rhs.PrimaryKey);
+				SecondaryKey = MoveTemp(rhs.SecondaryKey);
+				return *this;
+			}
+
+		private:
+			FSortPayload(const FSortPayload&);
+			FSortPayload& operator=(const FSortPayload&);
+		};
+	};
+
+	/** Wrapper type that will sort FString's using a more natural comparison (correctly sorts numbers and ignores underscores) */
+	struct FNumericStringWrapper
+	{
+		FString String;
+
+		FNumericStringWrapper()
+		{}
+		FNumericStringWrapper(FString&& InString)
+			: String(InString)
+		{}
+		FNumericStringWrapper(FNumericStringWrapper&& Other)
+			: String(MoveTemp(Other.String))
+		{}
+		FNumericStringWrapper& operator=(FNumericStringWrapper&& Other)
+		{
+			String = MoveTemp(Other.String);
+			return *this;
+		}
+		FORCEINLINE bool operator<(const FNumericStringWrapper& Other) const
+		{
+			return CompareNumeric(String, Other.String) < 0;
+		}
+		FORCEINLINE bool operator>(const FNumericStringWrapper& Other) const
+		{
+			return CompareNumeric(String, Other.String) > 0;
+		}
+		FORCEINLINE bool operator==(const FNumericStringWrapper& Other) const
+		{
+			return String == Other.String;
+		}
+
+		/** Compare the 2 specified strings using the specified default character comparison function */
+		typedef int32 (*StrncmpMethod)(const TCHAR*, const TCHAR*, SIZE_T Count);
+		static int32 CompareNumeric(const FString& A, const FString& B, StrncmpMethod Strncmp = FCString::Strnicmp)
+		{
+			const TCHAR* CharA = *A, *CharB = *B;
+			const TCHAR* EndA = CharA + A.Len(), *EndB = CharB + B.Len();
+
+			while(CharA < EndA && CharB < EndB)
+			{
+				// Ignore underscores
+				if (*CharA == '_')
+				{
+					++CharA;
+					continue;
+				}
+				if (*CharB == '_')
+				{
+					++CharB;
+					continue;
+				}
+				
+				TCHAR *IntAEnd = nullptr;
+				const uint64 IntA = FCString::Strtoui64(CharA, &IntAEnd, 10);
+
+				TCHAR *IntBEnd = nullptr;
+				const uint64 IntB = FCString::Strtoui64(CharB, &IntBEnd, 10);
+
+				const bool AIsNum = (IntAEnd != CharA), BIsNum = (IntBEnd != CharB);
+				if (AIsNum != BIsNum)
+				{
+					// At the current CharA/CharB position, either (but not both) of the strings is an integer
+					// Numbers are considered less than characters
+					return AIsNum ? -1 : 1;
+				}
+				else if (AIsNum)
+				{
+					CharA = IntAEnd;
+					CharB = IntBEnd;
+
+					if (IntA != IntB)
+					{
+						return IntA < IntB ? -1 : 1;
+					}
+				}
+				else if (int32 Cmp = (*Strncmp)(CharA, CharB, 1))
+				{
+					return Cmp;
+				}
+				else
+				{
+					++CharA;
+					++CharB;
+				}
+			}
+
+			if (CharA == EndA && CharB == EndB)
+			{
+				// Strings compared equal, return shortest first
+				if (const int32 LengthDifference = A.Len() - B.Len())
+				{
+					return LengthDifference / FMath::Abs(LengthDifference);
+				}
+				else
+				{
+					return 0;
+				}
+			}
+
+			// Strings are different comparative lengths, return shortest
+			return CharA == EndA ? -1 : 1;
+		}
+
+	private:
+		FNumericStringWrapper(const FNumericStringWrapper&);
+		FNumericStringWrapper& operator=(const FNumericStringWrapper&);
+	};
+}
\ No newline at end of file
diff --git a/Engine/Source/Editor/SceneOutliner/Public/WorldTreeItem.h b/Engine/Source/Editor/SceneOutliner/Public/WorldTreeItem.h
new file mode 100644
--- /dev/null
+++ b/Engine/Source/Editor/SceneOutliner/Public/WorldTreeItem.h
@@ -0,0 +1,64 @@
+// Copyright 1998-2014 Epic Games, Inc. All Rights Reserved.
+
+#pragma once
+
+#include "ITreeItem.h"
+
+namespace SceneOutliner
+{
+	/** A tree item that represents an entire world */
+	struct FWorldTreeItem : ITreeItem
+	{
+		/** The world this tree item is associated with. */
+		mutable TWeakObjectPtr<UWorld> World;
+
+		/** Construct this item from a world */
+		FWorldTreeItem(UWorld* InWorld);
+
+		/** Get this item's parent item. Always returns nullptr. */
+		virtual FTreeItemPtr FindParent(const FTreeItemMap& ExistingItems) const override;
+
+		/** Create this item's parent. Always returns nullptr. */
+		virtual FTreeItemPtr CreateParent() const override;
+
+ 		/** Visit this tree item */
+ 		virtual void Visit(const ITreeItemVisitor& Visitor) const override;
+ 		virtual void Visit(const IMutableTreeItemVisitor& Visitor) override;
+
+	public:
+
+		/** Get the ID that represents this tree item. Used to reference this item in a map */
+		virtual FTreeItemID GetID() const override;
+
+		/** Get the raw string to display for this tree item - used for sorting */
+		virtual FString GetDisplayString() const override;
+
+		/** Get the sort priority given to this item's type */
+		virtual int32 GetTypeSortPriority() const override;
+
+		/** Check whether it should be possible to interact with this tree item */
+		virtual bool CanInteract() const override;
+
+		/** Generate a context menu for this item. Only called if *only* this item is selected. */
+		virtual void GenerateContextMenu(FMenuBuilder& MenuBuilder, SSceneOutliner& Outliner) override;
+
+		/** Populate the specified drag/drop payload with any relevant information for this type */
+		virtual void PopulateDragDropPayload(FDragDropPayload& Payload) const override
+		{
+		}
+
+		/** Called to test whether the specified payload can be dropped onto this tree item */
+		virtual FDragValidationInfo ValidateDrop(FDragDropPayload& DraggedObjects, UWorld& World) const override;
+
+		/** Called to drop the specified objects on this item. Only called if ValidateDrop() allows. */
+		virtual void OnDrop(FDragDropPayload& DraggedObjects, UWorld& World, const FDragValidationInfo& ValidationInfo, TSharedRef<SWidget> DroppedOnWidget) override
+		{
+		}
+
+	public:
+
+		/** Open the world settings for the contained world */
+		void OpenWorldSettings() const;
+	};
+
+}		// namespace SceneOutliner
\ No newline at end of file
diff --git a/Engine/Source/Editor/UnrealEd/Private/EditorActorFolders.cpp b/Engine/Source/Editor/UnrealEd/Private/EditorActorFolders.cpp
--- a/Engine/Source/Editor/UnrealEd/Private/EditorActorFolders.cpp
+++ b/Engine/Source/Editor/UnrealEd/Private/EditorActorFolders.cpp
@@ -30,6 +30,7 @@ FName OldPathToNewPath(const FString& InOldBranch, const FString& InNewBranch, c
 
 // Static member definitions
 FOnActorFolderCreate	FActorFolders::OnFolderCreate;
+FOnActorFolderMove		FActorFolders::OnFolderMove;
 FOnActorFolderDelete	FActorFolders::OnFolderDelete;
 FActorFolders*			FActorFolders::Singleton;
 
@@ -405,6 +406,7 @@ bool FActorFolders::RenameFolderInWorld(UWorld& World, FName OldPath, FName NewP
 					// Otherwise use default properties
 					FoldersInWorld.Folders.Add(NewFolder);
 				}
+				OnFolderMove.Broadcast(World, Path, NewFolder);
 				OnFolderCreate.Broadcast(World, NewFolder);
 			}
 			RenamedFolders.Add(Path);
diff --git a/Engine/Source/Editor/UnrealEd/Public/DragAndDrop/MultipleDataDragDropOp.h b/Engine/Source/Editor/UnrealEd/Public/DragAndDrop/MultipleDataDragDropOp.h
new file mode 100644
--- /dev/null
+++ b/Engine/Source/Editor/UnrealEd/Public/DragAndDrop/MultipleDataDragDropOp.h
@@ -0,0 +1,47 @@
+// Copyright 1998-2014 Epic Games, Inc. All Rights Reserved.
+
+#pragma once
+
+#include "DecoratedDragDropOp.h"
+
+/** Drag payload type - implement derived types using DRAG_DROP_OPERATOR_TYPE */
+class IDragPayload
+{
+protected:
+	/** Friendship required so we can access IsOfType without exposing it */
+	friend class FMultipleDataDragOp;
+
+	/** Check if this payload type is the same as the specified type ID */
+	virtual bool IsOfTypeImpl(const FString& Type) const { return false; }
+};
+
+/** A drag/drop operation that contains multiple types of dragged data (payloads) */
+class FMultipleDataDragOp : public FDecoratedDragDropOp
+{
+public:
+	DRAG_DROP_OPERATOR_TYPE(FMultipleDataDragOp, FDecoratedDragDropOp)
+
+	/** Add some more data to this operation's payload */
+	void AddPayload(TUniquePtr<IDragPayload> InPayload)
+	{
+		Payload.Add(MoveTemp(InPayload));
+	}
+
+	/** Extract the specified data type from this operation */
+	virtual const void* Extract(const FString& Type) const
+	{
+		for (const auto& Entry : Payload)
+		{
+			if (Entry->IsOfTypeImpl(T::GetTypeId()))
+			{
+				return Entry.Get();
+			}
+		}
+		return FDecoratedDragDropOp::Extract(Type);
+	}
+
+protected:
+
+	/** The data that we're dragging */
+	TArray<TUniquePtr<IDragPayload>> Payload;
+};
\ No newline at end of file
diff --git a/Engine/Source/Editor/UnrealEd/Public/EditorActorFolders.h b/Engine/Source/Editor/UnrealEd/Public/EditorActorFolders.h
--- a/Engine/Source/Editor/UnrealEd/Public/EditorActorFolders.h
+++ b/Engine/Source/Editor/UnrealEd/Public/EditorActorFolders.h
@@ -7,6 +7,7 @@
 /** Multicast delegates for broadcasting various folder events */
 DECLARE_MULTICAST_DELEGATE_TwoParams(FOnActorFolderCreate, UWorld&, FName);
 DECLARE_MULTICAST_DELEGATE_TwoParams(FOnActorFolderDelete, UWorld&, FName);
+DECLARE_MULTICAST_DELEGATE_ThreeParams(FOnActorFolderMove, UWorld&, FName /* src */, FName /* dst */);
 
 USTRUCT()
 struct FActorFolderProps
@@ -59,6 +60,7 @@ struct UNREALED_API FActorFolders : public FGCObject
 
 	/** Folder creation and deletion events. Called whenever a folder is created or deleted in a world. */
 	static FOnActorFolderCreate OnFolderCreate;
+	static FOnActorFolderMove 	OnFolderMove;
 	static FOnActorFolderDelete OnFolderDelete;
 
 	/** Check if the specified path is a child of the specified parent */
