Fixed some Scene Outliner tree item IDs becoming invalid when the objects they reference disappear

diff --git a/Engine/Source/Editor/SceneOutliner/Private/ActorTreeItem.cpp b/Engine/Source/Editor/SceneOutliner/Private/ActorTreeItem.cpp
--- a/Engine/Source/Editor/SceneOutliner/Private/ActorTreeItem.cpp
+++ b/Engine/Source/Editor/SceneOutliner/Private/ActorTreeItem.cpp
@@ -181,6 +181,7 @@ void FActorDropTarget::DetachActorFromParent(AActor* ChildActor)
 
 FActorTreeItem::FActorTreeItem(AActor* InActor)
 	: Actor(InActor)
+	, ID(InActor)
 {
 	bExistsInCurrentWorldAndPIE = GEditor->ObjectsThatExistInEditorWorld.Get(InActor);
 }
@@ -194,8 +195,8 @@ FTreeItemPtr FActorTreeItem::FindParent(const FTreeItemMap& ExistingItems) const
 	}
 
 	// Parents should have already been added to the tree
-	TWeakObjectPtr<AActor> ParentActor = ActorPtr->GetAttachParentActor();
-	if (ParentActor.IsValid())
+	AActor* ParentActor = ActorPtr->GetAttachParentActor();
+	if (ParentActor)
 	{
 		return ExistingItems.FindRef(ParentActor);
 	}
@@ -262,7 +263,7 @@ void FActorTreeItem::Visit(const IMutableTreeItemVisitor& Visitor)
 
 FTreeItemID FActorTreeItem::GetID() const
 {
-	return FTreeItemID(Actor);
+	return ID;
 }
 
 FString FActorTreeItem::GetDisplayString() const
diff --git a/Engine/Source/Editor/SceneOutliner/Private/SceneOutlinerStandaloneTypes.h b/Engine/Source/Editor/SceneOutliner/Private/SceneOutlinerStandaloneTypes.h
--- a/Engine/Source/Editor/SceneOutliner/Private/SceneOutlinerStandaloneTypes.h
+++ b/Engine/Source/Editor/SceneOutliner/Private/SceneOutlinerStandaloneTypes.h
@@ -2,6 +2,7 @@
 
 #pragma once
 
+#include "ObjectKey.h"
 #include "SceneOutlinerFwd.h"
 
 #include "SceneOutlinerDragDrop.h"
@@ -33,27 +34,16 @@ namespace SceneOutliner
 		/** Default constructed null item ID */
 		FTreeItemID() : Type(EType::Null), CachedHash(0) {}
 
-		/** ID representing an actor */
-		FTreeItemID(const TWeakObjectPtr<AActor>& InActor) : Type(EType::Actor)
+		/** ID representing a UObject */
+		FTreeItemID(const UObject* InObject) : Type(EType::Object)
 		{
-			new (Data) TWeakObjectPtr<AActor>(InActor);
+			check(InObject);
+			new (Data) FObjectKey(InObject);
 			CachedHash = CalculateTypeHash();
 		}
-		FTreeItemID(const AActor* InActor) : Type(EType::Actor)
+		FTreeItemID(const FObjectKey& InKey) : Type(EType::Object)
 		{
-			new (Data) TWeakObjectPtr<AActor>(InActor);
-			CachedHash = CalculateTypeHash();
-		}
-
-		/** ID representing a world */
-		FTreeItemID(const TWeakObjectPtr<UWorld>& InWorld) : Type(EType::World)
-		{
-			new (Data) TWeakObjectPtr<UWorld>(InWorld);
-			CachedHash = CalculateTypeHash();
-		}
-		FTreeItemID(const UWorld* InWorld) : Type(EType::World)
-		{
-			new (Data) TWeakObjectPtr<UWorld>(InWorld);
+			new (Data) FObjectKey(InKey);
 			CachedHash = CalculateTypeHash();
 		}
 
@@ -81,8 +71,7 @@ namespace SceneOutliner
 			Type = Other.Type;
 			switch(Type)
 			{
-				case EType::Actor:			new (Data) TWeakObjectPtr<AActor>(Other.GetAsActorRef());								break;
-				case EType::World:			new (Data) TWeakObjectPtr<UWorld>(Other.GetAsWorldRef());								break;
+				case EType::Object:			new (Data) FObjectKey(Other.GetAsObjectKey());											break;
 				case EType::LevelBlueprint:	new (Data) FLevelBlueprintHandle(Other.GetAsLevelBlueprintRef());						break;
 				case EType::Folder:			new (Data) FName(Other.GetAsFolderRef());												break;
 				default:																											break;
@@ -107,8 +96,7 @@ namespace SceneOutliner
 		{
 			switch(Type)
 			{
-				case EType::Actor:			GetAsActorRef().~TWeakObjectPtr();						break;
-				case EType::World:			GetAsWorldRef().~TWeakObjectPtr();						break;
+				case EType::Object:			GetAsObjectKey().~FObjectKey();							break;
 				case EType::LevelBlueprint:	GetAsLevelBlueprintRef().~FLevelBlueprintHandle();		break;
 				case EType::Folder:			GetAsFolderRef().~FName();								break;
 				default:																			break;
@@ -129,8 +117,7 @@ namespace SceneOutliner
 			uint32 Hash = 0;
 			switch(Type)
 			{
-				case EType::Actor:			Hash = GetTypeHash(GetAsActorRef());				break;
-				case EType::World:			Hash = GetTypeHash(GetAsWorldRef());				break;
+				case EType::Object:			Hash = GetTypeHash(GetAsObjectKey());				break;
 				case EType::LevelBlueprint:	Hash = GetTypeHash(GetAsLevelBlueprintRef());		break;
 				case EType::Folder:			Hash = GetTypeHash(GetAsFolderRef());				break;
 				default:																		break;
@@ -146,18 +133,16 @@ namespace SceneOutliner
 
 	private:
 
-		TWeakObjectPtr<AActor>& 				GetAsActorRef() const 			{ return *reinterpret_cast<TWeakObjectPtr<AActor>*>(Data); }
-		TWeakObjectPtr<UWorld>& 				GetAsWorldRef() const 			{ return *reinterpret_cast<TWeakObjectPtr<UWorld>*>(Data); }
-		FLevelBlueprintHandle& 					GetAsLevelBlueprintRef() const 	{ return *reinterpret_cast<FLevelBlueprintHandle*>(Data); }
-		FName& 									GetAsFolderRef() const			{ return *reinterpret_cast<FName*>(Data); }
+		FObjectKey& 				GetAsObjectKey() const 			{ return *reinterpret_cast<FObjectKey*>(Data); }
+		FLevelBlueprintHandle& 		GetAsLevelBlueprintRef() const 	{ return *reinterpret_cast<FLevelBlueprintHandle*>(Data); }
+		FName& 						GetAsFolderRef() const			{ return *reinterpret_cast<FName*>(Data); }
 
 		/** Compares the specified ID with this one - assumes matching types */
 		bool Compare(const FTreeItemID& Other) const
 		{
 			switch(Type)
 			{
-				case EType::Actor:			return GetAsActorRef() == Other.GetAsActorRef();
-				case EType::World:			return GetAsWorldRef() == Other.GetAsWorldRef();
+				case EType::Object:			return GetAsObjectKey() == Other.GetAsObjectKey();
 				case EType::LevelBlueprint:	return GetAsLevelBlueprintRef() == Other.GetAsLevelBlueprintRef();
 				case EType::Folder:			return GetAsFolderRef() == Other.GetAsFolderRef();
 				case EType::Null:			return true;
@@ -165,11 +150,11 @@ namespace SceneOutliner
 			}
 		}
 
-		enum class EType : uint8 { Actor, World, LevelBlueprint, Folder, Null };
+		enum class EType : uint8 { Object, LevelBlueprint, Folder, Null };
 		EType Type;
 
 		uint32 CachedHash;
-		static const uint32 MaxSize = TMaxSizeof<TWeakObjectPtr<UWorld>, TWeakObjectPtr<ULevelScriptBlueprint>, TWeakObjectPtr<AActor>, FName>::Value;
+		static const uint32 MaxSize = TMaxSizeof<FObjectKey, FLevelBlueprintHandle, FName>::Value;
 		mutable uint8 Data[MaxSize];
 	};
 
diff --git a/Engine/Source/Editor/SceneOutliner/Private/WorldTreeItem.cpp b/Engine/Source/Editor/SceneOutliner/Private/WorldTreeItem.cpp
--- a/Engine/Source/Editor/SceneOutliner/Private/WorldTreeItem.cpp
+++ b/Engine/Source/Editor/SceneOutliner/Private/WorldTreeItem.cpp
@@ -12,6 +12,7 @@ namespace SceneOutliner
 
 FWorldTreeItem::FWorldTreeItem(UWorld* InWorld)
 	: World(InWorld)
+	, ID(InWorld)
 {
 
 }
@@ -38,7 +39,7 @@ void FWorldTreeItem::Visit(const IMutableTreeItemVisitor& Visitor)
 
 FTreeItemID FWorldTreeItem::GetID() const
 {
-	return FTreeItemID(World);
+	return ID;
 }
 
 FString FWorldTreeItem::GetDisplayString() const
diff --git a/Engine/Source/Editor/SceneOutliner/Public/ActorTreeItem.h b/Engine/Source/Editor/SceneOutliner/Public/ActorTreeItem.h
--- a/Engine/Source/Editor/SceneOutliner/Public/ActorTreeItem.h
+++ b/Engine/Source/Editor/SceneOutliner/Public/ActorTreeItem.h
@@ -3,6 +3,7 @@
 #pragma once
 
 #include "ITreeItem.h"
+#include "ObjectKey.h"
 
 namespace SceneOutliner
 {
@@ -39,6 +40,9 @@ namespace SceneOutliner
 		/** The actor this tree item is associated with. */
 		mutable TWeakObjectPtr<AActor> Actor;
 
+		/** Constant identifier for this tree item */
+		const FObjectKey ID;
+
 		/** Construct this item from an actor */
 		FActorTreeItem(AActor* InActor);
 
diff --git a/Engine/Source/Editor/SceneOutliner/Public/LevelBlueprintTreeItem.h b/Engine/Source/Editor/SceneOutliner/Public/LevelBlueprintTreeItem.h
--- a/Engine/Source/Editor/SceneOutliner/Public/LevelBlueprintTreeItem.h
+++ b/Engine/Source/Editor/SceneOutliner/Public/LevelBlueprintTreeItem.h
@@ -3,6 +3,7 @@
 #pragma once
 
 #include "ITreeItem.h"
+#include "ObjectKey.h"
 
 namespace SceneOutliner
 {
@@ -12,9 +13,11 @@ namespace SceneOutliner
 		/** We reference level blueprints by their parent level, as they don't always exist */
 		TWeakObjectPtr<ULevel> ParentLevel;
 
+		/** A key that is used to reference the level */
+		FObjectKey LevelKey;
+
 		/** Construct this handle from a level */
-		explicit FLevelBlueprintHandle(ULevel* InLevel) : ParentLevel(InLevel) {}
-		explicit FLevelBlueprintHandle(const TWeakObjectPtr<ULevel>& InLevel) : ParentLevel(InLevel) {}
+		explicit FLevelBlueprintHandle(ULevel* InLevel) : ParentLevel(InLevel), LevelKey(InLevel) {}
 
 		/** Get the level blueprint for the parent level (if it exists) */
 		ULevelScriptBlueprint* GetLevelBlueprint()
@@ -40,11 +43,11 @@ namespace SceneOutliner
 
 		friend uint32 GetTypeHash(const FLevelBlueprintHandle& Handle)
 		{
-			return GetTypeHash(Handle.ParentLevel);
+			return GetTypeHash(Handle.LevelKey);
 		}
 
-		bool operator==(const FLevelBlueprintHandle& Other) const { return ParentLevel == Other.ParentLevel; }
-		bool operator!=(const FLevelBlueprintHandle& Other) const { return ParentLevel != Other.ParentLevel; }
+		bool operator==(const FLevelBlueprintHandle& Other) const { return LevelKey == LevelKey; }
+		bool operator!=(const FLevelBlueprintHandle& Other) const { return LevelKey != LevelKey; }
 	};
 
 	/** A tree item that represents a level blueprint */
diff --git a/Engine/Source/Editor/SceneOutliner/Public/WorldTreeItem.h b/Engine/Source/Editor/SceneOutliner/Public/WorldTreeItem.h
--- a/Engine/Source/Editor/SceneOutliner/Public/WorldTreeItem.h
+++ b/Engine/Source/Editor/SceneOutliner/Public/WorldTreeItem.h
@@ -11,6 +11,9 @@ namespace SceneOutliner
 	{
 		/** The world this tree item is associated with. */
 		mutable TWeakObjectPtr<UWorld> World;
+		
+		/** Constant identifier for this tree item */
+		const FObjectKey ID;
 
 		/** Construct this item from a world */
 		FWorldTreeItem(UWorld* InWorld);
