Various Scene Outliner improvements
 - Dropping onto the root world now correctly gives feedback to the user
 - Items that are filtered out, but still appear in the tree (due to matching children) can now be interacted with properly, but cannot be selected.
 - Fixed actor selection being synchronized in actor pickers.
 - There is now a distinction between filtered out and non-interactive items, defined by the filters themselves to specify whether filtered out items should also be non-interactive.
 - Fixed level blueprints from other levels showing up when the "Current Level Only" filter is enabled.

diff --git a/Engine/Source/Editor/SceneOutliner/Private/ActorTreeItem.cpp b/Engine/Source/Editor/SceneOutliner/Private/ActorTreeItem.cpp
--- a/Engine/Source/Editor/SceneOutliner/Private/ActorTreeItem.cpp
+++ b/Engine/Source/Editor/SceneOutliner/Private/ActorTreeItem.cpp
@@ -280,7 +280,7 @@ int32 FActorTreeItem::GetTypeSortPriority() const
 bool FActorTreeItem::CanInteract() const
 {
 	AActor* ActorPtr = Actor.Get();
-	if (!ActorPtr || Flags.bIsFilteredOut)
+	if (!ActorPtr || !Flags.bInteractive)
 	{
 		return false;
 	}
diff --git a/Engine/Source/Editor/SceneOutliner/Private/FolderTreeItem.cpp b/Engine/Source/Editor/SceneOutliner/Private/FolderTreeItem.cpp
--- a/Engine/Source/Editor/SceneOutliner/Private/FolderTreeItem.cpp
+++ b/Engine/Source/Editor/SceneOutliner/Private/FolderTreeItem.cpp
@@ -198,7 +198,7 @@ int32 FFolderTreeItem::GetTypeSortPriority() const
 
 bool FFolderTreeItem::CanInteract() const
 {
-	return !Flags.bIsFilteredOut;
+	return Flags.bInteractive;
 }
 
 void FFolderTreeItem::Delete()
@@ -224,6 +224,7 @@ void FFolderTreeItem::Delete()
 			{
 				Child.Pin()->Visit(ResetFolders);
 			}
+			FolderItem.Delete();
 		}
 	};
 
diff --git a/Engine/Source/Editor/SceneOutliner/Private/LevelBlueprintTreeItem.cpp b/Engine/Source/Editor/SceneOutliner/Private/LevelBlueprintTreeItem.cpp
--- a/Engine/Source/Editor/SceneOutliner/Private/LevelBlueprintTreeItem.cpp
+++ b/Engine/Source/Editor/SceneOutliner/Private/LevelBlueprintTreeItem.cpp
@@ -86,7 +86,7 @@ bool FLevelBlueprintTreeItem::CanInteract() const
 			return false;
 		}
 	}
-	return !Flags.bIsFilteredOut;
+	return Flags.bInteractive;
 }
 
 void FLevelBlueprintTreeItem::GenerateContextMenu(FMenuBuilder& MenuBuilder, SSceneOutliner& Outliner)
diff --git a/Engine/Source/Editor/SceneOutliner/Private/SOutlinerTreeView.cpp b/Engine/Source/Editor/SceneOutliner/Private/SOutlinerTreeView.cpp
--- a/Engine/Source/Editor/SceneOutliner/Private/SOutlinerTreeView.cpp
+++ b/Engine/Source/Editor/SceneOutliner/Private/SOutlinerTreeView.cpp
@@ -105,44 +105,6 @@ namespace SceneOutliner
 		return FReply::Handled();
 	}
 
-	struct FDropHandler : ITreeItemVisitor
-	{
-		/** Pointer to the scene outliner */
-		const TWeakPtr<SSceneOutliner>& SceneOutlinerWeak;
-
-		/** The drag/drop event we are processing */
-		const FDragDropEvent& DragDropEvent;
-
-		/** Whether or not we should actually apply the drop, or just validate it */
-		bool bApplyDrop;
-
-		/** The resulting validation information */
-		mutable FDragValidationInfo ValidationInfo;
-		
-		/** The reply we shouyld return, if applicable */
-		mutable FReply Reply;
-
-		/** Constructor */
-		FDropHandler(const TWeakPtr<SSceneOutliner>& InSceneOutlinerWeak, const FDragDropEvent& InDragDropEvent, bool bInApplyDrop)
-			: SceneOutlinerWeak(InSceneOutlinerWeak)
-			, DragDropEvent(InDragDropEvent)
-			, bApplyDrop(bInApplyDrop)
-			, ValidationInfo(FDragValidationInfo::Invalid())
-			, Reply(FReply::Unhandled())
-		{}
-
-		virtual void Visit(const FActorTreeItem& ActorItem) const override
-		{
-			FActorDropTarget DropTarget(ActorItem.Actor.Get());
-			Reply = HandleDrop(SceneOutlinerWeak, DragDropEvent, DropTarget, ValidationInfo, bApplyDrop);
-		}
-		virtual void Visit(const FFolderTreeItem& FolderItem) const override
-		{
-			FFolderDropTarget DropTarget(FolderItem.Path);
-			Reply = HandleDrop(SceneOutlinerWeak, DragDropEvent, DropTarget, ValidationInfo, bApplyDrop);
-		}
-	};
-
 	void SOutlinerTreeView::Construct(const SOutlinerTreeView::FArguments& InArgs, TSharedRef<SSceneOutliner> Owner)
 	{
 		SceneOutlinerWeak = Owner;
@@ -196,9 +158,8 @@ namespace SceneOutliner
 		auto ItemPtr = Item.Pin();
 		if (ItemPtr.IsValid())
 		{
-			FDropHandler Handler(SceneOutlinerWeak, DragDropEvent, true);
-			ItemPtr->Visit(Handler);
-			return Handler.Reply;
+			FDragValidationInfo ValidationInfo = FDragValidationInfo::Invalid();
+			return HandleDrop(SceneOutlinerWeak, DragDropEvent, *ItemPtr, ValidationInfo, true);
 		}
 
 		return FReply::Unhandled();
@@ -209,10 +170,10 @@ namespace SceneOutliner
 		auto ItemPtr = Item.Pin();
 		if (ItemPtr.IsValid())
 		{
-			FDropHandler Handler(SceneOutlinerWeak, DragDropEvent, false);
-			ItemPtr->Visit(Handler);
+			FDragValidationInfo ValidationInfo = FDragValidationInfo::Invalid();
 
-			UpdateOperationDecorator(DragDropEvent, Handler.ValidationInfo);
+			HandleDrop(SceneOutlinerWeak, DragDropEvent, *ItemPtr, ValidationInfo, false);
+			UpdateOperationDecorator(DragDropEvent, ValidationInfo);
 		}
 	}
 
@@ -240,7 +201,18 @@ namespace SceneOutliner
 			}
 		}
 
-		return FReply::Unhandled();
+		return FReply::Handled();
+	}
+	
+	FReply SSceneOutlinerTreeRow::OnMouseButtonUp(const FGeometry& MyGeometry, const FPointerEvent& MouseEvent)
+	{
+		auto ItemPtr = Item.Pin();
+		if (ItemPtr.IsValid() && ItemPtr->CanInteract())
+		{
+			return SMultiColumnTableRow<FTreeItemPtr>::OnMouseButtonUp(MyGeometry, MouseEvent);
+		}
+
+		return FReply::Handled();
 	}
 
 	TSharedRef<SWidget> SSceneOutlinerTreeRow::GenerateWidgetForColumn( const FName& ColumnName )
diff --git a/Engine/Source/Editor/SceneOutliner/Private/SOutlinerTreeView.h b/Engine/Source/Editor/SceneOutliner/Private/SOutlinerTreeView.h
--- a/Engine/Source/Editor/SceneOutliner/Private/SOutlinerTreeView.h
+++ b/Engine/Source/Editor/SceneOutliner/Private/SOutlinerTreeView.h
@@ -55,6 +55,7 @@ namespace SceneOutliner
 	protected:
 
 		virtual FReply OnMouseButtonDown(const FGeometry& MyGeometry, const FPointerEvent& MouseEvent) override;
+		virtual FReply OnMouseButtonUp(const FGeometry& MyGeometry, const FPointerEvent& MouseEvent) override;
 
 		virtual void OnDragEnter( const FGeometry& MyGeometry, const FDragDropEvent& DragDropEvent ) override;
 
diff --git a/Engine/Source/Editor/SceneOutliner/Private/SSceneOutliner.cpp b/Engine/Source/Editor/SceneOutliner/Private/SSceneOutliner.cpp
--- a/Engine/Source/Editor/SceneOutliner/Private/SSceneOutliner.cpp
+++ b/Engine/Source/Editor/SceneOutliner/Private/SSceneOutliner.cpp
@@ -37,10 +37,11 @@ namespace SceneOutliner
 {
 	TSharedPtr< FOutlinerFilter > CreateSelectedActorFilter()
 	{
-		auto Predicate = FActorFilterPredicate::CreateStatic([](const AActor* InActor){			
-			return InActor->IsSelected();
-		});
-		return MakeShareable( new FOutlinerPredicateFilter(Predicate, EDefaultFilterBehaviour::Fail) );
+		auto* Filter = new FOutlinerPredicateFilter(FActorFilterPredicate::CreateStatic([](const AActor* InActor){	return InActor->IsSelected(); }), EDefaultFilterBehaviour::Fail);
+
+		// If anything fails this filter, make it non interactive. We don't want to allow selection of implicitly included parents which might nuke the actor selection.
+		Filter->FailedItemState = EFailedFilterState::NonInteractive;
+		return MakeShareable( Filter );
 	}
 
 	TSharedPtr< FOutlinerFilter > CreateHideTemporaryActorsFilter()
@@ -52,10 +53,21 @@ namespace SceneOutliner
 
 	TSharedPtr< FOutlinerFilter > CreateIsInCurrentLevelFilter()
 	{
-		auto Predicate = FActorFilterPredicate::CreateStatic([](const AActor* InActor){
-			return InActor->GetLevel() == InActor->GetWorld()->GetCurrentLevel();
-		});
-		return MakeShareable( new FOutlinerPredicateFilter(Predicate, EDefaultFilterBehaviour::Pass) );
+		struct FOnlyCurrentLevelFilter : FOutlinerFilter
+		{
+			FOnlyCurrentLevelFilter() : FOutlinerFilter(EDefaultFilterBehaviour::Fail, EFailedFilterState::Interactive) {}
+
+			virtual bool PassesFilter(const AActor* InActor) const override
+			{
+				return InActor->GetLevel() == InActor->GetWorld()->GetCurrentLevel();
+			}
+			virtual bool PassesFilter(const FLevelBlueprintHandle& LevelBlueprint) const override
+			{
+				return LevelBlueprint.ParentLevel == GWorld->GetCurrentLevel();
+			}
+		};
+
+		return MakeShareable( new FOnlyCurrentLevelFilter() );
 	}
 
 	struct FItemSelection : IMutableTreeItemVisitor
@@ -137,6 +149,7 @@ namespace SceneOutliner
 		bNeedsRefresh = true;
 		bIsReentrant = false;
 		bSortDirty = true;
+		bActorSelectionDirty = SharedData->Mode == ESceneOutlinerMode::ActorBrowsing;
 		FilteredActorCount = 0;
 		SortOutlinerTimer = 0.0f;
 		bPendingFocusNextFrame = InInitOptions.bFocusSearchBoxWhenOpened;
@@ -778,10 +791,10 @@ namespace SceneOutliner
 	{
 		EmptyTreeItems();
 
+		ConstructItemFor<FWorldTreeItem>(SharedData->RepresentingWorld);
+
 		if (!SharedData->bOnlyShowFolders)
 		{
-			ConstructItemFor<FWorldTreeItem>(SharedData->RepresentingWorld);
-
 			// Add the level blueprints for all levels
 			for (auto It = SharedData->RepresentingWorld->GetLevelIterator(); It ; ++It)
 			{
@@ -902,7 +915,8 @@ namespace SceneOutliner
 				auto NewParent = Item->CreateParent();
 				if (NewParent.IsValid())
 				{
-					NewParent->Flags.bIsFilteredOut = !Filters->PassesAllFilters(*NewParent) || !SearchBoxFilter->PassesFilter(*NewParent);
+					NewParent->Flags.bIsFilteredOut = !Filters->TestAndSetInteractiveState(*NewParent) || !SearchBoxFilter->PassesFilter(*NewParent);
+
 					AddUnfilteredItemToTree(NewParent.ToSharedRef());
 					return NewParent;
 				}
@@ -1077,10 +1091,10 @@ namespace SceneOutliner
 					const bool bCanMoveSelection = ItemSelection.Worlds.Num() + ItemSelection.LevelBlueprints.Num() == 0;
 					if (bCanMoveSelection)
 					{
-					MenuBuilder.AddSubMenu(
-						LOCTEXT("MoveActorsTo", "Move To"),
-						LOCTEXT("MoveActorsTo_Tooltip", "Move selection to another folder"),
-						FNewMenuDelegate::CreateSP(this,  &SSceneOutliner::FillFoldersSubMenu));
+						MenuBuilder.AddSubMenu(
+							LOCTEXT("MoveActorsTo", "Move To"),
+							LOCTEXT("MoveActorsTo_Tooltip", "Move selection to another folder"),
+							FNewMenuDelegate::CreateSP(this,  &SSceneOutliner::FillFoldersSubMenu));
 					}
 
 					// If we've only got folders selected, show the selection sub menu
diff --git a/Engine/Source/Editor/SceneOutliner/Private/WorldTreeItem.cpp b/Engine/Source/Editor/SceneOutliner/Private/WorldTreeItem.cpp
--- a/Engine/Source/Editor/SceneOutliner/Private/WorldTreeItem.cpp
+++ b/Engine/Source/Editor/SceneOutliner/Private/WorldTreeItem.cpp
@@ -58,7 +58,7 @@ int32 FWorldTreeItem::GetTypeSortPriority() const
 
 bool FWorldTreeItem::CanInteract() const
 {
-	return !Flags.bIsFilteredOut;
+	return Flags.bInteractive;
 }
 
 void FWorldTreeItem::GenerateContextMenu(FMenuBuilder& MenuBuilder, SSceneOutliner& Outliner)
@@ -70,7 +70,15 @@ void FWorldTreeItem::GenerateContextMenu(FMenuBuilder& MenuBuilder, SSceneOutlin
 
 FDragValidationInfo FWorldTreeItem::ValidateDrop(FDragDropPayload& DraggedObjects, UWorld& World) const
 {
-	return FDragValidationInfo::Invalid();
+	// Dropping on the world means 'moving to the root' in folder terms
+	FFolderDropTarget Target(NAME_None);
+	return Target.ValidateDrop(DraggedObjects, World);
+}
+
+void FWorldTreeItem::OnDrop(FDragDropPayload& DraggedObjects, UWorld& World, const FDragValidationInfo& ValidationInfo, TSharedRef<SWidget> DroppedOnWidget)
+{
+	FFolderDropTarget Target(NAME_None);
+	return Target.OnDrop(DraggedObjects, World, ValidationInfo, DroppedOnWidget);
 }
 
 void FWorldTreeItem::OpenWorldSettings() const
diff --git a/Engine/Source/Editor/SceneOutliner/Public/ITreeItem.h b/Engine/Source/Editor/SceneOutliner/Public/ITreeItem.h
--- a/Engine/Source/Editor/SceneOutliner/Public/ITreeItem.h
+++ b/Engine/Source/Editor/SceneOutliner/Public/ITreeItem.h
@@ -32,11 +32,14 @@ namespace SceneOutliner
 			/* true if this item is filtered out */
 			bool bIsFilteredOut : 1;
 
+			/* true if this item can be interacted with as per the current outliner filters */
+			bool bInteractive : 1;
+
 			/** true if this item's children need to be sorted */
 			bool bChildrenRequireSort : 1;
 
 			/** Default constructor */		
-			FlagsType() : bIsExpanded(1), bIsFilteredOut(0) {}
+			FlagsType() : bIsExpanded(1), bIsFilteredOut(0), bInteractive(1), bChildrenRequireSort(1) {}
 		};
 
 	public:
diff --git a/Engine/Source/Editor/SceneOutliner/Public/SceneOutlinerFilters.h b/Engine/Source/Editor/SceneOutliner/Public/SceneOutlinerFilters.h
--- a/Engine/Source/Editor/SceneOutliner/Public/SceneOutlinerFilters.h
+++ b/Engine/Source/Editor/SceneOutliner/Public/SceneOutlinerFilters.h
@@ -17,6 +17,9 @@ namespace SceneOutliner
 	/** Enum to specify how items that are not explicitly handled by this filter should be managed */
 	enum class EDefaultFilterBehaviour : uint8 { Pass, Fail };
 
+	/** Enum that defines how a tree item should be dealt with in the case where it appears in the tree, but doesn't match the filter (eg if it has a matching child) */
+	enum class EFailedFilterState : uint8 { Interactive, NonInteractive };
+
 	/** A filter that can be applied to any type in the tree */
 	class FOutlinerFilter : public ITreeItemVisitor, public IFilter<const ITreeItem&>
 	{
@@ -25,6 +28,9 @@ namespace SceneOutliner
 		/** Event that is fired if this filter changes */
 		DECLARE_DERIVED_EVENT(FOutlinerFilter, IFilter<const ITreeItem&>::FChangedEvent, FChangedEvent);
 		virtual FChangedEvent& OnChanged() override { return ChangedEvent; }
+		
+		/** Enum that defines how a tree item should be dealt with in the case where it appears in the tree, but doesn't match the filter (eg if it has a matching child) */
+		EFailedFilterState FailedItemState;
 
 	protected:
 
@@ -33,9 +39,11 @@ namespace SceneOutliner
 
 		/** Default result of the filter when not overridden in derived classes */
 		const EDefaultFilterBehaviour DefaultBehaviour;
-		
+
 		/** Constructor to specify the default result of a filter */
-		FOutlinerFilter(EDefaultFilterBehaviour InDefaultBehaviour = EDefaultFilterBehaviour::Pass) : DefaultBehaviour(InDefaultBehaviour) {}
+		FOutlinerFilter(EDefaultFilterBehaviour InDefaultBehaviour, EFailedFilterState InFailedFilterState = EFailedFilterState::NonInteractive)
+			: FailedItemState(InFailedFilterState), DefaultBehaviour(InDefaultBehaviour)
+		{}
 
 		/** Overridden in derived types to filter actors */
 		virtual bool PassesFilter(const AActor* Actor) const { return DefaultBehaviour == EDefaultFilterBehaviour::Pass; }
@@ -113,23 +121,23 @@ namespace SceneOutliner
 		/** Predicate used to filter level blueprints */
 		mutable FLevelBlueprintFilterPredicate LevelBlueprintPred;
 
-		FOutlinerPredicateFilter(FActorFilterPredicate InActorPred, EDefaultFilterBehaviour InDefaultBehaviour)
-			: FOutlinerFilter(InDefaultBehaviour)
+		FOutlinerPredicateFilter(FActorFilterPredicate InActorPred, EDefaultFilterBehaviour InDefaultBehaviour, EFailedFilterState InFailedFilterState = EFailedFilterState::NonInteractive)
+			: FOutlinerFilter(InDefaultBehaviour, InFailedFilterState)
 			, ActorPred(InActorPred)
 		{}
 
-		FOutlinerPredicateFilter(FWorldFilterPredicate InWorldPred, EDefaultFilterBehaviour InDefaultBehaviour)
-			: FOutlinerFilter(InDefaultBehaviour)
+		FOutlinerPredicateFilter(FWorldFilterPredicate InWorldPred, EDefaultFilterBehaviour InDefaultBehaviour, EFailedFilterState InFailedFilterState = EFailedFilterState::NonInteractive)
+			: FOutlinerFilter(InDefaultBehaviour, InFailedFilterState)
 			, WorldPred(InWorldPred)
 		{}
 
-		FOutlinerPredicateFilter(FFolderFilterPredicate InFolderPred, EDefaultFilterBehaviour InDefaultBehaviour)
-			: FOutlinerFilter(InDefaultBehaviour)
+		FOutlinerPredicateFilter(FFolderFilterPredicate InFolderPred, EDefaultFilterBehaviour InDefaultBehaviour, EFailedFilterState InFailedFilterState = EFailedFilterState::NonInteractive)
+			: FOutlinerFilter(InDefaultBehaviour, InFailedFilterState)
 			, FolderPred(InFolderPred)
 		{}
 
-		FOutlinerPredicateFilter(FLevelBlueprintFilterPredicate InLevelBpPred, EDefaultFilterBehaviour InDefaultBehaviour)
-			: FOutlinerFilter(InDefaultBehaviour)
+		FOutlinerPredicateFilter(FLevelBlueprintFilterPredicate InLevelBpPred, EDefaultFilterBehaviour InDefaultBehaviour, EFailedFilterState InFailedFilterState = EFailedFilterState::NonInteractive)
+			: FOutlinerFilter(InDefaultBehaviour, InFailedFilterState)
 			, LevelBlueprintPred(InLevelBpPred)
 		{}
 
@@ -157,11 +165,42 @@ namespace SceneOutliner
 	/** Scene outliner filter class. This class abstracts the filtering of both actors and folders and allows for filtering on both types */
 	struct FOutlinerFilters : public TFilterCollection<const ITreeItem&>
 	{
+		/** Overridden to ensure we only ever have FOutlinerFilters added */
+		int32 Add(const TSharedPtr<FOutlinerFilter>& Filter)
+		{
+			return TFilterCollection::Add(Filter);
+		}
+
+		/** Test whether this tree item passes all filters, and set its interactive state according to the filter it failed (if applicable) */
+		bool TestAndSetInteractiveState(ITreeItem& InItem) const
+		{
+			bool bPassed = false;
+
+			// Default to interactive
+			InItem.Flags.bInteractive = true;
+			
+			for (const auto& Filter : ChildFilters)
+			{
+				if (!Filter->PassesFilter(InItem))
+				{
+					bPassed = false;
+					InItem.Flags.bInteractive = StaticCastSharedPtr<FOutlinerFilter>(Filter)->FailedItemState == EFailedFilterState::Interactive;
+					// If this has failed, but is still interactive, we carry on to see if any others fail *and* set to non-interactive
+					if (!InItem.Flags.bInteractive)
+					{
+						return false;
+					}
+				}
+			}
+
+			return bPassed;
+		}
+
 		/** Add a filter predicate to this filter collection */
 		template<typename T>
-		void AddFilterPredicate(T Predicate, EDefaultFilterBehaviour InDefaultBehaviour = EDefaultFilterBehaviour::Fail)
+		void AddFilterPredicate(T Predicate, EDefaultFilterBehaviour InDefaultBehaviour = EDefaultFilterBehaviour::Fail, EFailedFilterState InFailedFilterState = EFailedFilterState::NonInteractive)
 		{
-			Add(MakeShareable(new FOutlinerPredicateFilter(Predicate, InDefaultBehaviour)));
+			Add(MakeShareable(new FOutlinerPredicateFilter(Predicate, InDefaultBehaviour, InFailedFilterState)));
 		}
 	};
 }
\ No newline at end of file
diff --git a/Engine/Source/Editor/SceneOutliner/Public/WorldTreeItem.h b/Engine/Source/Editor/SceneOutliner/Public/WorldTreeItem.h
--- a/Engine/Source/Editor/SceneOutliner/Public/WorldTreeItem.h
+++ b/Engine/Source/Editor/SceneOutliner/Public/WorldTreeItem.h
@@ -54,9 +54,7 @@ namespace SceneOutliner
 		virtual FDragValidationInfo ValidateDrop(FDragDropPayload& DraggedObjects, UWorld& World) const override;
 
 		/** Called to drop the specified objects on this item. Only called if ValidateDrop() allows. */
-		virtual void OnDrop(FDragDropPayload& DraggedObjects, UWorld& World, const FDragValidationInfo& ValidationInfo, TSharedRef<SWidget> DroppedOnWidget) override
-		{
-		}
+		virtual void OnDrop(FDragDropPayload& DraggedObjects, UWorld& World, const FDragValidationInfo& ValidationInfo, TSharedRef<SWidget> DroppedOnWidget) override;
 
 	public:
 
