Level blueprints in the scene outliner now always represent the editor world, even when in PIE

diff --git a/Engine/Source/Editor/SceneOutliner/Private/SSceneOutliner.cpp b/Engine/Source/Editor/SceneOutliner/Private/SSceneOutliner.cpp
--- a/Engine/Source/Editor/SceneOutliner/Private/SSceneOutliner.cpp
+++ b/Engine/Source/Editor/SceneOutliner/Private/SSceneOutliner.cpp
@@ -63,7 +63,11 @@ namespace SceneOutliner
 			}
 			virtual bool PassesFilter(const FLevelBlueprintHandle& LevelBlueprint) const override
 			{
-				return LevelBlueprint.ParentLevel == GWorld->GetCurrentLevel();
+				if (ULevel* Level = LevelBlueprint.ParentLevel.Get())
+				{
+					return Level == Level->GetWorld()->GetCurrentLevel();
+				}
+				return false;
 			}
 		};
 
@@ -795,8 +799,8 @@ namespace SceneOutliner
 
 		if (!SharedData->bOnlyShowFolders)
 		{
-			// Add the level blueprints for all levels
-			for (auto It = SharedData->RepresentingWorld->GetLevelIterator(); It ; ++It)
+			// Add the level blueprints for all levels - we always use the editor world for this, even if we're PIE
+			for (auto It = GEditor->GetEditorWorldContext().World()->GetLevelIterator(); It ; ++It)
 			{
 				ConstructItemFor<FLevelBlueprintTreeItem>(FLevelBlueprintHandle(*It));
 			}
