Removed "Create New Folder" button in Scene Outliner overlay in Layers tab.

diff --git a/Engine/Source/Editor/Layers/Private/SLayerBrowser.cpp b/Engine/Source/Editor/Layers/Private/SLayerBrowser.cpp
--- a/Engine/Source/Editor/Layers/Private/SLayerBrowser.cpp
+++ b/Engine/Source/Editor/Layers/Private/SLayerBrowser.cpp
@@ -128,6 +128,7 @@ void SLayerBrowser::Construct(const FArguments& InArgs)
 		// We hide the header row to keep the UI compact.
 		InitOptions.bShowHeaderRow = false;
 		InitOptions.bShowParentTree = false;
+		InitOptions.bShowCreateNewFolder = false;
 		InitOptions.CustomDelete = FCustomSceneOutlinerDeleteDelegate::CreateSP(this, &SLayerBrowser::RemoveActorsFromSelectedLayer);
 
 		// Outliner Gutter
diff --git a/Engine/Source/Editor/SceneOutliner/Private/SSceneOutliner.cpp b/Engine/Source/Editor/SceneOutliner/Private/SSceneOutliner.cpp
--- a/Engine/Source/Editor/SceneOutliner/Private/SSceneOutliner.cpp
+++ b/Engine/Source/Editor/SceneOutliner/Private/SSceneOutliner.cpp
@@ -221,7 +221,7 @@ namespace SceneOutliner
 			.OnTextCommitted( this, &SSceneOutliner::OnFilterTextCommitted )
 		];
 
-		if (SharedData->Mode == ESceneOutlinerMode::ActorBrowsing)
+		if (SharedData->Mode == ESceneOutlinerMode::ActorBrowsing && InInitOptions.bShowCreateNewFolder)
 		{
 			Toolbar->AddSlot()
 				.VAlign(VAlign_Center)
diff --git a/Engine/Source/Editor/SceneOutliner/Public/SceneOutlinerPublicTypes.h b/Engine/Source/Editor/SceneOutliner/Public/SceneOutlinerPublicTypes.h
--- a/Engine/Source/Editor/SceneOutliner/Public/SceneOutlinerPublicTypes.h
+++ b/Engine/Source/Editor/SceneOutliner/Public/SceneOutlinerPublicTypes.h
@@ -128,6 +128,9 @@ namespace SceneOutliner
 		/** If true, the search box will gain focus when the scene outliner is created */
 		bool bFocusSearchBoxWhenOpened : 1;
 
+		/** If true, the Scene Outliner will expose a Create New Folder button */
+		bool bShowCreateNewFolder : 1;
+
 		/** Optional collection of filters to use when filtering in the Scene Outliner */
 		TSharedPtr<FOutlinerFilters> Filters;
 
@@ -141,6 +144,7 @@ namespace SceneOutliner
 			: bShowHeaderRow( true )
 			, bShowSearchBox( true )
 			, bFocusSearchBoxWhenOpened( false )
+			, bShowCreateNewFolder( true )
 			, Filters( new FOutlinerFilters )
 		{}
 	};
