Worlds and level blueprints no longer show FName numbers in the scene outliner

diff --git a/Engine/Source/Editor/SceneOutliner/Private/LevelBlueprintTreeItem.cpp b/Engine/Source/Editor/SceneOutliner/Private/LevelBlueprintTreeItem.cpp
--- a/Engine/Source/Editor/SceneOutliner/Private/LevelBlueprintTreeItem.cpp
+++ b/Engine/Source/Editor/SceneOutliner/Private/LevelBlueprintTreeItem.cpp
@@ -63,7 +63,7 @@ FString FLevelBlueprintTreeItem::GetDisplayString() const
 	FFormatOrderedArguments Args;
 	if (ULevel* LevelPtr = Handle.ParentLevel.Get())
 	{
-		return FPackageName::GetShortName(LevelPtr->GetOutermost());
+		return FPackageName::GetShortName(LevelPtr->GetOutermost()->GetFName().GetPlainNameString());
 	}
 	else
 	{
diff --git a/Engine/Source/Editor/SceneOutliner/Private/WorldTreeItem.cpp b/Engine/Source/Editor/SceneOutliner/Private/WorldTreeItem.cpp
--- a/Engine/Source/Editor/SceneOutliner/Private/WorldTreeItem.cpp
+++ b/Engine/Source/Editor/SceneOutliner/Private/WorldTreeItem.cpp
@@ -46,7 +46,7 @@ FString FWorldTreeItem::GetDisplayString() const
 {
 	if (UWorld* WorldPtr = World.Get())
 	{
-		return World->GetName();
+		return World->GetFName().GetPlainNameString();
 	}
 	return FString();
 }
