Editor: Fixed a copy-paste typo in a SSceneOutliner comment

diff --git a/Engine/Source/Editor/SceneOutliner/Private/SSceneOutliner.cpp b/Engine/Source/Editor/SceneOutliner/Private/SSceneOutliner.cpp
--- a/Engine/Source/Editor/SceneOutliner/Private/SSceneOutliner.cpp
+++ b/Engine/Source/Editor/SceneOutliner/Private/SSceneOutliner.cpp
@@ -2129,7 +2129,7 @@ namespace SceneOutliner
 		// We only allow these operations in actor browsing mode
 		if( SharedData->Mode == ESceneOutlinerMode::ActorBrowsing )
 		{
-			// Delete key: Delete selected actors (not rebindable, because it doesn't make much sense to bind.)
+			// Rename key: Rename selected actors (not rebindable, because it doesn't make much sense to bind.)
 			if( InKeyEvent.GetKey() == EKeys::F2 )
 			{
 				if (OutlinerTreeView->GetNumItemsSelected() == 1)
