Prevent AI controllers (and others) from being hidden in PIE Scene Outliner

We still hide transient actors in non-PIE worlds

diff --git a/Engine/Source/Editor/SceneOutliner/Private/SSceneOutliner.cpp b/Engine/Source/Editor/SceneOutliner/Private/SSceneOutliner.cpp
--- a/Engine/Source/Editor/SceneOutliner/Private/SSceneOutliner.cpp
+++ b/Engine/Source/Editor/SceneOutliner/Private/SSceneOutliner.cpp
@@ -1534,7 +1534,7 @@ namespace SceneOutliner
 		return	!SharedData->bOnlyShowFolders && 												// Don't show actors if we're only showing folders
 				Actor->IsEditable() &&															// Only show actors that are allowed to be selected and drawn in editor
 				Actor->IsListedInSceneOutliner() &&
-				!Actor->HasAnyFlags(RF_Transient) &&											// Don't show transient actors
+				(SharedData->bRepresentingPlayWorld || !Actor->HasAnyFlags(RF_Transient)) &&    // Don't show transient actors in non-play worlds
 				!Actor->IsTemplate() &&															// Should never happen, but we never want CDOs displayed
 				!FActorEditorUtils::IsABuilderBrush(Actor) &&									// Don't show the builder brush
 				!Actor->IsA( AWorldSettings::StaticClass() ) &&									// Don't show the WorldSettings actor, even though it is technically editable
