[INTEGRATE] Change 2426247 by [REDACTED_PERSON][REDACTED_HANDLE] on 2015/01/30 15:17:31

	Scene Outliner: Add a tooltip to world settings and level blueprint labels indicating that you can edit them by double-clicking, since they don't display in the details view when selected [UE-8455] [UE-8457]

diff --git a/Engine/Source/Editor/SceneOutliner/Private/SceneOutlinerItemLabelColumn.cpp b/Engine/Source/Editor/SceneOutliner/Private/SceneOutlinerItemLabelColumn.cpp
--- a/Engine/Source/Editor/SceneOutliner/Private/SceneOutlinerItemLabelColumn.cpp
+++ b/Engine/Source/Editor/SceneOutliner/Private/SceneOutlinerItemLabelColumn.cpp
@@ -14,6 +14,7 @@
 #include "EditorActorFolders.h"
 
 #include "SInlineEditableTextBlock.h"
+#include "IDocumentation.h"
 
 #define LOCTEXT_NAMESPACE "SceneOutlinerItemLabelColumn"
 
@@ -378,6 +379,11 @@ struct SWorldTreeLabel : FCommonLabelData, public SCompoundWidget
 				.Text(this, &SWorldTreeLabel::GetDisplayText)
 				.HighlightText(SceneOutliner.GetFilterHighlightText())
 				.ColorAndOpacity(this, &SWorldTreeLabel::GetForegroundColor)
+				.ToolTip(IDocumentation::Get()->CreateToolTip(
+					TAttribute<FText>(this, &SWorldTreeLabel::GetTooltipText),
+					nullptr,
+					TEXT("Shared/LevelEditor/SceneOutliner"),
+					TEXT("WorldSettingsLabel")))
 			]
 		];
 	}
@@ -391,6 +397,14 @@ private:
 		return Item.IsValid() ? FText::FromString(Item->GetDisplayString()) : FText();
 	}
 
+	FText GetTooltipText() const
+	{
+		auto Item = TreeItemPtr.Pin();
+		FText PersistentLevelDisplayName = Item.IsValid() ? FText::FromString(Item->GetDisplayString()) : FText();
+		
+		return FText::Format(LOCTEXT("WorldLabel_Tooltip", "The world settings for {0}, double-click to edit"), PersistentLevelDisplayName);
+	}
+
 	FSlateColor GetForegroundColor() const
 	{
 		if (auto BaseColor = FCommonLabelData::GetForegroundColor(TreeItemPtr.Pin()))
@@ -440,6 +454,11 @@ struct SLevelBlueprintTreeLabel : FCommonLabelData, public SCompoundWidget
 				.Text(this, &SLevelBlueprintTreeLabel::GetDisplayText)
 				.HighlightText(SceneOutliner.GetFilterHighlightText())
 				.ColorAndOpacity(this, &SLevelBlueprintTreeLabel::GetForegroundColor)
+				.ToolTip(IDocumentation::Get()->CreateToolTip(
+					TAttribute<FText>(this, &SLevelBlueprintTreeLabel::GetTooltipText),
+					nullptr,
+					TEXT("Shared/LevelEditor/SceneOutliner"),
+					TEXT("LevelBlueprintLabel")))
 			]
 		];
 	}
@@ -453,6 +472,14 @@ private:
 		return Item.IsValid() ? FText::FromString(Item->GetDisplayString()) : FText();
 	}
 
+	FText GetTooltipText() const
+	{
+		auto Item = TreeItemPtr.Pin();
+		FText LevelDisplayName = Item.IsValid() ? FText::FromString(Item->GetDisplayString()) : FText();
+
+		return FText::Format(LOCTEXT("LevelBlueprintLabel_Tooltip", "The Level Blueprint for {0}, double-click to edit"), LevelDisplayName);
+	}
+
 	FSlateColor GetForegroundColor() const
 	{
 		if (auto BaseColor = FCommonLabelData::GetForegroundColor(TreeItemPtr.Pin()))
