[INTEGRATE] Change 2436957 by [REDACTED_PERSON][REDACTED_HANDLE] on 2015/02/07 16:50:14

	Editor: Show hyperlinks for native classes in the world outliner, as long as they are game-specific classes

diff --git a/Engine/Source/Editor/SceneOutliner/Private/SceneOutlinerActorInfoColumn.cpp b/Engine/Source/Editor/SceneOutliner/Private/SceneOutlinerActorInfoColumn.cpp
--- a/Engine/Source/Editor/SceneOutliner/Private/SceneOutlinerActorInfoColumn.cpp
+++ b/Engine/Source/Editor/SceneOutliner/Private/SceneOutlinerActorInfoColumn.cpp
@@ -217,12 +217,29 @@ TSharedPtr<SWidget> FActorInfoColumn::ConstructClassHyperlink( ITreeItem& TreeIt
 		{
 			if (AActor* Actor = ActorItem.Actor.Get())
 			{
-				UClass* ActorClass = Actor->GetClass();
-
-				// Only create source link in SO for Blueprints (not native classes, too noisy having header link in SO)
-				if (ActorClass && UBlueprint::GetBlueprintFromClass(ActorClass) != NULL)
+				if (UClass* ActorClass = Actor->GetClass())
 				{
-					return FEditorClassUtils::GetSourceLink(ActorClass, Actor);
+					// Always show blueprints
+					const bool bIsBlueprintClass = UBlueprint::GetBlueprintFromClass(ActorClass) != nullptr;
+
+					// Also show game or game plugin native classes (but not engine classes as that makes the scene outliner pretty noisy)
+					bool bIsGameClass = false;
+					if (!bIsBlueprintClass)
+					{
+						UPackage* Package = ActorClass->GetOutermost();
+						const FString ModuleName = FPackageName::GetShortName(Package->GetFName());
+
+						FModuleStatus PackageModuleStatus;
+						if (FModuleManager::Get().QueryModule(*ModuleName, /*out*/ PackageModuleStatus))
+						{
+							bIsGameClass = PackageModuleStatus.bIsGameModule;
+						}
+					}
+
+					if (bIsBlueprintClass || bIsGameClass)
+					{
+						return FEditorClassUtils::GetSourceLink(ActorClass, Actor);
+					}
 				}
 			}
 
