[INTEGRATE] Change 2440620 by [REDACTED_PERSON][REDACTED_HANDLE] on 2015/02/10 18:23:29

	Changed color for dynamically spawned actors from blue to yellow
	- The color blue is usually used for Blueprint-related objects in UE4

diff --git a/Engine/Source/Editor/SceneOutliner/Private/SceneOutlinerItemLabelColumn.cpp b/Engine/Source/Editor/SceneOutliner/Private/SceneOutlinerItemLabelColumn.cpp
--- a/Engine/Source/Editor/SceneOutliner/Private/SceneOutlinerItemLabelColumn.cpp
+++ b/Engine/Source/Editor/SceneOutliner/Private/SceneOutlinerItemLabelColumn.cpp
@@ -266,7 +266,7 @@ private:
 		if (TreeItem->GetSharedData().bRepresentingPlayWorld && !TreeItem->bExistsInCurrentWorldAndPIE)
 		{
 			// Highlight actors that are exclusive to PlayWorld
-			return FLinearColor(0.12f, 0.56f, 1.0f);
+			return FLinearColor(0.9f, 0.8f, 0.4f);
 		}
 
 		// also darken items that are non selectable in the active mode(s)
