Child actors can no longer be dragged around in the scene outliner

They should be tied to their parent, and thus can no longer be re-parented using the scene outliner.
This addresses UE-7781 - Can attach a Child Actor to an Actor that isn't its parent.

diff --git a/Engine/Source/Editor/SceneOutliner/Private/ActorTreeItem.cpp b/Engine/Source/Editor/SceneOutliner/Private/ActorTreeItem.cpp
--- a/Engine/Source/Editor/SceneOutliner/Private/ActorTreeItem.cpp
+++ b/Engine/Source/Editor/SceneOutliner/Private/ActorTreeItem.cpp
@@ -44,10 +44,20 @@ FDragValidationInfo FActorDropTarget::ValidateDrop(FDragDropPayload& DraggedObje
 		AActor* DragActor = Actor.Get();
 		if (DragActor)
 		{
-			if (!GEditor->CanParentActors(DropTarget, DragActor, &AttachErrorMsg))
+			if (bCanAttach)
 			{
-				bCanAttach = false;
+				if (DragActor->ParentComponentActor.Get())
+				{
+					AttachErrorMsg = FText::Format(LOCTEXT("Error_AttachChildActor", "Cannot move {0} as it is a child actor."), FText::FromString(DragActor->GetActorLabel()));
+					bCanAttach = bDraggedOntoAttachmentParent = false;
+					break;
+				}
+				if (!GEditor->CanParentActors(DropTarget, DragActor, &AttachErrorMsg))
+				{
+					bCanAttach = false;
+				}
 			}
+
 			if (DragActor->GetAttachParentActor() != DropTarget)
 			{
 				bDraggedOntoAttachmentParent = false;
diff --git a/Engine/Source/Editor/SceneOutliner/Private/FolderTreeItem.cpp b/Engine/Source/Editor/SceneOutliner/Private/FolderTreeItem.cpp
--- a/Engine/Source/Editor/SceneOutliner/Private/FolderTreeItem.cpp
+++ b/Engine/Source/Editor/SceneOutliner/Private/FolderTreeItem.cpp
@@ -74,8 +74,11 @@ FDragValidationInfo FFolderDropTarget::ValidateDrop(FDragDropPayload& DraggedObj
 		for (auto WeakActor : DraggedObjects.Actors.GetValue())
 		{
 			AActor* Actor = WeakActor.Get();
-
-			if (Actor->GetFolderPath() == DestinationPath && !Actor->GetAttachParentActor())
+			if (Actor->ParentComponentActor.Get())
+			{
+				return FDragValidationInfo(FActorDragDropGraphEdOp::ToolTip_IncompatibleGeneric, FText::Format(LOCTEXT("Error_AttachChildActor", "Cannot move {0} as it is a child actor."), FText::FromString(Actor->GetActorLabel())));
+			}
+			else if (Actor->GetFolderPath() == DestinationPath && !Actor->GetAttachParentActor())
 			{
 				FFormatNamedArguments Args;
 				Args.Add(TEXT("SourceName"), FText::FromString(Actor->GetActorLabel()));
