Added "Create Folder" to Level Root world outliner context menu
This addresses UE-8541

diff --git a/Engine/Source/Editor/SceneOutliner/Private/WorldTreeItem.cpp b/Engine/Source/Editor/SceneOutliner/Private/WorldTreeItem.cpp
--- a/Engine/Source/Editor/SceneOutliner/Private/WorldTreeItem.cpp
+++ b/Engine/Source/Editor/SceneOutliner/Private/WorldTreeItem.cpp
@@ -63,11 +63,31 @@ bool FWorldTreeItem::CanInteract() const
 
 void FWorldTreeItem::GenerateContextMenu(FMenuBuilder& MenuBuilder, SSceneOutliner& Outliner)
 {
+	auto SharedOutliner = StaticCastSharedRef<SSceneOutliner>(Outliner.AsShared());
+	
 	const FSlateIcon WorldSettingsIcon(FEditorStyle::GetStyleSetName(), "LevelEditor.WorldProperties.Tab");
+	const FSlateIcon NewFolderIcon(FEditorStyle::GetStyleSetName(), "SceneOutliner.NewFolderIcon");
 
+	MenuBuilder.AddMenuEntry(LOCTEXT("CreateFolder", "Create Folder"), FText(), NewFolderIcon, FUIAction(FExecuteAction::CreateSP(this, &FWorldTreeItem::CreateFolder, TWeakPtr<SSceneOutliner>(SharedOutliner))));
 	MenuBuilder.AddMenuEntry(LOCTEXT("OpenWorldSettings", "World Settings"), FText(), WorldSettingsIcon, FExecuteAction::CreateSP(this, &FWorldTreeItem::OpenWorldSettings));
 }
 
+void FWorldTreeItem::CreateFolder(TWeakPtr<SSceneOutliner> WeakOutliner)
+{
+	auto Outliner = WeakOutliner.Pin();
+
+	if (Outliner.IsValid() && SharedData->RepresentingWorld)
+	{
+		const FScopedTransaction Transaction(LOCTEXT("UndoAction_CreateFolder", "Create Folder"));
+
+		const FName NewFolderName = FActorFolders::Get().GetDefaultFolderName(*SharedData->RepresentingWorld, "");
+		FActorFolders::Get().CreateFolder(*SharedData->RepresentingWorld, NewFolderName);
+
+		// At this point the new folder will be in our newly added list, so select it and open a rename when it gets refreshed
+		Outliner->OnItemAdded(NewFolderName, ENewItemAction::Select | ENewItemAction::Rename);
+	}
+}
+
 FDragValidationInfo FWorldTreeItem::ValidateDrop(FDragDropPayload& DraggedObjects, UWorld& World) const
 {
 	// Dropping on the world means 'moving to the root' in folder terms
diff --git a/Engine/Source/Editor/SceneOutliner/Public/WorldTreeItem.h b/Engine/Source/Editor/SceneOutliner/Public/WorldTreeItem.h
--- a/Engine/Source/Editor/SceneOutliner/Public/WorldTreeItem.h
+++ b/Engine/Source/Editor/SceneOutliner/Public/WorldTreeItem.h
@@ -56,6 +56,11 @@ namespace SceneOutliner
 		/** Called to drop the specified objects on this item. Only called if ValidateDrop() allows. */
 		virtual void OnDrop(FDragDropPayload& DraggedObjects, UWorld& World, const FDragValidationInfo& ValidationInfo, TSharedRef<SWidget> DroppedOnWidget) override;
 
+	private:
+
+		/** Create a new folder at the root of this world */
+		void CreateFolder(TWeakPtr<SSceneOutliner> WeakOutliner);
+		
 	public:
 
 		/** Open the world settings for the contained world */
