Non-unity compile fix

diff --git a/Engine/Source/Editor/SceneOutliner/Private/WorldTreeItem.cpp b/Engine/Source/Editor/SceneOutliner/Private/WorldTreeItem.cpp
--- a/Engine/Source/Editor/SceneOutliner/Private/WorldTreeItem.cpp
+++ b/Engine/Source/Editor/SceneOutliner/Private/WorldTreeItem.cpp
@@ -4,6 +4,7 @@
 
 #include "WorldTreeItem.h"
 #include "LevelEditor.h"
+#include "EditorActorFolders.h"
 
 #define LOCTEXT_NAMESPACE "SceneOutliner_WorldTreeItem"
 
