[UE-8455] Remove level blueprint items from the World Outliner view.

change summary:
- removed FLevelBlueprintTreeItem.h/cpp and related code (per discussion threads)

diff --git a/Engine/Source/Editor/SceneOutliner/Private/LevelBlueprintTreeItem.cpp b/Engine/Source/Editor/SceneOutliner/Private/LevelBlueprintTreeItem.cpp
deleted file mode 100644
diff --git a/Engine/Source/Editor/SceneOutliner/Private/SSceneOutliner.cpp b/Engine/Source/Editor/SceneOutliner/Private/SSceneOutliner.cpp
--- a/Engine/Source/Editor/SceneOutliner/Private/SSceneOutliner.cpp
+++ b/Engine/Source/Editor/SceneOutliner/Private/SSceneOutliner.cpp
@@ -61,14 +61,6 @@ namespace SceneOutliner
 			{
 				return InActor->GetLevel() == InActor->GetWorld()->GetCurrentLevel();
 			}
-			virtual bool PassesFilter(const FLevelBlueprintHandle& LevelBlueprint) const override
-			{
-				if (ULevel* Level = LevelBlueprint.ParentLevel.Get())
-				{
-					return Level == Level->GetWorld()->GetCurrentLevel();
-				}
-				return false;
-			}
 		};
 
 		return MakeShareable( new FOnlyCurrentLevelFilter() );
@@ -78,7 +70,6 @@ namespace SceneOutliner
 	{
 		mutable TArray<FActorTreeItem*> Actors;
 		mutable TArray<FWorldTreeItem*> Worlds;
-		mutable TArray<FLevelBlueprintTreeItem*> LevelBlueprints;
 		mutable TArray<FFolderTreeItem*> Folders;
 
 		FItemSelection()
@@ -127,10 +118,6 @@ namespace SceneOutliner
 		{
 			Worlds.Add(&WorldItem);
 		}
-		virtual void Visit(FLevelBlueprintTreeItem& LevelBlueprintItem) const override
-		{
-			LevelBlueprints.Add(&LevelBlueprintItem);
-		}
 		virtual void Visit(FFolderTreeItem& FolderItem) const override
 		{
 			Folders.Add(&FolderItem);
@@ -802,12 +789,6 @@ namespace SceneOutliner
 
 		if (!SharedData->bOnlyShowFolders)
 		{
-			// Add the level blueprints for all levels - we always use the editor world for this, even if we're PIE
-			for (auto It = GEditor->GetEditorWorldContext().World()->GetLevelIterator(); It ; ++It)
-			{
-				ConstructItemFor<FLevelBlueprintTreeItem>(FLevelBlueprintHandle(*It));
-			}
-
 			// Iterate over every actor in memory. WARNING: This is potentially very expensive!
 			for( FActorIterator ActorIt(SharedData->RepresentingWorld); ActorIt; ++ActorIt )
 			{
@@ -1108,7 +1089,7 @@ namespace SceneOutliner
 				else
 				{
 					// Can't move worlds or level blueprints
-					const bool bCanMoveSelection = ItemSelection.Worlds.Num() + ItemSelection.LevelBlueprints.Num() == 0;
+					const bool bCanMoveSelection = ItemSelection.Worlds.Num() == 0;
 					if (bCanMoveSelection)
 					{
 						MenuBuilder.AddSubMenu(
@@ -1725,10 +1706,6 @@ namespace SceneOutliner
 
 					.Folder(ExpandCollapseFolder)
 
-					.LevelBlueprint([](const FLevelBlueprintTreeItem& LevelBlueprintItem){
-						LevelBlueprintItem.Open();
-					})
-
 					.World([](const FWorldTreeItem& WorldItem){
 						WorldItem.OpenWorldSettings();
 					})
diff --git a/Engine/Source/Editor/SceneOutliner/Private/SceneOutlinerActorInfoColumn.cpp b/Engine/Source/Editor/SceneOutliner/Private/SceneOutlinerActorInfoColumn.cpp
--- a/Engine/Source/Editor/SceneOutliner/Private/SceneOutlinerActorInfoColumn.cpp
+++ b/Engine/Source/Editor/SceneOutliner/Private/SceneOutlinerActorInfoColumn.cpp
@@ -86,25 +86,6 @@ struct FGetInfo : TTreeItemGetter<FString>
 			return FString();
 		}
 	}
-
-	virtual FString Get(const FLevelBlueprintTreeItem& Item) const override
-	{
-		switch(CurrentMode)
-		{
-		case ECustomColumnMode::Class:
-			return LOCTEXT("LevelBlueprintTypeName", "Level Blueprint").ToString();
-
-		case ECustomColumnMode::Level:
-			if (ULevel* Level = Item.Handle.ParentLevel.Get())
-			{
-				return FPackageName::GetShortName(Level->GetOutermost()->GetName());
-			}
-			return FString();
-
-		default:
-			return FString();
-		}
-	}
 };
 
 
diff --git a/Engine/Source/Editor/SceneOutliner/Private/SceneOutlinerGutter.cpp b/Engine/Source/Editor/SceneOutliner/Private/SceneOutlinerGutter.cpp
--- a/Engine/Source/Editor/SceneOutliner/Private/SceneOutlinerGutter.cpp
+++ b/Engine/Source/Editor/SceneOutliner/Private/SceneOutlinerGutter.cpp
@@ -55,11 +55,6 @@ bool FGetVisibilityVisitor::Get(const FWorldTreeItem& WorldItem) const
 	return RecurseChildren(WorldItem);
 }
 
-bool FGetVisibilityVisitor::Get(const FLevelBlueprintTreeItem&) const
-{
-	return false;
-}
-
 bool FGetVisibilityVisitor::Get(const FFolderTreeItem& FolderItem) const
 {
 	return RecurseChildren(FolderItem);
@@ -334,16 +329,6 @@ SHeaderRow::FColumn::FArguments FSceneOutlinerGutter::ConstructHeaderRowColumn()
 
 const TSharedRef<SWidget> FSceneOutlinerGutter::ConstructRowWidget(FTreeItemRef TreeItem, const STableRow<FTreeItemPtr>& Row)
 {
-	bool bGutterIsApplicable = true;
-	TreeItem->Visit(FFunctionalVisitor().LevelBlueprint([&](const FLevelBlueprintTreeItem&){
-		bGutterIsApplicable = false;
-	}));
-
-	if (!bGutterIsApplicable)
-	{
-		return SNullWidget::NullWidget;
-	}
-	
 	return SNew(SHorizontalBox)
 		+SHorizontalBox::Slot()
 		.AutoWidth()
diff --git a/Engine/Source/Editor/SceneOutliner/Private/SceneOutlinerGutter.h b/Engine/Source/Editor/SceneOutliner/Private/SceneOutlinerGutter.h
--- a/Engine/Source/Editor/SceneOutliner/Private/SceneOutlinerGutter.h
+++ b/Engine/Source/Editor/SceneOutliner/Private/SceneOutlinerGutter.h
@@ -22,9 +22,6 @@ struct FGetVisibilityVisitor : TTreeItemGetter<bool>
 	/** Get a World's visibility */
 	virtual bool Get(const FWorldTreeItem& WorldItem) const override;
 
-	/** Get a Level Blueprint's visibility */
-	virtual bool Get(const FLevelBlueprintTreeItem& LevelBlueprintItem) const override;
-
 	/** Get a folder's visibility */
 	virtual bool Get(const FFolderTreeItem& FolderItem) const override;
 };
diff --git a/Engine/Source/Editor/SceneOutliner/Private/SceneOutlinerItemLabelColumn.cpp b/Engine/Source/Editor/SceneOutliner/Private/SceneOutlinerItemLabelColumn.cpp
--- a/Engine/Source/Editor/SceneOutliner/Private/SceneOutlinerItemLabelColumn.cpp
+++ b/Engine/Source/Editor/SceneOutliner/Private/SceneOutlinerItemLabelColumn.cpp
@@ -416,81 +416,6 @@ private:
 	}
 };
 
-struct SLevelBlueprintTreeLabel : FCommonLabelData, public SCompoundWidget
-{
-	SLATE_BEGIN_ARGS(SLevelBlueprintTreeLabel){}
-	SLATE_END_ARGS()
-
-	void Construct(const FArguments& InArgs, FLevelBlueprintTreeItem& LevelBlueprintItem, ISceneOutliner& SceneOutliner, const STableRow<FTreeItemPtr>& InRow)
-	{
-		TreeItemPtr = StaticCastSharedRef<FLevelBlueprintTreeItem>(LevelBlueprintItem.AsShared());
-		WeakSceneOutliner = StaticCastSharedRef<ISceneOutliner>(SceneOutliner.AsShared());
-
-		ChildSlot
-		[
-			SNew(SHorizontalBox)
-
-			+ SHorizontalBox::Slot()
-			.AutoWidth()
-			.VAlign(VAlign_Center)
-			.Padding(FDefaultTreeItemMetrics::IconPadding())
-			[
-				SNew(SBox)
-				.WidthOverride(FDefaultTreeItemMetrics::IconSize())
-				.HeightOverride(FDefaultTreeItemMetrics::IconSize())
-				[
-					SNew(SImage)
-					.Image(FEditorStyle::GetBrush("LevelEditor.OpenLevelBlueprint"))
-					.ToolTipText(LOCTEXT("LevelBlueprint_Tooltip", "Level Blueprint"))
-				]
-			]
-
-			+ SHorizontalBox::Slot()
-			.FillWidth(1.0f)
-			.VAlign(VAlign_Center)
-			.Padding(0.0f, 2.0f)
-			[
-				SNew(STextBlock)
-				.Text(this, &SLevelBlueprintTreeLabel::GetDisplayText)
-				.HighlightText(SceneOutliner.GetFilterHighlightText())
-				.ColorAndOpacity(this, &SLevelBlueprintTreeLabel::GetForegroundColor)
-				.ToolTip(IDocumentation::Get()->CreateToolTip(
-					TAttribute<FText>(this, &SLevelBlueprintTreeLabel::GetTooltipText),
-					nullptr,
-					TEXT("Shared/LevelEditor/SceneOutliner"),
-					TEXT("LevelBlueprintLabel")))
-			]
-		];
-	}
-
-private:
-	TWeakPtr<FLevelBlueprintTreeItem> TreeItemPtr;
-
-	FText GetDisplayText() const
-	{
-		auto Item = TreeItemPtr.Pin();
-		return Item.IsValid() ? FText::FromString(Item->GetDisplayString()) : FText();
-	}
-
-	FText GetTooltipText() const
-	{
-		auto Item = TreeItemPtr.Pin();
-		FText LevelDisplayName = Item.IsValid() ? FText::FromString(Item->GetDisplayString()) : FText();
-
-		return FText::Format(LOCTEXT("LevelBlueprintLabel_Tooltip", "The Level Blueprint for {0}, double-click to edit"), LevelDisplayName);
-	}
-
-	FSlateColor GetForegroundColor() const
-	{
-		if (auto BaseColor = FCommonLabelData::GetForegroundColor(TreeItemPtr.Pin()))
-		{
-			return BaseColor.GetValue();
-		}
-		
-		return FSlateColor::UseForeground();
-	}
-};
-
 struct SFolderTreeLabel : FCommonLabelData, public SCompoundWidget
 {
 	SLATE_BEGIN_ARGS(SFolderTreeLabel){}
@@ -710,13 +635,6 @@ TSharedRef<SWidget> FItemLabelColumn::GenerateWidget( FWorldTreeItem& TreeItem,
 	return SNew(SWorldTreeLabel, TreeItem, *Outliner, InRow);
 }
 
-TSharedRef<SWidget> FItemLabelColumn::GenerateWidget( FLevelBlueprintTreeItem& TreeItem, const STableRow<FTreeItemPtr>& InRow )
-{
-	ISceneOutliner* Outliner = WeakSceneOutliner.Pin().Get();
-	check(Outliner);
-	return SNew(SLevelBlueprintTreeLabel, TreeItem, *Outliner, InRow);
-}
-
 TSharedRef<SWidget> FItemLabelColumn::GenerateWidget( FFolderTreeItem& TreeItem, const STableRow<FTreeItemPtr>& InRow )
 {
 	ISceneOutliner* Outliner = WeakSceneOutliner.Pin().Get();
diff --git a/Engine/Source/Editor/SceneOutliner/Private/SceneOutlinerItemLabelColumn.h b/Engine/Source/Editor/SceneOutliner/Private/SceneOutlinerItemLabelColumn.h
--- a/Engine/Source/Editor/SceneOutliner/Private/SceneOutlinerItemLabelColumn.h
+++ b/Engine/Source/Editor/SceneOutliner/Private/SceneOutlinerItemLabelColumn.h
@@ -46,7 +46,6 @@ private:
 
 	TSharedRef<SWidget> GenerateWidget( FActorTreeItem& TreeItem, const STableRow<FTreeItemPtr>& InRow );
 	TSharedRef<SWidget> GenerateWidget( FWorldTreeItem& TreeItem, const STableRow<FTreeItemPtr>& InRow );
-	TSharedRef<SWidget> GenerateWidget( FLevelBlueprintTreeItem& TreeItem, const STableRow<FTreeItemPtr>& InRow );
 	TSharedRef<SWidget> GenerateWidget( FFolderTreeItem& TreeItem, const STableRow<FTreeItemPtr>& InRow );
 
 	struct FColumnWidgetGenerator : FColumnGenerator
@@ -58,7 +57,6 @@ private:
 
 		virtual TSharedRef<SWidget> GenerateWidget(FFolderTreeItem& Item) const override { return Column.GenerateWidget(Item, Row); }
 		virtual TSharedRef<SWidget> GenerateWidget(FWorldTreeItem& Item) const override { return Column.GenerateWidget(Item, Row); }
-		virtual TSharedRef<SWidget> GenerateWidget(FLevelBlueprintTreeItem& Item) const override { return Column.GenerateWidget(Item, Row); }
 		virtual TSharedRef<SWidget> GenerateWidget(FActorTreeItem& Item) const override { return Column.GenerateWidget(Item, Row); }
 	};
 	friend FColumnWidgetGenerator;
diff --git a/Engine/Source/Editor/SceneOutliner/Private/SceneOutlinerStandaloneTypes.h b/Engine/Source/Editor/SceneOutliner/Private/SceneOutlinerStandaloneTypes.h
--- a/Engine/Source/Editor/SceneOutliner/Private/SceneOutlinerStandaloneTypes.h
+++ b/Engine/Source/Editor/SceneOutliner/Private/SceneOutlinerStandaloneTypes.h
@@ -47,13 +47,6 @@ namespace SceneOutliner
 			CachedHash = CalculateTypeHash();
 		}
 
-		/** ID representing a level blueprint */
-		explicit FTreeItemID(const FLevelBlueprintHandle& InLevelBlueprint) : Type(EType::LevelBlueprint)
-		{
-			new (Data) FLevelBlueprintHandle(InLevelBlueprint);
-			CachedHash = CalculateTypeHash();
-		}
-
 		/** ID representing a folder */
 		FTreeItemID(const FName& InFolder) : Type(EType::Folder)
 		{
@@ -72,7 +65,6 @@ namespace SceneOutliner
 			switch(Type)
 			{
 				case EType::Object:			new (Data) FObjectKey(Other.GetAsObjectKey());											break;
-				case EType::LevelBlueprint:	new (Data) FLevelBlueprintHandle(Other.GetAsLevelBlueprintRef());						break;
 				case EType::Folder:			new (Data) FName(Other.GetAsFolderRef());												break;
 				default:																											break;
 			}
@@ -97,7 +89,6 @@ namespace SceneOutliner
 			switch(Type)
 			{
 				case EType::Object:			GetAsObjectKey().~FObjectKey();							break;
-				case EType::LevelBlueprint:	GetAsLevelBlueprintRef().~FLevelBlueprintHandle();		break;
 				case EType::Folder:			GetAsFolderRef().~FName();								break;
 				default:																			break;
 			}
@@ -118,7 +109,6 @@ namespace SceneOutliner
 			switch(Type)
 			{
 				case EType::Object:			Hash = GetTypeHash(GetAsObjectKey());				break;
-				case EType::LevelBlueprint:	Hash = GetTypeHash(GetAsLevelBlueprintRef());		break;
 				case EType::Folder:			Hash = GetTypeHash(GetAsFolderRef());				break;
 				default:																		break;
 			}
@@ -134,7 +124,6 @@ namespace SceneOutliner
 	private:
 
 		FObjectKey& 				GetAsObjectKey() const 			{ return *reinterpret_cast<FObjectKey*>(Data); }
-		FLevelBlueprintHandle& 		GetAsLevelBlueprintRef() const 	{ return *reinterpret_cast<FLevelBlueprintHandle*>(Data); }
 		FName& 						GetAsFolderRef() const			{ return *reinterpret_cast<FName*>(Data); }
 
 		/** Compares the specified ID with this one - assumes matching types */
@@ -143,18 +132,17 @@ namespace SceneOutliner
 			switch(Type)
 			{
 				case EType::Object:			return GetAsObjectKey() == Other.GetAsObjectKey();
-				case EType::LevelBlueprint:	return GetAsLevelBlueprintRef() == Other.GetAsLevelBlueprintRef();
 				case EType::Folder:			return GetAsFolderRef() == Other.GetAsFolderRef();
 				case EType::Null:			return true;
 				default: check(false);		return false;
 			}
 		}
 
-		enum class EType : uint8 { Object, LevelBlueprint, Folder, Null };
+		enum class EType : uint8 { Object, Folder, Null };
 		EType Type;
 
 		uint32 CachedHash;
-		static const uint32 MaxSize = TMaxSizeof<FObjectKey, FLevelBlueprintHandle, FName>::Value;
+		static const uint32 MaxSize = TMaxSizeof<FObjectKey, FName>::Value;
 		mutable uint8 Data[MaxSize];
 	};
 
diff --git a/Engine/Source/Editor/SceneOutliner/Public/LevelBlueprintTreeItem.h b/Engine/Source/Editor/SceneOutliner/Public/LevelBlueprintTreeItem.h
deleted file mode 100644
diff --git a/Engine/Source/Editor/SceneOutliner/Public/SceneOutliner.h b/Engine/Source/Editor/SceneOutliner/Public/SceneOutliner.h
--- a/Engine/Source/Editor/SceneOutliner/Public/SceneOutliner.h
+++ b/Engine/Source/Editor/SceneOutliner/Public/SceneOutliner.h
@@ -17,6 +17,5 @@
 #include "ActorTreeItem.h"
 #include "FolderTreeItem.h"
 #include "WorldTreeItem.h"
-#include "LevelBlueprintTreeItem.h"
 
 #include "SceneOutlinerFilters.h"
\ No newline at end of file
diff --git a/Engine/Source/Editor/SceneOutliner/Public/SceneOutlinerFilters.h b/Engine/Source/Editor/SceneOutliner/Public/SceneOutlinerFilters.h
--- a/Engine/Source/Editor/SceneOutliner/Public/SceneOutlinerFilters.h
+++ b/Engine/Source/Editor/SceneOutliner/Public/SceneOutlinerFilters.h
@@ -9,7 +9,6 @@
 #include "FolderTreeItem.h"
 #include "ActorTreeItem.h"
 #include "WorldTreeItem.h"
-#include "LevelBlueprintTreeItem.h"
 
 namespace SceneOutliner
 {
@@ -50,9 +49,6 @@ namespace SceneOutliner
 		
 		/** Overridden in derived types to filter worlds */
 		virtual bool PassesFilter(const UWorld* World) const { return DefaultBehaviour == EDefaultFilterBehaviour::Pass; }
-		
-		/** Overridden in derived types to filter level blueprints */
-		virtual bool PassesFilter(const FLevelBlueprintHandle& LevelBlueprint) const { return DefaultBehaviour == EDefaultFilterBehaviour::Pass; }
 
 		/** Overridden in derived types to filter folders */
 		virtual bool PassesFilter(FName Folder) const { return DefaultBehaviour == EDefaultFilterBehaviour::Pass; }
@@ -86,11 +82,6 @@ namespace SceneOutliner
 			}
 		}
 
-		virtual void Visit(const FLevelBlueprintTreeItem& LevelBlueprintItem) const override
-		{
-			bTransientFilterResult = PassesFilter(LevelBlueprintItem.Handle);
-		}
-
 		virtual void Visit(const FFolderTreeItem& FolderItem) const override
 		{
 			bTransientFilterResult = PassesFilter(FolderItem.Path);
@@ -106,7 +97,6 @@ namespace SceneOutliner
 
 	DECLARE_DELEGATE_RetVal_OneParam( bool, FActorFilterPredicate, const AActor* );
 	DECLARE_DELEGATE_RetVal_OneParam( bool, FWorldFilterPredicate, const UWorld* );
-	DECLARE_DELEGATE_RetVal_OneParam( bool, FLevelBlueprintFilterPredicate, const FLevelBlueprintHandle& );
 	DECLARE_DELEGATE_RetVal_OneParam( bool, FFolderFilterPredicate, FName );
 
 	/** Predicate based filter for the outliner */
@@ -118,8 +108,6 @@ namespace SceneOutliner
 		mutable FWorldFilterPredicate	WorldPred;
 		/** Predicate used to filter Folders */
 		mutable FFolderFilterPredicate	FolderPred;
-		/** Predicate used to filter level blueprints */
-		mutable FLevelBlueprintFilterPredicate LevelBlueprintPred;
 
 		FOutlinerPredicateFilter(FActorFilterPredicate InActorPred, EDefaultFilterBehaviour InDefaultBehaviour, EFailedFilterState InFailedFilterState = EFailedFilterState::NonInteractive)
 			: FOutlinerFilter(InDefaultBehaviour, InFailedFilterState)
@@ -136,11 +124,6 @@ namespace SceneOutliner
 			, FolderPred(InFolderPred)
 		{}
 
-		FOutlinerPredicateFilter(FLevelBlueprintFilterPredicate InLevelBpPred, EDefaultFilterBehaviour InDefaultBehaviour, EFailedFilterState InFailedFilterState = EFailedFilterState::NonInteractive)
-			: FOutlinerFilter(InDefaultBehaviour, InFailedFilterState)
-			, LevelBlueprintPred(InLevelBpPred)
-		{}
-
 		virtual bool PassesFilter(const AActor* Actor) const override
 		{
 			return ActorPred.IsBound() ? ActorPred.Execute(Actor) : DefaultBehaviour == EDefaultFilterBehaviour::Pass;
@@ -151,11 +134,6 @@ namespace SceneOutliner
 			return WorldPred.IsBound() ? WorldPred.Execute(World) : DefaultBehaviour == EDefaultFilterBehaviour::Pass;
 		}
 
-		virtual bool PassesFilter(const FLevelBlueprintHandle& Handle) const override
-		{
-			return LevelBlueprintPred.IsBound() ? LevelBlueprintPred.Execute(Handle) : DefaultBehaviour == EDefaultFilterBehaviour::Pass;
-		}
-
 		virtual bool PassesFilter(FName Folder) const override
 		{
 			return FolderPred.IsBound() ? FolderPred.Execute(Folder) : DefaultBehaviour == EDefaultFilterBehaviour::Pass;
diff --git a/Engine/Source/Editor/SceneOutliner/Public/SceneOutlinerFwd.h b/Engine/Source/Editor/SceneOutliner/Public/SceneOutlinerFwd.h
--- a/Engine/Source/Editor/SceneOutliner/Public/SceneOutlinerFwd.h
+++ b/Engine/Source/Editor/SceneOutliner/Public/SceneOutlinerFwd.h
@@ -15,7 +15,6 @@ namespace SceneOutliner
 	struct ITreeItem;
 	struct FActorTreeItem;
 	struct FWorldTreeItem;
-	struct FLevelBlueprintTreeItem;
 	struct FFolderTreeItem;
 
 	struct ITreeItemVisitor;
diff --git a/Engine/Source/Editor/SceneOutliner/Public/SceneOutlinerVisitorTypes.h b/Engine/Source/Editor/SceneOutliner/Public/SceneOutlinerVisitorTypes.h
--- a/Engine/Source/Editor/SceneOutliner/Public/SceneOutlinerVisitorTypes.h
+++ b/Engine/Source/Editor/SceneOutliner/Public/SceneOutlinerVisitorTypes.h
@@ -11,7 +11,6 @@ namespace SceneOutliner
 	{
 		virtual void Visit(const FActorTreeItem& Actor) const {}
 		virtual void Visit(const FWorldTreeItem& World) const {}
-		virtual void Visit(const FLevelBlueprintTreeItem& LevelBlueprint) const {}
 		virtual void Visit(const FFolderTreeItem& Folder) const {}
 	};
 
@@ -20,7 +19,6 @@ namespace SceneOutliner
 	{
 		virtual void Visit(FActorTreeItem& Actor) const {}
 		virtual void Visit(FWorldTreeItem& World) const {}
-		virtual void Visit(FLevelBlueprintTreeItem& LevelBlueprint) const {}
 		virtual void Visit(FFolderTreeItem& Folder) const {}
 	};
 
@@ -29,23 +27,19 @@ namespace SceneOutliner
 	{
 		typedef TFunctionRef<void(const FActorTreeItem&)> 				FActorFunction;
 		typedef TFunctionRef<void(const FWorldTreeItem&)> 				FWorldFunction;
-		typedef TFunctionRef<void(const FLevelBlueprintTreeItem&)> 		FLevelBlueprintFunction;
 		typedef TFunctionRef<void(const FFolderTreeItem&)> 				FFolderFunction;
 
 		FFunctionalVisitor& Actor(const FActorFunction& InFunction) 					{ ActorFunction = InFunction; return *this; }
 		FFunctionalVisitor& World(const FWorldFunction& InFunction) 					{ WorldFunction = InFunction; return *this; }
-		FFunctionalVisitor& LevelBlueprint(const FLevelBlueprintFunction& InFunction) 	{ LevelBlueprintFunction = InFunction; return *this; }
 		FFunctionalVisitor& Folder(const FFolderFunction& InFunction) 					{ FolderFunction = InFunction; return *this; }
 
 	private:
 		TOptional<FActorFunction> ActorFunction;
 		TOptional<FWorldFunction> WorldFunction;
-		TOptional<FLevelBlueprintFunction> LevelBlueprintFunction;
 		TOptional<FFolderFunction> FolderFunction;
 
 		virtual void Visit(const FActorTreeItem& Item) const override 					{ if (ActorFunction) { ActorFunction.GetValue()(Item); } }
 		virtual void Visit(const FWorldTreeItem& Item) const override 					{ if (WorldFunction) { WorldFunction.GetValue()(Item); } }
-		virtual void Visit(const FLevelBlueprintTreeItem& Item) const override 			{ if (LevelBlueprintFunction) { LevelBlueprintFunction.GetValue()(Item); } }
 		virtual void Visit(const FFolderTreeItem& Item) const override 					{ if (FolderFunction) { FolderFunction.GetValue()(Item); } }
 	};
 
@@ -58,7 +52,6 @@ namespace SceneOutliner
 		/** Override to extract the data from specific tree item types */
 		virtual TDataType Get(const FActorTreeItem& ActorItem) const 					{ return TDataType(); }
 		virtual TDataType Get(const FWorldTreeItem& WorldItem) const 					{ return TDataType(); }
-		virtual TDataType Get(const FLevelBlueprintTreeItem& Item) const 				{ return TDataType(); }
 		virtual TDataType Get(const FFolderTreeItem& FolderItem) const 					{ return TDataType(); }
 
 		/** Return the result returned from Get() */
@@ -67,7 +60,6 @@ namespace SceneOutliner
 	private:
 		virtual void Visit(const FActorTreeItem& ActorItem) const override 				{ Data = Get(ActorItem); }
 		virtual void Visit(const FWorldTreeItem& WorldItem) const override 				{ Data = Get(WorldItem); }
-		virtual void Visit(const FLevelBlueprintTreeItem& Item) const override 			{ Data = Get(Item); }
 		virtual void Visit(const FFolderTreeItem& FolderItem) const override 			{ Data = Get(FolderItem); }
 	};
 
@@ -78,13 +70,11 @@ namespace SceneOutliner
 
 		virtual TSharedRef<SWidget> GenerateWidget(FActorTreeItem& Item) const 			{ return SNullWidget::NullWidget; }
 		virtual TSharedRef<SWidget> GenerateWidget(FWorldTreeItem& Item) const 			{ return SNullWidget::NullWidget; }
-		virtual TSharedRef<SWidget> GenerateWidget(FLevelBlueprintTreeItem& Item) const { return SNullWidget::NullWidget; }
 		virtual TSharedRef<SWidget> GenerateWidget(FFolderTreeItem& Item) const 		{ return SNullWidget::NullWidget; }
 
 	private:
 		virtual void Visit(FActorTreeItem& Item) const override 						{ Widget = GenerateWidget(Item); }
 		virtual void Visit(FWorldTreeItem& Item) const override 						{ Widget = GenerateWidget(Item); }
-		virtual void Visit(FLevelBlueprintTreeItem& Item) const override 				{ Widget = GenerateWidget(Item); }
 		virtual void Visit(FFolderTreeItem& Item) const override 						{ Widget = GenerateWidget(Item); }
 	};
 	
