Fixed "select immediate children" in the world outliner selecting recursively

This addresses UE-7385.

diff --git a/Engine/Source/Editor/SceneOutliner/Private/SSceneOutliner.cpp b/Engine/Source/Editor/SceneOutliner/Private/SSceneOutliner.cpp
--- a/Engine/Source/Editor/SceneOutliner/Private/SSceneOutliner.cpp
+++ b/Engine/Source/Editor/SceneOutliner/Private/SSceneOutliner.cpp
@@ -1274,12 +1274,15 @@ namespace SceneOutliner
 				GEditor->SelectActor(Actor, true, /*bNotify=*/false);
 			}
 		}
+	};
 
+	struct FSelectActorsRecursive : FSelectActors
+	{
 		virtual void Visit(const FFolderTreeItem& FolderItem) const override
 		{
 			for (auto& Child : FolderItem.GetChildren())
 			{
-				Child.Pin()->Visit(FSelectActors());
+				Child.Pin()->Visit(FSelectActorsRecursive());
 			}
 		}
 	};
@@ -1319,7 +1322,7 @@ namespace SceneOutliner
 
 			for (const auto& Folder : SelectedFolders)
 			{
-				Folder->Visit(FSelectActors());
+				Folder->Visit(FSelectActorsRecursive());
 			}
 
 			GEditor->GetSelectedActors()->EndBatchSelectOperation();
