Fixed clang compile error

diff --git a/Engine/Source/Editor/SceneOutliner/Private/SSceneOutliner.cpp b/Engine/Source/Editor/SceneOutliner/Private/SSceneOutliner.cpp
--- a/Engine/Source/Editor/SceneOutliner/Private/SSceneOutliner.cpp
+++ b/Engine/Source/Editor/SceneOutliner/Private/SSceneOutliner.cpp
@@ -1278,6 +1278,8 @@ namespace SceneOutliner
 
 	struct FSelectActorsRecursive : FSelectActors
 	{
+		using FSelectActors::Visit;
+
 		virtual void Visit(const FFolderTreeItem& FolderItem) const override
 		{
 			for (auto& Child : FolderItem.GetChildren())
