Deleting folders in the World Outliner now promotes all children to the level above instead of directly to the root.

diff --git a/Engine/Source/Editor/SceneOutliner/Private/FolderTreeItem.cpp b/Engine/Source/Editor/SceneOutliner/Private/FolderTreeItem.cpp
--- a/Engine/Source/Editor/SceneOutliner/Private/FolderTreeItem.cpp
+++ b/Engine/Source/Editor/SceneOutliner/Private/FolderTreeItem.cpp
@@ -213,27 +213,29 @@ void FFolderTreeItem::Delete()
 
 	struct FResetActorFolders : IMutableTreeItemVisitor
 	{
+		explicit FResetActorFolders(FName InParentPath) : ParentPath(InParentPath) {}
+
 		virtual void Visit(FActorTreeItem& ActorItem) const override
 		{
 			if (AActor* Actor = ActorItem.Actor.Get())
 			{
-				Actor->SetFolderPath(FName());
+				Actor->SetFolderPath(ParentPath);
 			}	
 		}
 		virtual void Visit(FFolderTreeItem& FolderItem) const override
 		{
-			FResetActorFolders ResetFolders;
-			for (auto& Child : FolderItem.GetChildren())
-			{
-				Child.Pin()->Visit(ResetFolders);
-			}
-			FolderItem.Delete();
+			UWorld* World = FolderItem.SharedData->RepresentingWorld;
+			check(World != nullptr);
+
+			MoveFolderTo(FolderItem.Path, ParentPath, *World);
 		}
+
+		FName ParentPath;
 	};
 
 	const FScopedTransaction Transaction( LOCTEXT("DeleteFolder", "Delete Folder") );
 
-	FResetActorFolders ResetFolders;
+	FResetActorFolders ResetFolders(GetParentPath(Path));
 	for (auto& Child : GetChildren())
 	{
 		Child.Pin()->Visit(ResetFolders);
