Don't rename selection set in to the PIE world as that has a number of knock on problems
Selection set will not be added to the transaction buffer if it contains any PIE objects
Transaction will not actually be created when clicking on a PIE object in the scene outliner

diff --git a/Engine/Source/Editor/SceneOutliner/Private/SSceneOutliner.cpp b/Engine/Source/Editor/SceneOutliner/Private/SSceneOutliner.cpp
--- a/Engine/Source/Editor/SceneOutliner/Private/SSceneOutliner.cpp
+++ b/Engine/Source/Editor/SceneOutliner/Private/SSceneOutliner.cpp
@@ -1580,8 +1580,13 @@ namespace SceneOutliner
 				auto SelectedActors = TSet<AActor*>(Selection.GetActorPtrs());
 
 				bool bChanged = false;
+				bool bAnyInPIE = false;
 				for (auto* Actor : SelectedActors)
 				{
+					if (!bAnyInPIE && Actor && (Actor->GetOutermost()->PackageFlags & PKG_PlayInEditor) != 0)
+					{
+						bAnyInPIE = true;
+					}
 					if (!GEditor->GetSelectedActors()->IsSelected(Actor))
 					{
 						bChanged = true;
@@ -1592,6 +1597,10 @@ namespace SceneOutliner
 				for (FSelectionIterator SelectionIt( *GEditor->GetSelectedActors() ); SelectionIt && !bChanged; ++SelectionIt)
 				{
 					AActor* Actor = CastChecked< AActor >( *SelectionIt );
+					if (!bAnyInPIE && (Actor->GetOutermost()->PackageFlags & PKG_PlayInEditor) != 0)
+					{
+						bAnyInPIE = true;
+					}
 					if (!SelectedActors.Contains(Actor))
 					{
 						// Actor has been deselected
@@ -1616,7 +1625,7 @@ namespace SceneOutliner
 				// If there's a discrepancy, update the selected actors to reflect this list.
 				if ( bChanged )
 				{
-					const FScopedTransaction Transaction( NSLOCTEXT("UnrealEd", "ClickingOnActors", "Clicking on Actors") );
+					const FScopedTransaction Transaction( NSLOCTEXT("UnrealEd", "ClickingOnActors", "Clicking on Actors"), !bAnyInPIE );
 					GEditor->GetSelectedActors()->Modify();
 
 					// Clear the selection.
diff --git a/Engine/Source/Editor/UnrealEd/Private/PlayLevel.cpp b/Engine/Source/Editor/UnrealEd/Private/PlayLevel.cpp
--- a/Engine/Source/Editor/UnrealEd/Private/PlayLevel.cpp
+++ b/Engine/Source/Editor/UnrealEd/Private/PlayLevel.cpp
@@ -92,10 +92,6 @@ void UEditorEngine::EndPlayMap()
 	// Enable screensavers when ending PIE.
 	EnableScreenSaver( true );
 
-	// Move SelectedActors and SelectedComponents object back to the transient package.
-	GetSelectedActors()->Rename(nullptr, GetTransientPackage(), REN_DoNotDirty | REN_DontCreateRedirectors | REN_NonTransactional);
-	GetSelectedComponents()->Rename(nullptr,GetTransientPackage(), REN_DoNotDirty | REN_DontCreateRedirectors | REN_NonTransactional);
-
 	// Make a list of all the actors that should be selected
 	TArray<UObject *> SelectedActors;
 	if ( ActorsThatWereSelected.Num() > 0 )
@@ -2743,11 +2739,6 @@ UGameInstance* UEditorEngine::CreatePIEGameInstance(int32 PIEInstance, bool bInS
 		}
 	}
 
-	// Move SelectedActors global object to the PIE package for the duration of the PIE session.
-	// This will stop any transactions on it from being saved during PIE.
-	GetSelectedActors()->Rename(nullptr, GWorld->GetOutermost(), REN_DoNotDirty | REN_DontCreateRedirectors | REN_NonTransactional);
-	GetSelectedComponents()->Rename(nullptr, GWorld->GetOutermost(), REN_DoNotDirty | REN_DontCreateRedirectors | REN_NonTransactional);
-
 	// For play in editor, this is the viewport widget where the game is being displayed
 	TSharedPtr<SViewport> PieViewportWidget;
 
diff --git a/Engine/Source/Runtime/Engine/Classes/Engine/Selection.h b/Engine/Source/Runtime/Engine/Classes/Engine/Selection.h
--- a/Engine/Source/Runtime/Engine/Classes/Engine/Selection.h
+++ b/Engine/Source/Runtime/Engine/Classes/Engine/Selection.h
@@ -292,6 +292,7 @@ public:
 
 	// Begin UObject Interface
 	virtual void Serialize(FArchive& Ar) override;
+	virtual bool Modify( bool bAlwaysMarkDirty=true) override;
 	// End UObject Interface
 
 
diff --git a/Engine/Source/Runtime/Engine/Private/Selection.cpp b/Engine/Source/Runtime/Engine/Private/Selection.cpp
--- a/Engine/Source/Runtime/Engine/Private/Selection.cpp
+++ b/Engine/Source/Runtime/Engine/Private/Selection.cpp
@@ -163,3 +163,18 @@ void USelection::Serialize(FArchive& Ar)
 	Super::Serialize( Ar );
 	Ar << SelectedObjects;
 }
+
+bool USelection::Modify(bool bAlwaysMarkDirty/* =true */)
+{
+	// If the selection currently contains any PIE objects we should not be including it in the transaction buffer
+	for (auto ObjectPtr : SelectedObjects)
+	{
+		UObject* Object = ObjectPtr.Get();
+		if (Object && (Object->GetOutermost()->PackageFlags & (PKG_PlayInEditor | PKG_ContainsScript | PKG_CompiledIn)) != 0)
+		{
+			return false;
+		}
+	}
+
+	return Super::Modify(bAlwaysMarkDirty);
+}
\ No newline at end of file
