Non-unity fix

diff --git a/Engine/Source/Editor/SceneOutliner/Private/FolderTreeItem.cpp b/Engine/Source/Editor/SceneOutliner/Private/FolderTreeItem.cpp
--- a/Engine/Source/Editor/SceneOutliner/Private/FolderTreeItem.cpp
+++ b/Engine/Source/Editor/SceneOutliner/Private/FolderTreeItem.cpp
@@ -3,6 +3,7 @@
 #include "SceneOutlinerPrivatePCH.h"
 #include "SSceneOutliner.h"
 
+#include "ActorEditorUtils.h"
 #include "SceneOutlinerFilters.h"
 
 #include "EditorActorFolders.h"
