SceneOutliner optims
- Disabled intermediate sorting on Undo/Redo operations
- Optimized the natural string cmp used by the SceneOutliner sort
#rnx

diff --git a/Engine/Source/Editor/SceneOutliner/Private/SSceneOutliner.cpp b/Engine/Source/Editor/SceneOutliner/Private/SSceneOutliner.cpp
--- a/Engine/Source/Editor/SceneOutliner/Private/SSceneOutliner.cpp
+++ b/Engine/Source/Editor/SceneOutliner/Private/SSceneOutliner.cpp
@@ -2967,6 +2967,7 @@ namespace SceneOutliner
 		// Refresh our tree in case any changes have been made to the scene that might effect our actor list
 		if( !bIsReentrant )
 		{
+            bDisableIntermediateSorting = true;
 			FullRefresh();
 		}
 	}
diff --git a/Engine/Source/Editor/SceneOutliner/Public/SortHelper.h b/Engine/Source/Editor/SceneOutliner/Public/SortHelper.h
--- a/Engine/Source/Editor/SceneOutliner/Public/SortHelper.h
+++ b/Engine/Source/Editor/SceneOutliner/Public/SortHelper.h
@@ -193,79 +193,97 @@ namespace SceneOutliner
 			return String == Other.String;
 		}
 
-		/** Compare the 2 specified strings using the specified default character comparison function */
-		typedef int32 (*StrncmpMethod)(const TCHAR*, const TCHAR*, SIZE_T Count);
-		static int32 CompareNumeric(const FString& A, const FString& B, StrncmpMethod Strncmp = FCString::Strnicmp)
+	private:
+		static bool BothAscii(TCHAR C1, TCHAR C2)
 		{
-			const TCHAR* CharA = *A, *CharB = *B;
-			const TCHAR* EndA = CharA + A.Len(), *EndB = CharB + B.Len();
+			return (((uint32)C1 | (uint32)C2) & 0xffffff80) == 0;
+		}
+
+		static bool BothNumbers(TCHAR C1, TCHAR C2)
+		{
+			return C1 >= '0' && C1 <= '9' && C2 >= '0' && C2 <= '9';
+		}
 
-			while(CharA < EndA && CharB < EndB)
+		static int32 CompareNumeric(const FString& A, const FString& B)
+		{
+			const TCHAR* String1 = *A;
+			const TCHAR* String2 = *B;
+
+			while (true)
 			{
+				TCHAR C1 = *String1;
+				TCHAR C2 = *String2;
+
 				// Ignore underscores
-				if (*CharA == '_')
+				if (TChar<TCHAR>::IsUnderscore(C1))
 				{
-					++CharA;
+					String1++;
 					continue;
 				}
-				if (*CharB == '_')
+
+				// Ignore underscores
+				if (TChar<TCHAR>::IsUnderscore(C2))
 				{
-					++CharB;
+					String2++;
 					continue;
 				}
-				
-				TCHAR *IntAEnd = nullptr;
-				const uint64 IntA = FCString::Strtoui64(CharA, &IntAEnd, 10);
-
-				TCHAR *IntBEnd = nullptr;
-				const uint64 IntB = FCString::Strtoui64(CharB, &IntBEnd, 10);
 
-				const bool AIsNum = (IntAEnd != CharA), BIsNum = (IntBEnd != CharB);
-				if (AIsNum != BIsNum)
-				{
-					// At the current CharA/CharB position, either (but not both) of the strings is an integer
-					// Numbers are considered less than characters
-					return AIsNum ? -1 : 1;
-				}
-				else if (AIsNum)
+				// Sort numerically when numbers are found 
+				if (BothNumbers(C1, C2))
 				{
-					CharA = IntAEnd;
-					CharB = IntBEnd;
+					int64 IntA = 0;
+					do
+					{
+						IntA *= 10;
+						IntA += TChar<TCHAR>::ConvertCharDigitToInt(C1);
+
+						String1++;
+						C1 = *String1;
+					} while (C1 >= '0' && C1 <= '9');
+
+					int64 IntB = 0;
+					do
+					{
+						IntB *= 10;
+						IntB += TChar<TCHAR>::ConvertCharDigitToInt(C2);
+
+						String2++;
+						C2 = *String2;
+					} while (C2 >= '0' && C2 <= '9');
 
 					if (IntA != IntB)
 					{
 						return IntA < IntB ? -1 : 1;
 					}
+
+					continue;
 				}
-				else if (int32 Cmp = (*Strncmp)(CharA, CharB, 1))
-				{
-					return Cmp;
-				}
-				else
+				else if (C1 == C2)
 				{
-					++CharA;
-					++CharB;
+					// Reached the end of the string
+					if (C1 == 0)
+					{
+						// Strings compared equal, return shortest first
+						return A.Len() == B.Len() ? 0 : A.Len() < B.Len() ? -1 : 1;
+					}
 				}
-			}
-
-			if (CharA == EndA && CharB == EndB)
-			{
-				// Strings compared equal, return shortest first
-				if (const int32 LengthDifference = A.Len() - B.Len())
+				else if (BothAscii(C1, C2))
 				{
-					return LengthDifference / FMath::Abs(LengthDifference);
+					if (int32 Diff = TChar<TCHAR>::ToUnsigned(TChar<TCHAR>::ToLower(C1)) - TChar<TCHAR>::ToUnsigned(TChar<TCHAR>::ToLower(C2)))
+					{
+						return Diff;
+					}
 				}
 				else
 				{
-					return 0;
+					return TChar<TCHAR>::ToUnsigned(C1) - TChar<TCHAR>::ToUnsigned(C2);
 				}
-			}
 
-			// Strings are different comparative lengths, return shortest
-			return CharA == EndA ? -1 : 1;
+				String1++;
+				String2++;
+			}
 		}
 
-	private:
 		FNumericStringWrapper(const FNumericStringWrapper&);
 		FNumericStringWrapper& operator=(const FNumericStringWrapper&);
 	};
