Improve multi-selection performance
   -120 seconds to select 6000 actors to instantaneous to select 60 000.

From dev-enterprise[REDACTED_HANDLE]

diff --git a/Engine/Source/Editor/SceneOutliner/Private/SSceneOutliner.cpp b/Engine/Source/Editor/SceneOutliner/Private/SSceneOutliner.cpp
--- a/Engine/Source/Editor/SceneOutliner/Private/SSceneOutliner.cpp
+++ b/Engine/Source/Editor/SceneOutliner/Private/SSceneOutliner.cpp
@@ -3299,46 +3299,44 @@ namespace SceneOutliner
 
 		// Deselect actors in the tree that are no longer selected in the world
 		FItemSelection Selection(*OutlinerTreeView);
-		for (FActorTreeItem* ActorItem : Selection.Actors)
+		if (Selection.Actors.Num())
 		{
-			if(!ActorItem->Actor.IsValid() || !ActorItem->Actor.Get()->IsSelected())
+			TArray<FTreeItemPtr> ActorItems;
+			for (FActorTreeItem* ActorItem : Selection.Actors)
 			{
-				OutlinerTreeView->SetItemSelection(ActorItem->AsShared(), false);
-			}
-		}
-
-
-		// See if the tree view selector is pointing at a selected item
-		bool bSelectorInSelectionSet = false;
-		for (FSelectionIterator SelectionIt(*SelectedActors); SelectionIt; ++SelectionIt)
-		{
-			AActor* Actor = CastChecked< AActor >(*SelectionIt);
-			if (FTreeItemPtr* ActorItem = TreeItemMap.Find(Actor))
-			{
-				if (OutlinerTreeView->Private_HasSelectorFocus(*ActorItem))
+				if(!ActorItem->Actor.IsValid() || !ActorItem->Actor.Get()->IsSelected())
 				{
-					bSelectorInSelectionSet = true;
-					break;
+					ActorItems.Add(ActorItem->AsShared());
 				}
 			}
-		}
 
-		// If NOT bSelectorInSelectionSet then we want to just move the selector to the first selected item.
-		ESelectInfo::Type SelectInfo = bSelectorInSelectionSet ? ESelectInfo::Direct : ESelectInfo::OnMouseClick;
+			OutlinerTreeView->SetItemSelection(ActorItems, false);
+		}
+		
 		// Show actor selection but only if sub objects are not selected
 		if (Selection.Components.Num() == 0 && Selection.SubComponents.Num() == 0)
 		{
-			// Ensure that all selected actors in the world are selected in the tree
+			// See if the tree view selector is pointing at a selected item
+			bool bSelectorInSelectionSet = false;
+
+			TArray<FTreeItemPtr> ActorItems;
 			for (FSelectionIterator SelectionIt(*SelectedActors); SelectionIt; ++SelectionIt)
 			{
 				AActor* Actor = CastChecked< AActor >(*SelectionIt);
 				if (FTreeItemPtr* ActorItem = TreeItemMap.Find(Actor))
 				{
-					OutlinerTreeView->SetItemSelection(*ActorItem, true, SelectInfo);				
-					SelectInfo = ESelectInfo::Direct;
-				}
+					if (!bSelectorInSelectionSet && OutlinerTreeView->Private_HasSelectorFocus(*ActorItem))
+					{
+						bSelectorInSelectionSet = true;
+					}
 
+					ActorItems.Add(*ActorItem);
+				}
 			}
+
+			// If NOT bSelectorInSelectionSet then we want to just move the selector to the first selected item.
+			ESelectInfo::Type SelectInfo = bSelectorInSelectionSet ? ESelectInfo::Direct : ESelectInfo::OnMouseClick;
+			OutlinerTreeView->SetItemSelection(ActorItems, true, SelectInfo);
 		}
 
 		// Broadcast selection changed delegate
diff --git a/Engine/Source/Editor/UnrealEd/Private/EditorSelectUtils.cpp b/Engine/Source/Editor/UnrealEd/Private/EditorSelectUtils.cpp
--- a/Engine/Source/Editor/UnrealEd/Private/EditorSelectUtils.cpp
+++ b/Engine/Source/Editor/UnrealEd/Private/EditorSelectUtils.cpp
@@ -592,7 +592,7 @@ void UUnrealEdEngine::SelectActor(AActor* Actor, bool bInSelected, bool bNotify,
 
 		// Don't do any work if the actor's selection state is already the selected state.
 		const bool bActorSelected = Actor->IsSelected();
-		if ( (bActorSelected && !bInSelected) || (!bActorSelected && bInSelected) )
+		if ( bActorSelected != bInSelected )
 		{
 			if(bInSelected)
 			{
diff --git a/Engine/Source/Runtime/Slate/Public/Widgets/Views/SListView.h b/Engine/Source/Runtime/Slate/Public/Widgets/Views/SListView.h
--- a/Engine/Source/Runtime/Slate/Public/Widgets/Views/SListView.h
+++ b/Engine/Source/Runtime/Slate/Public/Widgets/Views/SListView.h
@@ -1244,6 +1244,30 @@ public:
 		Private_SignalSelectionChanged(SelectInfo);
 	}
 
+	/**
+	 * Set the selection state of multiple items.
+	 *
+	 * @param InItems     The Items whose selection state to modify
+	 * @param bSelected   true to select the items; false to unselect
+	 * @param SelectInfo  Provides context on how the selection changed
+	 */
+	void SetItemSelection( const TArray<ItemType>& InItems, bool bSelected, ESelectInfo::Type SelectInfo = ESelectInfo::Direct )
+	{
+		if ( InItems.Num() == 0 || SelectionMode.Get() == ESelectionMode::None )
+		{
+			return;
+		}
+
+		for (const ItemType & Item : InItems)
+		{
+			Private_SetItemSelection(Item, bSelected, SelectInfo != ESelectInfo::Direct);
+
+			// Any item after the first one selected will be direct
+			SelectInfo = ESelectInfo::Direct;
+		}
+		Private_SignalSelectionChanged(SelectInfo);
+	}
+
 	/**
 	 * Empty the selection set.
 	 */
