Remove ActorComponents.Enable console command. Feature is enabled by default, and registering command causes warnings when opening WorldOutliner more than once.

diff --git a/Engine/Source/Editor/SceneOutliner/Private/SSceneOutliner.cpp b/Engine/Source/Editor/SceneOutliner/Private/SSceneOutliner.cpp
--- a/Engine/Source/Editor/SceneOutliner/Private/SSceneOutliner.cpp
+++ b/Engine/Source/Editor/SceneOutliner/Private/SSceneOutliner.cpp
@@ -258,8 +258,6 @@ namespace SceneOutliner
 		bFullRefresh = true;
 		bNeedsRefresh = true;
 		bIsReentrant = false;
-		bActorComponentsEnabled = true;// || InInitOptions.Mode = ESceneOutlinerMode::ComponentPicker;
-
 		bSortDirty = true;
 		bActorSelectionDirty = SharedData->Mode == ESceneOutlinerMode::ActorBrowsing;
 		FilteredActorCount = 0;
@@ -764,21 +762,18 @@ namespace SceneOutliner
 					EUserInterfaceActionType::ToggleButton
 				);
 
-				if (bActorComponentsEnabled)
-				{
-					MenuBuilder.AddMenuEntry(
-						LOCTEXT("ToggleShowActorComponents", "Show Actor Components"),
-						LOCTEXT("ToggleShowActorComponentsToolTip", "When enabled, shows components beloging to actors."),
-						FSlateIcon(),
-						FUIAction(
-							FExecuteAction::CreateSP(this, &SSceneOutliner::ToggleShowActorComponents),
-							FCanExecuteAction(),
-							FIsActionChecked::CreateSP(this, &SSceneOutliner::IsShowingActorComponents)
-						),
-						NAME_None,
-						EUserInterfaceActionType::ToggleButton
-					);
-				}
+				MenuBuilder.AddMenuEntry(
+					LOCTEXT("ToggleShowActorComponents", "Show Actor Components"),
+					LOCTEXT("ToggleShowActorComponentsToolTip", "When enabled, shows components beloging to actors."),
+					FSlateIcon(),
+					FUIAction(
+						FExecuteAction::CreateSP(this, &SSceneOutliner::ToggleShowActorComponents),
+						FCanExecuteAction(),
+						FIsActionChecked::CreateSP(this, &SSceneOutliner::IsShowingActorComponents)
+					),
+					NAME_None,
+					EUserInterfaceActionType::ToggleButton
+				);
 
 				// Add additional filters
 				FSceneOutlinerModule& SceneOutlinerModule = FModuleManager::LoadModuleChecked< FSceneOutlinerModule >("SceneOutliner");
@@ -944,7 +939,7 @@ namespace SceneOutliner
 
 	bool SSceneOutliner::IsShowingActorComponents()
 	{
-		return SharedData->Mode == ESceneOutlinerMode::ComponentPicker || (bActorComponentsEnabled && GetDefault<USceneOutlinerSettings>()->bShowActorComponents);
+		return SharedData->Mode == ESceneOutlinerMode::ComponentPicker || (GetDefault<USceneOutlinerSettings>()->bShowActorComponents);
 	}
 
 	void SSceneOutliner::ToggleShowActorComponents()
diff --git a/Engine/Source/Editor/SceneOutliner/Public/SSceneOutliner.h b/Engine/Source/Editor/SceneOutliner/Public/SSceneOutliner.h
--- a/Engine/Source/Editor/SceneOutliner/Public/SSceneOutliner.h
+++ b/Engine/Source/Editor/SceneOutliner/Public/SSceneOutliner.h
@@ -84,8 +84,7 @@ namespace SceneOutliner
 		void Construct( const FArguments& InArgs, const FInitializationOptions& InitOptions );
 
 		/** Default constructor - initializes data that is shared between all tree items */
-			SSceneOutliner() : SharedData(MakeShareable(new FSharedOutlinerData)),
-				ActorComponentsEnable(TEXT("ActorComponents.Enable"), TEXT("Toggles actor components visible in Scene outliner"), FConsoleCommandDelegate::CreateRaw(this, &SSceneOutliner::ActorComponentsModeToggle)) {}
+			SSceneOutliner() : SharedData(MakeShareable(new FSharedOutlinerData)) {}
 
 		/** SSceneOutliner destructor */
 		~SSceneOutliner();
@@ -630,12 +629,6 @@ namespace SceneOutliner
 		/** Reentrancy guard */
 		bool bIsReentrant;
 
-		/** Whether Actor Components are enabled in the scene outliner */
-		bool bActorComponentsEnabled;
-
-		/** Console commands for enabling/disabling actor components in the scene outliner while it is still in development */
-		FAutoConsoleCommand ActorComponentsEnable;
-
 		/* Widget containing the filtering text box */
 		TSharedPtr< SSearchBox > FilterTextBoxWidget;
 
@@ -653,8 +646,6 @@ namespace SceneOutliner
 
 		TMap<FName, const FSlateBrush*> CachedIcons;
 
-		void ActorComponentsModeToggle() { bActorComponentsEnabled = !bActorComponentsEnabled; FullRefresh(); }
-
 	private:
 		/** Functions relating to sorting */
 
