Fix the UI of the scene outliner leaving some (deleted actor)

diff --git a/Engine/Source/Editor/SceneOutliner/Private/SSceneOutliner.cpp b/Engine/Source/Editor/SceneOutliner/Private/SSceneOutliner.cpp
--- a/Engine/Source/Editor/SceneOutliner/Private/SSceneOutliner.cpp
+++ b/Engine/Source/Editor/SceneOutliner/Private/SSceneOutliner.cpp
@@ -1355,8 +1355,11 @@ namespace SceneOutliner
 
 		PendingTreeItemMap.Remove(ItemID);
 
-		// If a tree item already exists that represents the same data, bail
-		if (TreeItemMap.Find(ItemID))
+		FValidateItemBeforeAddingToTree ValidateItemVisitor;
+		Item->Visit(ValidateItemVisitor);
+
+		// If a tree item already exists that represents the same data or if the actor is invalid, bail
+		if (TreeItemMap.Find(ItemID)  || !ValidateItemVisitor.Result())
 		{
 			return false;
 		}
diff --git a/Engine/Source/Editor/SceneOutliner/Public/SSceneOutliner.h b/Engine/Source/Editor/SceneOutliner/Public/SSceneOutliner.h
--- a/Engine/Source/Editor/SceneOutliner/Public/SSceneOutliner.h
+++ b/Engine/Source/Editor/SceneOutliner/Public/SSceneOutliner.h
@@ -208,6 +208,16 @@ namespace SceneOutliner
 			}
 		}
 
+		/** Visitor that is used to validate if the item should added to the tree */
+		struct FValidateItemBeforeAddingToTree : TTreeItemGetter<bool>
+		{
+			/** Override to extract the data from specific tree item types */
+			virtual bool Get(const FActorTreeItem& ActorItem) const { return ActorItem.Actor.IsValid(); }
+			virtual bool Get(const FWorldTreeItem& WorldItem) const { return true; }
+			virtual bool Get(const FFolderTreeItem& FolderItem) const { return true; }
+			virtual bool Get(const FComponentTreeItem& ComponentFunction) const { return true; }
+			virtual bool Get(const FSubComponentTreeItem& CustomFunction) const { return true; }
+		};
 
 		/** Visitor that is used to set up type-specific data after tree items are added to the tree */
 		struct FOnItemAddedToTree : IMutableTreeItemVisitor
