Scene outline syncronization now makes sure the treeview selector is within the selection set.

diff --git a/Engine/Source/Editor/SceneOutliner/Private/SSceneOutliner.cpp b/Engine/Source/Editor/SceneOutliner/Private/SSceneOutliner.cpp
--- a/Engine/Source/Editor/SceneOutliner/Private/SSceneOutliner.cpp
+++ b/Engine/Source/Editor/SceneOutliner/Private/SSceneOutliner.cpp
@@ -2515,13 +2515,32 @@ namespace SceneOutliner
 			}
 		}
 
+		// See if the tree view selector is pointing at a selected item
+		bool bSelectorInSelectionSet = false;
+		for (FSelectionIterator SelectionIt(*SelectedActors); SelectionIt; ++SelectionIt)
+		{
+			AActor* Actor = CastChecked< AActor >(*SelectionIt);
+			if (FTreeItemPtr* ActorItem = TreeItemMap.Find(Actor))
+			{
+				if (OutlinerTreeView->Private_HasSelectorFocus(*ActorItem))
+				{
+					bSelectorInSelectionSet = true;
+					break;
+				}
+			}
+		}
+
+		// If NOT bSelectorInSelectionSet then we want to just move the selector to the first selected item.
+		ESelectInfo::Type SelectInfo = bSelectorInSelectionSet ? ESelectInfo::Direct : ESelectInfo::OnMouseClick;
+
 		// Ensure that all selected actors in the world are selected in the tree
 		for (FSelectionIterator SelectionIt( *SelectedActors ); SelectionIt; ++SelectionIt)
 		{
 			AActor* Actor = CastChecked< AActor >(*SelectionIt);
 			if (FTreeItemPtr* ActorItem = TreeItemMap.Find(Actor))
 			{
-				OutlinerTreeView->SetItemSelection(*ActorItem, true);
+				OutlinerTreeView->SetItemSelection(*ActorItem, true, SelectInfo);
+				SelectInfo = ESelectInfo::Direct;
 			}
 		}
 
