World Outliner: added Edit menu/hotkey functionality for folders (Cut/Copy/Paste/Duplicate/Delete/Rename), added Duplicate Folder Hierarchy menu option, fixed Select All Descendants menu option, fixed folder delete bug, fixed Actor Edit not handling folders.

diff --git a/Engine/Source/Editor/LevelEditor/Private/LevelEditorActions.cpp b/Engine/Source/Editor/LevelEditor/Private/LevelEditorActions.cpp
--- a/Engine/Source/Editor/LevelEditor/Private/LevelEditorActions.cpp
+++ b/Engine/Source/Editor/LevelEditor/Private/LevelEditorActions.cpp
@@ -89,6 +89,7 @@
 #include "EditorLevelUtils.h"
 #include "ActorGroupingUtils.h"
 #include "LevelUtils.h"
+#include "ISceneOutliner.h"
 
 DEFINE_LOG_CATEGORY_STATIC(LevelEditorActions, Log, All);
 
@@ -1426,6 +1427,19 @@ bool FLevelEditorActionCallbacks::Duplicate_CanExecute()
 		}
 	}
 
+	if (!bCanCopy)
+	{
+		TWeakPtr<class SLevelEditor> LevelEditor = FModuleManager::GetModuleChecked<FLevelEditorModule>(TEXT("LevelEditor")).GetLevelEditorInstance();
+		if (LevelEditor.IsValid())
+		{
+			TSharedPtr<class ISceneOutliner> SceneOutlinerPtr = LevelEditor.Pin()->GetSceneOutliner();
+			if (SceneOutlinerPtr.IsValid())
+			{
+				bCanCopy = SceneOutlinerPtr->Copy_CanExecute();
+			}
+		}
+	}
+
 	return bCanCopy;
 }
 
@@ -1466,6 +1480,19 @@ bool FLevelEditorActionCallbacks::Delete_CanExecute()
 		}
 	}
 
+	if (!bCanDelete)
+	{
+		TWeakPtr<class SLevelEditor> LevelEditor = FModuleManager::GetModuleChecked<FLevelEditorModule>(TEXT("LevelEditor")).GetLevelEditorInstance();
+		if (LevelEditor.IsValid())
+		{
+			TSharedPtr<class ISceneOutliner> SceneOutlinerPtr = LevelEditor.Pin()->GetSceneOutliner();
+			if (SceneOutlinerPtr.IsValid())
+			{
+				bCanDelete = SceneOutlinerPtr->Delete_CanExecute();
+			}
+		}
+	}
+
 	return bCanDelete;
 }
 
@@ -1476,12 +1503,20 @@ void FLevelEditorActionCallbacks::Rename_Execute()
 	{
 		GEditor->BroadcastLevelComponentRequestRename(Component);
 	}
+	else if (AActor* Actor = Cast<AActor>(*GEditor->GetSelectedActorIterator()))
+	{
+		GEditor->BroadcastLevelActorRequestRename(Actor);
+	}
 	else
 	{
-		AActor* Actor = Cast<AActor>(*GEditor->GetSelectedActorIterator());
-		if (Actor)
+		TWeakPtr<class SLevelEditor> LevelEditor = FModuleManager::GetModuleChecked<FLevelEditorModule>(TEXT("LevelEditor")).GetLevelEditorInstance();
+		if (LevelEditor.IsValid())
 		{
-			GEditor->BroadcastLevelActorRequestRename(Actor);
+			TSharedPtr<class ISceneOutliner> SceneOutlinerPtr = LevelEditor.Pin()->GetSceneOutliner();
+			if (SceneOutlinerPtr.IsValid())
+			{
+				SceneOutlinerPtr->Rename_Execute();
+			}
 		}
 	}
 }
@@ -1502,6 +1537,19 @@ bool FLevelEditorActionCallbacks::Rename_CanExecute()
 		bCanRename = GEditor->GetSelectedActorCount() == 1;
 	}
 
+	if (!bCanRename)
+	{
+		TWeakPtr<class SLevelEditor> LevelEditor = FModuleManager::GetModuleChecked<FLevelEditorModule>(TEXT("LevelEditor")).GetLevelEditorInstance();
+		if (LevelEditor.IsValid())
+		{
+			TSharedPtr<class ISceneOutliner> SceneOutlinerPtr = LevelEditor.Pin()->GetSceneOutliner();
+			if (SceneOutlinerPtr.IsValid())
+			{
+				bCanRename = SceneOutlinerPtr->Rename_CanExecute();
+			}
+		}
+	}
+
 	return bCanRename;
 }
 
@@ -1544,6 +1592,19 @@ bool FLevelEditorActionCallbacks::Cut_CanExecute()
 		}
 	}
 
+	if (!bCanCut)
+	{
+		TWeakPtr<class SLevelEditor> LevelEditor = FModuleManager::GetModuleChecked<FLevelEditorModule>(TEXT("LevelEditor")).GetLevelEditorInstance();
+		if (LevelEditor.IsValid())
+		{
+			TSharedPtr<class ISceneOutliner> SceneOutlinerPtr = LevelEditor.Pin()->GetSceneOutliner();
+			if (SceneOutlinerPtr.IsValid())
+			{
+				bCanCut = SceneOutlinerPtr->Cut_CanExecute();
+			}
+		}
+	}
+
 	return bCanCut;
 }
 
@@ -1584,6 +1645,19 @@ bool FLevelEditorActionCallbacks::Copy_CanExecute()
 		}
 	}
 
+	if (!bCanCopy)
+	{
+		TWeakPtr<class SLevelEditor> LevelEditor = FModuleManager::GetModuleChecked<FLevelEditorModule>(TEXT("LevelEditor")).GetLevelEditorInstance();
+		if (LevelEditor.IsValid())
+		{
+			TSharedPtr<class ISceneOutliner> SceneOutlinerPtr = LevelEditor.Pin()->GetSceneOutliner();
+			if (SceneOutlinerPtr.IsValid())
+			{
+				bCanCopy = SceneOutlinerPtr->Copy_CanExecute();
+			}
+		}
+	}
+
 	return bCanCopy;
 }
 
@@ -1622,6 +1696,19 @@ bool FLevelEditorActionCallbacks::Paste_CanExecute()
 		}
 	}
 
+	if (!bCanPaste)
+	{
+		TWeakPtr<class SLevelEditor> LevelEditor = FModuleManager::GetModuleChecked<FLevelEditorModule>(TEXT("LevelEditor")).GetLevelEditorInstance();
+		if (LevelEditor.IsValid())
+		{
+			TSharedPtr<class ISceneOutliner> SceneOutlinerPtr = LevelEditor.Pin()->GetSceneOutliner();
+			if (SceneOutlinerPtr.IsValid())
+			{
+				bCanPaste = SceneOutlinerPtr->Paste_CanExecute();
+			}
+		}
+	}
+
 	return bCanPaste;
 }
 
diff --git a/Engine/Source/Editor/LevelEditor/Private/LevelEditorContextMenu.cpp b/Engine/Source/Editor/LevelEditor/Private/LevelEditorContextMenu.cpp
--- a/Engine/Source/Editor/LevelEditor/Private/LevelEditorContextMenu.cpp
+++ b/Engine/Source/Editor/LevelEditor/Private/LevelEditorContextMenu.cpp
@@ -198,7 +198,7 @@ void FLevelEditorContextMenu::FillMenu( FMenuBuilder& MenuBuilder, TWeakPtr<SLev
 
 		FComponentEditorUtils::FillComponentContextMenuOptions(MenuBuilder, SelectedComponents);
 	}
-	else
+	else if (GEditor->GetSelectedActorCount() > 0)
 	{
 		// Generate information about our selection
 		TArray<AActor*> SelectedActors;
@@ -440,6 +440,20 @@ void FLevelEditorContextMenu::FillMenu( FMenuBuilder& MenuBuilder, TWeakPtr<SLev
 
 		MenuBuilder.PopExtender();
 	}	
+	else if (ContextType == LevelEditorMenuContext::SceneOutliner)
+	{
+		TWeakPtr<ISceneOutliner> SceneOutlinerPtr = LevelEditor.Pin()->GetSceneOutliner();
+		if (SceneOutlinerPtr.IsValid())
+		{
+			MenuBuilder.BeginSection("SelectVisibilityLevels");
+			{
+				MenuBuilder.AddSubMenu(
+					LOCTEXT("EditSubMenu", "Edit"),
+					FText::GetEmpty(),
+					FNewMenuDelegate::CreateStatic(&FLevelEditorContextMenuImpl::FillEditMenu, ContextType));
+			}
+		}
+	}
 
 	MenuBuilder.PopCommandList();
 }
diff --git a/Engine/Source/Editor/LevelEditor/Private/LevelEditorContextMenu.h b/Engine/Source/Editor/LevelEditor/Private/LevelEditorContextMenu.h
--- a/Engine/Source/Editor/LevelEditor/Private/LevelEditorContextMenu.h
+++ b/Engine/Source/Editor/LevelEditor/Private/LevelEditorContextMenu.h
@@ -15,8 +15,8 @@ enum class LevelEditorMenuContext
 {
 	/** This context menu is applicable to a viewport */
 	Viewport,
-	/** This context menu is applicable to an external UI or dialog (disables click-position-based menu items) */
-	NonViewport,
+	/** This context menu is applicable to the Scene Outliner (disables click-position-based menu items) */
+	SceneOutliner,
 };
 
 /**
diff --git a/Engine/Source/Editor/LevelEditor/Private/SLevelEditor.cpp b/Engine/Source/Editor/LevelEditor/Private/SLevelEditor.cpp
--- a/Engine/Source/Editor/LevelEditor/Private/SLevelEditor.cpp
+++ b/Engine/Source/Editor/LevelEditor/Private/SLevelEditor.cpp
@@ -663,11 +663,8 @@ TSharedRef<SDockTab> SLevelEditor::SpawnLevelEditorTab( const FSpawnTabArgs& Arg
 			InitOptions.DefaultMenuExtender->AddMenuExtension(
 				"MainSection", EExtensionHook::Before, GetLevelEditorActions(),
 				FMenuExtensionDelegate::CreateStatic([](FMenuBuilder& MenuBuilder, TWeakPtr<SLevelEditor> InWeakLevelEditor){
-					// Only extend the menu if we have actors selected
-					if (GEditor->GetSelectedActorCount() > 0)
-					{
-						FLevelEditorContextMenu::FillMenu(MenuBuilder, InWeakLevelEditor, LevelEditorMenuContext::NonViewport, TSharedPtr<FExtender>());
-					}
+					// Extend the menu even if no actors selected, as Edit menu should always exist for scene outliner
+					FLevelEditorContextMenu::FillMenu(MenuBuilder, InWeakLevelEditor, LevelEditorMenuContext::SceneOutliner, TSharedPtr<FExtender>());
 				}, WeakLevelEditor)
 			);
 		}
@@ -675,7 +672,12 @@ TSharedRef<SDockTab> SLevelEditor::SpawnLevelEditorTab( const FSpawnTabArgs& Arg
 
 		FText Label = NSLOCTEXT( "LevelEditor", "SceneOutlinerTabTitle", "World Outliner" );
 
-		FSceneOutlinerModule& SceneOutlinerModule = FModuleManager::Get().LoadModuleChecked<FSceneOutlinerModule>( "SceneOutliner" );
+		FSceneOutlinerModule& SceneOutlinerModule = FModuleManager::Get().LoadModuleChecked<FSceneOutlinerModule>("SceneOutliner");
+		TSharedRef<ISceneOutliner> SceneOutlinerRef = SceneOutlinerModule.CreateSceneOutliner(
+			InitOptions,
+			FOnActorPicked() /* Not used for outliner when in browsing mode */);
+		SceneOutlinerPtr = SceneOutlinerRef;
+
 		return SNew( SDockTab )
 			.Icon( FEditorStyle::GetBrush( "LevelEditor.Tabs.Outliner" ) )
 			.Label( Label )
@@ -686,9 +688,7 @@ TSharedRef<SDockTab> SLevelEditor::SpawnLevelEditorTab( const FSpawnTabArgs& Arg
 				.BorderImage( FEditorStyle::GetBrush("ToolPanel.GroupBorder") )
 				.AddMetaData<FTutorialMetaData>(FTutorialMetaData(TEXT("SceneOutliner"), TEXT("LevelEditorSceneOutliner")))
 				[
-					SceneOutlinerModule.CreateSceneOutliner(
-						InitOptions,
-						FOnActorPicked() /* Not used for outliner when in browsing mode */ )
+					SceneOutlinerRef
 				]
 			];
 	}
diff --git a/Engine/Source/Editor/LevelEditor/Private/SLevelEditor.h b/Engine/Source/Editor/LevelEditor/Private/SLevelEditor.h
--- a/Engine/Source/Editor/LevelEditor/Private/SLevelEditor.h
+++ b/Engine/Source/Editor/LevelEditor/Private/SLevelEditor.h
@@ -128,6 +128,10 @@ public:
 	
 	/** Attaches a sequencer asset editor used to animate objects in the level to this level editor */
 	void AttachSequencer( TSharedPtr<SWidget> SequencerWidget, TSharedPtr<IAssetEditorInstance> NewSequencerAssetEditor );
+
+	/** Returns current scene outliner associated with level editor's scene outliner tab, if it exists */
+	TSharedPtr<ISceneOutliner> GetSceneOutliner() { return SceneOutlinerPtr.Pin();  }
+
 private:
 	
 	TSharedRef<SDockTab> SpawnLevelEditorTab(const FSpawnTabArgs& Args, FName TabIdentifier, FString InitializationPayload);
@@ -242,4 +246,7 @@ private:
 
 	/** Weak pointer to the level editor's Sequencer widget */
 	TWeakPtr<SWidget> SequencerWidgetPtr;
+
+	/** Weak pointer to the level editor's scene outliner */
+	TWeakPtr<ISceneOutliner> SceneOutlinerPtr;
 };
diff --git a/Engine/Source/Editor/SceneOutliner/Private/FolderTreeItem.cpp b/Engine/Source/Editor/SceneOutliner/Private/FolderTreeItem.cpp
--- a/Engine/Source/Editor/SceneOutliner/Private/FolderTreeItem.cpp
+++ b/Engine/Source/Editor/SceneOutliner/Private/FolderTreeItem.cpp
@@ -239,7 +239,7 @@ bool FFolderTreeItem::CanInteract() const
 	return Flags.bInteractive;
 }
 
-void FFolderTreeItem::Delete()
+void FFolderTreeItem::Delete(FName InNewParentPath)
 {
 	if (!SharedData->RepresentingWorld)
 	{
@@ -270,7 +270,7 @@ void FFolderTreeItem::Delete()
 
 	const FScopedTransaction Transaction( LOCTEXT("DeleteFolderTransaction", "Delete Folder") );
 
-	FResetActorFolders ResetFolders(GetParentPath(Path));
+	FResetActorFolders ResetFolders(InNewParentPath);
 	for (auto& Child : GetChildren())
 	{
 		Child.Pin()->Visit(ResetFolders);
@@ -295,6 +295,17 @@ void FFolderTreeItem::CreateSubFolder(TWeakPtr<SSceneOutliner> WeakOutliner)
 	}
 }
 
+void FFolderTreeItem::DuplicateHierarchy(TWeakPtr<SSceneOutliner> WeakOutliner)
+{
+	TSharedPtr<SSceneOutliner> Outliner = WeakOutliner.Pin();
+
+	if (Outliner.IsValid() && SharedData->RepresentingWorld)
+	{
+		const FScopedTransaction Transaction(LOCTEXT("UndoAction_DuplicateHierarchy", "Duplicate Folder Hierarchy"));
+		Outliner->DuplicateFoldersHierarchy();
+	}
+}
+
 void FFolderTreeItem::OnExpansionChanged()
 {
 	if (!SharedData->RepresentingWorld)
@@ -316,8 +327,7 @@ void FFolderTreeItem::GenerateContextMenu(FMenuBuilder& MenuBuilder, SSceneOutli
 	const FSlateIcon NewFolderIcon(FEditorStyle::GetStyleSetName(), "SceneOutliner.NewFolderIcon");
 	
 	MenuBuilder.AddMenuEntry(LOCTEXT("CreateSubFolder", "Create Sub Folder"), FText(), NewFolderIcon, FUIAction(FExecuteAction::CreateSP(this, &FFolderTreeItem::CreateSubFolder, TWeakPtr<SSceneOutliner>(SharedOutliner))));
-	MenuBuilder.AddMenuEntry(LOCTEXT("RenameFolder", "Rename"), FText(), FSlateIcon(), FUIAction(FExecuteAction::CreateSP(&Outliner, &SSceneOutliner::InitiateRename, AsShared())));
-	MenuBuilder.AddMenuEntry(LOCTEXT("DeleteFolder", "Delete"), FText(), FSlateIcon(), FUIAction(FExecuteAction::CreateSP(this, &FFolderTreeItem::Delete)));
+	MenuBuilder.AddMenuEntry(LOCTEXT("DuplicateFolderHierarchy", "Duplicate Hierarchy"), FText(), FSlateIcon(), FUIAction(FExecuteAction::CreateSP(&Outliner, &SSceneOutliner::DuplicateFoldersHierarchy)));
 }
 
 void FFolderTreeItem::PopulateDragDropPayload(FDragDropPayload& Payload) const
diff --git a/Engine/Source/Editor/SceneOutliner/Private/SSceneOutliner.cpp b/Engine/Source/Editor/SceneOutliner/Private/SSceneOutliner.cpp
--- a/Engine/Source/Editor/SceneOutliner/Private/SSceneOutliner.cpp
+++ b/Engine/Source/Editor/SceneOutliner/Private/SSceneOutliner.cpp
@@ -7,6 +7,7 @@
 #include "Modules/ModuleManager.h"
 #include "Layout/WidgetPath.h"
 #include "Framework/Application/SlateApplication.h"
+//#include "Framework/Commands/GenericCommands.h"
 #include "Textures/SlateIcon.h"
 #include "Framework/Commands/UIAction.h"
 #include "Widgets/Layout/SSeparator.h"
@@ -28,9 +29,7 @@
 #include "SceneOutlinerSettings.h"
 #include "ISceneOutlinerColumn.h"
 #include "SceneOutlinerModule.h"
-
-
-
+#include "HAL/PlatformApplicationMisc.h"
 #include "Widgets/Input/SSearchBox.h"
 
 #include "SceneOutlinerDragDrop.h"
@@ -507,6 +506,17 @@ namespace SceneOutliner
 		Folders.OnFolderCreate.AddSP(this, &SSceneOutliner::OnBroadcastFolderCreate);
 		Folders.OnFolderMove.AddSP(this, &SSceneOutliner::OnBroadcastFolderMove);
 		Folders.OnFolderDelete.AddSP(this, &SSceneOutliner::OnBroadcastFolderDelete);
+
+		FEditorDelegates::OnEditCutActorsBegin.AddSP(this, &SSceneOutliner::OnEditCutActorsBegin);
+		FEditorDelegates::OnEditCutActorsEnd.AddSP(this, &SSceneOutliner::OnEditCutActorsEnd);
+		FEditorDelegates::OnEditCopyActorsBegin.AddSP(this, &SSceneOutliner::OnEditCopyActorsBegin);
+		FEditorDelegates::OnEditCopyActorsEnd.AddSP(this, &SSceneOutliner::OnEditCopyActorsEnd);
+		FEditorDelegates::OnEditPasteActorsBegin.AddSP(this, &SSceneOutliner::OnEditPasteActorsBegin);
+		FEditorDelegates::OnEditPasteActorsEnd.AddSP(this, &SSceneOutliner::OnEditPasteActorsEnd);
+		FEditorDelegates::OnDuplicateActorsBegin.AddSP(this, &SSceneOutliner::OnDuplicateActorsBegin);
+		FEditorDelegates::OnDuplicateActorsEnd.AddSP(this, &SSceneOutliner::OnDuplicateActorsEnd);
+		FEditorDelegates::OnDeleteActorsBegin.AddSP(this, &SSceneOutliner::OnDeleteActorsBegin);
+		FEditorDelegates::OnDeleteActorsEnd.AddSP(this, &SSceneOutliner::OnDeleteActorsEnd);
 	}
 
 	void SSceneOutliner::SetupColumns(SHeaderRow& HeaderRow)
@@ -601,6 +611,17 @@ namespace SceneOutliner
 			Folders.OnFolderMove.RemoveAll(this);
 			Folders.OnFolderDelete.RemoveAll(this);
 		}
+
+		FEditorDelegates::OnEditCutActorsBegin.RemoveAll(this);
+		FEditorDelegates::OnEditCutActorsEnd.RemoveAll(this);
+		FEditorDelegates::OnEditCopyActorsBegin.RemoveAll(this);
+		FEditorDelegates::OnEditCopyActorsEnd.RemoveAll(this);
+		FEditorDelegates::OnEditPasteActorsBegin.RemoveAll(this);
+		FEditorDelegates::OnEditPasteActorsEnd.RemoveAll(this);
+		FEditorDelegates::OnDuplicateActorsBegin.RemoveAll(this);
+		FEditorDelegates::OnDuplicateActorsEnd.RemoveAll(this);
+		FEditorDelegates::OnDeleteActorsBegin.RemoveAll(this);
+		FEditorDelegates::OnDeleteActorsEnd.RemoveAll(this);
 	}
 
 	void SSceneOutliner::OnItemAdded(const FTreeItemID& ItemID, uint8 ActionMask)
@@ -874,7 +895,7 @@ namespace SceneOutliner
 	{
 		// Block events while we clear out the list.  We don't want actors in the level to become deselected
 		// while we doing this
-		TGuardValue<bool> ReentrantGuard(bIsReentrant, true);
+ 		TGuardValue<bool> ReentrantGuard(bIsReentrant, true);
 
 		SharedData->RepresentingWorld = nullptr;
 
@@ -985,6 +1006,15 @@ namespace SceneOutliner
 		PendingOperations.RemoveAt(0, End);
 		SetParentsExpansionState(ExpansionStateInfo);
 
+		for (FName Folder : PendingFoldersSelect)
+		{
+			if (FTreeItemPtr* Item = TreeItemMap.Find(Folder))
+			{
+				OutlinerTreeView->SetItemSelection(*Item, true);
+			}
+		}
+		PendingFoldersSelect.Empty();
+
 		if (bMadeAnySignificantChanges && !SharedData->bRepresentingPlayWorld)
 		{
 			RequestSort();
@@ -1270,6 +1300,20 @@ namespace SceneOutliner
 		return FItemSelection(*OutlinerTreeView).Folders;
 	}
 
+	TArray<FName> SSceneOutliner::GetSelectedFolderNames() const
+	{
+		TArray<FFolderTreeItem*> SelectedFolders = GetSelectedFolders();
+		TArray<FName> SelectedFolderNames;
+		for (FFolderTreeItem* Folder : SelectedFolders)
+		{
+			if (Folder)
+			{
+				SelectedFolderNames.Add(Folder->Path);
+			}
+		}
+		return SelectedFolderNames;
+	}
+
 	TSharedPtr<SWidget> SSceneOutliner::OnOpenContextMenu()
 	{
 		TArray<AActor*> SelectedActors;
@@ -1301,6 +1345,80 @@ namespace SceneOutliner
 
 		return TSharedPtr<SWidget>();
 	}
+	
+	bool SSceneOutliner::Delete_CanExecute()
+	{
+		FItemSelection ItemSelection(*OutlinerTreeView);
+		if (ItemSelection.Folders.Num() > 0 && ItemSelection.Folders.Num() == OutlinerTreeView->GetNumItemsSelected())
+		{
+			return true;
+		}
+		return false;
+	}
+
+	bool SSceneOutliner::Rename_CanExecute()
+	{
+		FItemSelection ItemSelection(*OutlinerTreeView);
+		if (ItemSelection.Folders.Num() == 1 && ItemSelection.Folders.Num() == OutlinerTreeView->GetNumItemsSelected())
+		{
+			return true;
+		}
+		return false;
+	}
+
+	void SSceneOutliner::Rename_Execute()
+	{
+		// handle folders only here, actors and components are handled in LevelEditorActions::Rename_Execute
+		FItemSelection ItemSelection(*OutlinerTreeView);
+		if (ItemSelection.Folders.Num() == 1 && ItemSelection.Folders.Num() == OutlinerTreeView->GetNumItemsSelected())
+		{
+			FTreeItemPtr ItemToRename = OutlinerTreeView->GetSelectedItems()[0];
+
+			if (ItemToRename->CanInteract())
+			{
+				PendingRenameItem = ItemToRename->AsShared();
+				ScrollItemIntoView(ItemToRename);
+			}
+		}
+	}
+
+	bool SSceneOutliner::Cut_CanExecute()
+	{
+		FItemSelection ItemSelection(*OutlinerTreeView);
+		if (ItemSelection.Folders.Num() > 0 && ItemSelection.Folders.Num() == OutlinerTreeView->GetNumItemsSelected())
+		{
+			return true;
+		}
+		return false;
+	}
+
+	bool SSceneOutliner::Copy_CanExecute()
+	{
+		FItemSelection ItemSelection(*OutlinerTreeView);
+		if (ItemSelection.Folders.Num() > 0 && ItemSelection.Folders.Num() == OutlinerTreeView->GetNumItemsSelected())
+		{
+			return true;
+		}
+		return false;
+	}
+
+	bool SSceneOutliner::Paste_CanExecute()
+	{
+		if (CanPasteFoldersOnlyFromClipboard())
+		{
+			return true;
+		}
+		return false;
+	}
+
+	bool SSceneOutliner::CanPasteFoldersOnlyFromClipboard()
+	{
+		// Intentionally not checking if the level is locked/hidden here, as it's better feedback for the user if they attempt to paste
+		// and get the message explaining why it's failed, than just not having the option available to them.
+		FString PasteString;
+		FPlatformApplicationMisc::ClipboardPaste(PasteString);
+		return PasteString.StartsWith("BEGIN FOLDERLIST");
+	}
 
 	TSharedPtr<SWidget> SSceneOutliner::BuildDefaultContextMenu()
 	{
@@ -1309,57 +1427,68 @@ namespace SceneOutliner
 			return nullptr;
 		}
 
+		FItemSelection ItemSelection(*OutlinerTreeView);
+
+		const auto NumSelectedItems = OutlinerTreeView->GetNumItemsSelected();
+		bool bPasteCommandOnly = (NumSelectedItems == 0);
+
 		// Build up the menu for a selection
 		const bool bCloseAfterSelection = true;
 		FMenuBuilder MenuBuilder(bCloseAfterSelection, TSharedPtr<FUICommandList>(), SharedData->DefaultMenuExtender);
 
-		const auto NumSelectedItems = OutlinerTreeView->GetNumItemsSelected();
-		if (NumSelectedItems == 1)
-		{
-			OutlinerTreeView->GetSelectedItems()[0]->GenerateContextMenu(MenuBuilder, *this);
-		}
-
 		bool MenuBuilderHasContent = false;
 
-		// We always create a section here, even if there is no parent so that clients can still extend the menu
-		MenuBuilder.BeginSection("MainSection");
+		if (SharedData->bShowParentTree)
 		{
-			FItemSelection ItemSelection(*OutlinerTreeView);
-
-			// Don't add any of these menu items if we're not showing the parent tree
-			if (SharedData->bShowParentTree)
+			if (NumSelectedItems == 0)
 			{
+				const FSlateIcon NewFolderIcon(FEditorStyle::GetStyleSetName(), "SceneOutliner.NewFolderIcon");
+				MenuBuilder.AddMenuEntry(LOCTEXT("CreateFolder", "Create Folder"), FText(), NewFolderIcon, FUIAction(FExecuteAction::CreateSP(this, &SSceneOutliner::CreateFolder)));
 				MenuBuilderHasContent = true;
+			}
+			else
+			{
+				if (NumSelectedItems == 1)
+				{
+					OutlinerTreeView->GetSelectedItems()[0]->GenerateContextMenu(MenuBuilder, *this);
+					MenuBuilderHasContent = true;
+				}
 
-				if (NumSelectedItems == 0)
+				// If we've only got folders selected, show the selection and edit sub menus
+				if (NumSelectedItems > 0 && ItemSelection.Folders.Num() == NumSelectedItems)
 				{
-					const FSlateIcon NewFolderIcon(FEditorStyle::GetStyleSetName(), "SceneOutliner.NewFolderIcon");
-					MenuBuilder.AddMenuEntry(LOCTEXT("CreateFolder", "Create Folder"), FText(), NewFolderIcon, FUIAction(FExecuteAction::CreateSP(this, &SSceneOutliner::CreateFolder)));
+					MenuBuilder.AddSubMenu(
+						LOCTEXT("SelectSubmenu", "Select"),
+						LOCTEXT("SelectSubmenu_Tooltip", "Select the contents of the current selection"),
+						FNewMenuDelegate::CreateSP(this, &SSceneOutliner::FillSelectionSubMenu));
+					MenuBuilderHasContent = true;
 				}
-				else
+			}
+		}
+
+			// We always create a section here, even if there is no parent so that clients can still extend the menu
+			MenuBuilder.BeginSection("MainSection");
+			{
+				// Don't add any of these menu items if we're not showing the parent tree
+				if (SharedData->bShowParentTree)
 				{
 					// Can't move worlds or level blueprints
-					const bool bCanMoveSelection = ItemSelection.Worlds.Num() == 0;
-					if (bCanMoveSelection)
+					if (NumSelectedItems > 0)
 					{
-						MenuBuilder.AddSubMenu(
-							LOCTEXT("MoveActorsTo", "Move To"),
-							LOCTEXT("MoveActorsTo_Tooltip", "Move selection to another folder"),
-							FNewMenuDelegate::CreateSP(this,  &SSceneOutliner::FillFoldersSubMenu));
-					}
-
-					// If we've only got folders selected, show the selection sub menu
-					if (ItemSelection.Folders.Num() == NumSelectedItems)
-					{
-						MenuBuilder.AddSubMenu(
-							LOCTEXT("SelectSubmenu", "Select"),
-							LOCTEXT("SelectSubmenu_Tooltip", "Select the contents of the current selection"),
-							FNewMenuDelegate::CreateSP(this, &SSceneOutliner::FillSelectionSubMenu));
+						const bool bCanMoveSelection = ItemSelection.Worlds.Num() == 0;
+						if (bCanMoveSelection)
+						{
+							MenuBuilder.AddSubMenu(
+								LOCTEXT("MoveActorsTo", "Move To"),
+								LOCTEXT("MoveActorsTo_Tooltip", "Move selection to another folder"),
+								FNewMenuDelegate::CreateSP(this, &SSceneOutliner::FillFoldersSubMenu));
+							MenuBuilderHasContent = true;
+						}
 					}
 				}
 			}
-		}
-		MenuBuilder.EndSection();
+			MenuBuilder.EndSection();
+		
 
 		if (MenuBuilderHasContent)
 		{
@@ -1510,66 +1639,114 @@ namespace SceneOutliner
 	{
 		MenuBuilder.AddMenuEntry(
 			LOCTEXT( "AddChildrenToSelection", "Immediate Children" ),
-			LOCTEXT( "AddChildrenToSelection_ToolTip", "Select all immediate children of the selected folders" ),
+			LOCTEXT( "AddChildrenToSelection_ToolTip", "Select all immediate actor children of the selected folders" ),
 			FSlateIcon(),
 			FExecuteAction::CreateSP(this, &SSceneOutliner::SelectFoldersImmediateChildren));
 		MenuBuilder.AddMenuEntry(
 			LOCTEXT( "AddDescendantsToSelection", "All Descendants" ),
-			LOCTEXT( "AddDescendantsToSelection_ToolTip", "Select all descendants of the selected folders" ),
+			LOCTEXT( "AddDescendantsToSelection_ToolTip", "Select all actor descendants of the selected folders" ),
 			FSlateIcon(),
 			FExecuteAction::CreateSP(this, &SSceneOutliner::SelectFoldersDescendants));
 	}
 
-	struct FSelectActors : ITreeItemVisitor
+	void SSceneOutliner::SelectFoldersImmediateChildren()
 	{
-		virtual void Visit(const FActorTreeItem& ActorItem) const override
+		struct FSelectActors : ITreeItemVisitor
 		{
-			if (AActor* Actor = ActorItem.Actor.Get())
+			virtual void Visit(const FActorTreeItem& ActorItem) const override
 			{
-				GEditor->SelectActor(Actor, true, /*bNotify=*/false);
+				if (AActor* Actor = ActorItem.Actor.Get())
+				{
+					GEditor->SelectActor(Actor, true, /*bNotify=*/false);
+				}
 			}
-		}
-	};
-
-	struct FSelectActorsRecursive : FSelectActors
-	{
-		using FSelectActors::Visit;
+		};
 
-		virtual void Visit(const FFolderTreeItem& FolderItem) const override
+		struct FSelectFolders : IMutableTreeItemVisitor
 		{
-			for (auto& Child : FolderItem.GetChildren())
+			SSceneOutliner& Outliner;
+			FSelectFolders(SSceneOutliner& InOutliner) : Outliner(InOutliner) {}
+
+			virtual void Visit(FFolderTreeItem& FolderItem) const override
 			{
-				Child.Pin()->Visit(FSelectActorsRecursive());
+				Outliner.OutlinerTreeView->SetItemSelection(FolderItem.AsShared(), true);
 			}
-		}
-	};
+		};
+
+		TArray<FFolderTreeItem*> SelectedFolders = GetSelectedFolders();
+		OutlinerTreeView->ClearSelection();
 
-	void SSceneOutliner::SelectFoldersImmediateChildren() const
-	{
-		auto SelectedFolders = GetSelectedFolders();
 		if (SelectedFolders.Num())
 		{
 			// We'll batch selection changes instead by using BeginBatchSelectOperation()
 			GEditor->GetSelectedActors()->BeginBatchSelectOperation();
 
-			OutlinerTreeView->ClearSelection();
-
-			for (const auto& Folder : SelectedFolders)
+			FSelectActors SelectActorsVisitor;
+			for (const FFolderTreeItem* Folder : SelectedFolders)
 			{
-				for (auto& Child : Folder->GetChildren())
+				for (TWeakPtr<ITreeItem> Child : Folder->GetChildren())
 				{
-					Child.Pin()->Visit(FSelectActors());
+					Child.Pin()->Visit(SelectActorsVisitor);
 				}
 			}
 
 			GEditor->GetSelectedActors()->EndBatchSelectOperation();
 			GEditor->NoteSelectionChange();
 		}
+
+		FSelectFolders SelectFoldersVisitor(*this);
+		for (const FFolderTreeItem* Folder : SelectedFolders)
+		{
+			for (TWeakPtr<ITreeItem> Child : Folder->GetChildren())
+			{
+				Child.Pin()->Visit(SelectFoldersVisitor);
+			}
+		}
 	}
 
-	void SSceneOutliner::SelectFoldersDescendants() const
+	void SSceneOutliner::SelectFoldersDescendants()
 	{
-		auto SelectedFolders = GetSelectedFolders();
+		struct FSelectActorsRecursive : ITreeItemVisitor
+		{
+			virtual void Visit(const FActorTreeItem& ActorItem) const override
+			{
+				if (AActor* Actor = ActorItem.Actor.Get())
+				{
+					GEditor->SelectActor(Actor, true, /*bNotify=*/false);
+				}
+
+				for (TWeakPtr<ITreeItem> Child : ActorItem.GetChildren())
+				{
+					Child.Pin()->Visit(*this);
+				}
+			}
+
+			virtual void Visit(const FFolderTreeItem& FolderItem) const override
+			{
+				for (TWeakPtr<ITreeItem> Child : FolderItem.GetChildren())
+				{
+					Child.Pin()->Visit(*this);
+				}
+			}
+		};
+
+		struct FSelectFoldersRecursive : IMutableTreeItemVisitor
+		{
+			SSceneOutliner& Outliner;
+			FSelectFoldersRecursive(SSceneOutliner& InOutliner) : Outliner(InOutliner) {}
+
+			virtual void Visit(FFolderTreeItem& FolderItem) const override
+			{
+				Outliner.OutlinerTreeView->SetItemSelection(FolderItem.AsShared(), true);
+
+				for (TWeakPtr<ITreeItem> Child : FolderItem.GetChildren())
+				{
+					Child.Pin()->Visit(*this);
+				}
+			}
+		};
+
+		TArray<FFolderTreeItem*> SelectedFolders = GetSelectedFolders();
 		if (SelectedFolders.Num())
 		{
 			// We'll batch selection changes instead by using BeginBatchSelectOperation()
@@ -1577,14 +1754,24 @@ namespace SceneOutliner
 
 			OutlinerTreeView->ClearSelection();
 
-			for (const auto& Folder : SelectedFolders)
+			FSelectActorsRecursive SelectActorsRecursive;
+			for (const FFolderTreeItem* Folder : SelectedFolders)
 			{
-				Folder->Visit(FSelectActorsRecursive());
+				Folder->Visit(SelectActorsRecursive);
 			}
 
 			GEditor->GetSelectedActors()->EndBatchSelectOperation();
 			GEditor->NoteSelectionChange();
 		}
+
+		FSelectFoldersRecursive SelectFoldersRecursiveVisitor(*this);
+		for (const FFolderTreeItem* Folder : SelectedFolders)
+		{
+			for (TWeakPtr<ITreeItem> Child : Folder->GetChildren())
+			{
+				Child.Pin()->Visit(SelectFoldersRecursiveVisitor);
+			}
+		}
 	}
 
 	void SSceneOutliner::MoveSelectionTo(FTreeItemRef NewParent)
@@ -1726,6 +1913,413 @@ namespace SceneOutliner
 		}
 	}
 
+	void SSceneOutliner::OnEditCutActorsBegin()
+	{
+		CopyFoldersBegin();
+		DeleteFoldersBegin();
+	}
+
+	void SSceneOutliner::OnEditCutActorsEnd()
+	{
+		CopyFoldersEnd();
+		DeleteFoldersEnd();
+	}
+
+	void SSceneOutliner::OnEditCopyActorsBegin()
+	{
+		CopyFoldersBegin();
+	}
+
+	void SSceneOutliner::OnEditCopyActorsEnd()
+	{
+		CopyFoldersEnd();
+	}
+
+	void SSceneOutliner::OnEditPasteActorsBegin()
+	{
+		TArray<FName> Folders = GetClipboardPasteFolders();
+		PasteFoldersBegin(Folders);
+	}
+
+	void SSceneOutliner::OnEditPasteActorsEnd()
+	{
+		PasteFoldersEnd();
+	}
+
+	void SSceneOutliner::OnDuplicateActorsBegin()
+	{
+		TArray<FFolderTreeItem*> SelectedFolders = GetSelectedFolders();
+		PasteFoldersBegin(SelectedFolders);
+	}
+
+	void SSceneOutliner::OnDuplicateActorsEnd()
+	{
+		PasteFoldersEnd();
+	}
+
+	void SSceneOutliner::OnDeleteActorsBegin()
+	{
+		DeleteFoldersBegin();
+	}
+
+	void SSceneOutliner::OnDeleteActorsEnd()
+	{
+		DeleteFoldersEnd();
+	}
+
+	void SSceneOutliner::CopyFoldersBegin()
+	{
+		CacheFoldersEdit = GetSelectedFolderNames();
+		FPlatformApplicationMisc::ClipboardPaste(CacheClipboardContents);
+	}
+
+	void SSceneOutliner::CopyFoldersEnd()
+	{
+		if (CacheFoldersEdit.Num() > 0)
+		{
+			CopyFoldersToClipboard(CacheFoldersEdit, CacheClipboardContents);
+			CacheFoldersEdit.Reset();
+		}
+	}
+
+	void SSceneOutliner::CopyFoldersToClipboard(const TArray<FName>& InFolders, const FString& InPrevClipboardContents)
+	{
+		if (InFolders.Num() > 0)
+		{
+			// If clipboard paste has changed since we cached it, actors must have been cut 
+			// so folders need to appended to clipboard contents rather than replacing them.
+			FString CurrClipboardContents;
+			FPlatformApplicationMisc::ClipboardPaste(CurrClipboardContents);
+
+			FString Buffer = ExportFolderList(InFolders);
+
+			FString* SourceData = (CurrClipboardContents != InPrevClipboardContents ? &CurrClipboardContents : nullptr);
+
+			if (SourceData)
+			{
+				SourceData->Append(Buffer);
+			}
+			else
+			{
+				SourceData = &Buffer;
+			}
+
+			// Replace clipboard contents with original plus folders appended
+			FPlatformApplicationMisc::ClipboardCopy(**SourceData);
+		}
+	}
+
+	void SSceneOutliner::PasteFoldersBegin(TArray<FFolderTreeItem*> InFolders)
+	{
+		TArray<FName> FolderNames;
+
+		for (FFolderTreeItem* Folder : InFolders)
+		{
+			if (Folder)
+			{
+				FolderNames.Add(Folder->Path);
+			}
+		}
+
+		PasteFoldersBegin(FolderNames);
+	}
+
+	void SSceneOutliner::PasteFoldersBegin(TArray<FName> InFolders)
+	{
+		struct FCacheExistingChildrenAction : IMutableTreeItemVisitor
+		{
+			SSceneOutliner& Outliner;
+
+			FCacheExistingChildrenAction(SSceneOutliner& InOutliner) : Outliner(InOutliner) {}
+
+			virtual void Visit(FFolderTreeItem& FolderItem) const override
+			{
+				TArray<FTreeItemID> ExistingChildren;
+				for (const TWeakPtr<ITreeItem>& Child : FolderItem.GetChildren())
+				{
+					if (Child.IsValid())
+					{
+						ExistingChildren.Add(Child.Pin()->GetID());
+					}
+				}
+
+				Outliner.CachePasteFolderExistingChildrenMap.Add(FolderItem.Path, ExistingChildren);
+			}
+		};
+
+
+		CacheFoldersEdit.Reset();
+		CachePasteFolderExistingChildrenMap.Reset();
+		PendingFoldersSelect.Reset();
+
+		CacheFoldersEdit = InFolders;
+
+		// Sort folder names so parents appear before children
+		CacheFoldersEdit.Sort();
+
+		// Cache existing children
+		for (FName Folder : CacheFoldersEdit)
+		{
+			if (FTreeItemPtr* TreeItem = TreeItemMap.Find(Folder))
+			{
+				FCacheExistingChildrenAction CacheExistingChildrenAction(*this);
+				(*TreeItem)->Visit(CacheExistingChildrenAction);
+			}
+		}
+	}
+
+	void SSceneOutliner::PasteFoldersEnd()
+	{
+		struct FReparentDuplicatedActorsAction : IMutableTreeItemVisitor
+		{
+			SSceneOutliner& Outliner;
+			mutable bool bVisitedFolder;
+			mutable FName* NewFolderPath;
+			TMap<FName, FName>* FolderMap;
+
+			FReparentDuplicatedActorsAction(SSceneOutliner& InOutliner, TMap<FName, FName>* InFolderMap) : Outliner(InOutliner), bVisitedFolder(false), NewFolderPath(nullptr), FolderMap(InFolderMap) {}
+
+			virtual void Visit(FFolderTreeItem& FolderItem) const override
+			{
+				if (!bVisitedFolder)
+				{
+					bVisitedFolder = true;
+
+					if (!FolderItem.Path.IsNone())
+					{
+						NewFolderPath = FolderMap->Find(FolderItem.Path);
+						if (NewFolderPath && *NewFolderPath != FolderItem.Path)
+						{
+							for (TWeakPtr<ITreeItem> Child : FolderItem.GetChildren())
+							{
+								if (Child.IsValid())
+								{
+									Child.Pin()->Visit(*this);
+								}
+							}
+						}
+					}
+				}
+			}
+
+			virtual void Visit(FActorTreeItem& ActorItem) const override
+			{
+				if (NewFolderPath)
+				{
+					if (AActor* Actor = ActorItem.Actor.Get())
+					{
+						FName ParentPath = Actor->GetFolderPath();
+						TArray<FTreeItemID>* ExistingChildren = Outliner.CachePasteFolderExistingChildrenMap.Find(ParentPath);
+
+						if (ExistingChildren && !ExistingChildren->Contains(ActorItem.GetID()))
+						{
+							Actor->SetFolderPath_Recursively(*NewFolderPath);
+						}
+					}
+				}
+			}
+		};
+
+		if (!CheckWorld())
+		{
+			return;
+		}
+			
+		const FScopedTransaction Transaction(NSLOCTEXT("UnrealEd", "PasteItems", "Paste Items"));
+
+		// Create new folder
+		TMap<FName, FName> FolderMap;
+		for (FName Folder : CacheFoldersEdit)
+		{
+			FName ParentPath = GetParentPath(Folder);
+			FName LeafName = GetFolderLeafName(Folder);
+			if (LeafName != TEXT(""))
+			{
+				if (FName* NewParentPath = FolderMap.Find(ParentPath))
+				{
+					ParentPath = *NewParentPath;
+				}
+
+				FName NewFolderPath = FActorFolders::Get().GetFolderName(*SharedData->RepresentingWorld, ParentPath, LeafName);
+				FActorFolders::Get().CreateFolder(*SharedData->RepresentingWorld, NewFolderPath);
+				FolderMap.Add(Folder, NewFolderPath);
+			}
+		}
+
+		// Populate our data set
+		Populate();
+
+		// Reparent duplicated actors if the folder has been pasted/duplicated
+		for (FName Folder : CacheFoldersEdit)
+		{
+			if (const FName* NewFolder = FolderMap.Find(Folder))
+			{
+				if (FTreeItemPtr* FolderItem = TreeItemMap.Find(Folder))
+				{
+					FReparentDuplicatedActorsAction ReparentDuplicatedActors(*this, &FolderMap);
+					(*FolderItem)->Visit(ReparentDuplicatedActors);
+				}
+
+				Folder = *NewFolder;
+			}
+
+			PendingFoldersSelect.Add(Folder);
+		}
+
+		CacheFoldersEdit.Reset();
+		CachePasteFolderExistingChildrenMap.Reset();
+		FullRefresh();
+	}
+
+	void SSceneOutliner::DuplicateFoldersHierarchy()
+	{
+		struct FSelectFoldersRecursive : IMutableTreeItemVisitor
+		{
+			SSceneOutliner& Outliner;
+			FSelectFoldersRecursive(SSceneOutliner& InOutliner) : Outliner(InOutliner) {}
+
+			virtual void Visit(FFolderTreeItem& FolderItem) const override
+			{
+				// select folders to be duplicated
+				Outliner.OutlinerTreeView->SetItemSelection(FolderItem.AsShared(), true);
+
+				// FItemSelection Selection(*Outliner.OutlinerTreeView);
+
+				for (TWeakPtr<ITreeItem> Child : FolderItem.GetChildren())
+				{
+					Child.Pin()->Visit(*this);
+				}
+			}
+		};
+
+		if (!CheckWorld())
+		{
+			return;
+		}
+
+		const FScopedTransaction Transaction(NSLOCTEXT("UnrealEd", "DuplicateFoldersHierarchy", "Duplicate Folders Hierarchy"));
+
+		TArray<FFolderTreeItem*> SelectedFolders = GetSelectedFolders();
+
+		if (SelectedFolders.Num() > 0)
+		{
+			// Select actor descendants
+			SelectFoldersDescendants();
+
+			// Select all sub-folders
+			for (FFolderTreeItem* Folder : SelectedFolders)
+			{
+				FSelectFoldersRecursive SelectFoldersRecursive(*this);
+				Folder->Visit(SelectFoldersRecursive);
+			}
+
+			// Duplicate selected
+			GUnrealEd->Exec(SharedData->RepresentingWorld, TEXT("DUPLICATE"));
+		}
+	}
+
+	void SSceneOutliner::DeleteFoldersBegin()
+	{
+		CacheFoldersDelete = GetSelectedFolders();
+	}
+
+	void SSceneOutliner::DeleteFoldersEnd()
+	{
+		struct FMatchName
+		{
+			FMatchName(const FName InPathName)
+				: PathName(InPathName) {}
+
+			const FName PathName;
+
+			bool operator()(const FFolderTreeItem *Entry)
+			{
+				return PathName == Entry->Path;
+			}
+		};
+
+		if (CacheFoldersDelete.Num() > 0)
+		{
+			// Sort in descending order so children will be deleted before parents
+			CacheFoldersDelete.Sort([](const FFolderTreeItem& FolderA, const FFolderTreeItem& FolderB)
+			{
+				return (FolderA.Path > FolderB.Path);
+			});
+
+			for (FFolderTreeItem* Folder : CacheFoldersDelete)
+			{
+				if (Folder)
+				{
+					// Find lowest parent not being deleted, for reparenting children of current folder
+					FName NewParentPath = GetParentPath(Folder->Path);
+					while (!NewParentPath.IsNone() && CacheFoldersDelete.FindByPredicate(FMatchName(NewParentPath)))
+					{
+						NewParentPath = GetParentPath(NewParentPath);
+					}
+
+					Folder->Delete(NewParentPath);
+				}
+			}
+
+			CacheFoldersDelete.Reset();
+			FullRefresh();
+		}
+	}
+
+	TArray<FName> SSceneOutliner::GetClipboardPasteFolders() const
+	{
+		FString PasteString;
+		FPlatformApplicationMisc::ClipboardPaste(PasteString);
+		return ImportFolderList(*PasteString);
+	}
+
+	FString SSceneOutliner::ExportFolderList(TArray<FName> InFolders) const
+	{
+		FString Buffer = FString(TEXT("Begin FolderList\n"));
+
+		for (auto& FolderName : InFolders)
+		{
+			Buffer.Append(FString(TEXT("\tFolder=")) + FolderName.ToString() + FString(TEXT("\n")));
+		}
+
+		Buffer += FString(TEXT("End FolderList\n"));
+
+		return Buffer;
+	}
+
+	TArray<FName> SSceneOutliner::ImportFolderList(const FString& InStrBuffer) const
+	{
+		TArray<FName> Folders;
+
+		int32 Index = InStrBuffer.Find(TEXT("Begin FolderList"));
+		if (Index != INDEX_NONE)
+		{
+			FString TmpStr = InStrBuffer.RightChop(Index);
+			const TCHAR* Buffer = *TmpStr;
+
+			FString StrLine;
+			while (FParse::Line(&Buffer, StrLine))
+			{
+				const TCHAR* Str = *StrLine;				
+				FString FolderName;
+
+				if (FParse::Command(&Str, TEXT("Begin")) && FParse::Command(&Str, TEXT("FolderList")))
+				{
+					continue;
+				}
+				else if (FParse::Command(&Str, TEXT("End")) && FParse::Command(&Str, TEXT("FolderList")))
+				{
+					break;
+				}
+				else if (FParse::Value(Str, TEXT("Folder="), FolderName))
+				{
+					Folders.Add(FName(*FolderName));
+				}
+			}
+		}
+		return Folders;
+	}
+
 	void SSceneOutliner::ScrollItemIntoView(FTreeItemPtr Item)
 	{
 		auto Parent = Item->GetParent();
@@ -1738,15 +2332,6 @@ namespace SceneOutliner
 		OutlinerTreeView->RequestScrollIntoView(Item);
 	}
 
-	void SSceneOutliner::InitiateRename(TSharedRef<ITreeItem> Item)
-	{
-		if (Item->CanInteract())
-		{
-			PendingRenameItem = Item;
-			ScrollItemIntoView(Item);
-		}
-	}
-
 	TSharedRef< ITableRow > SSceneOutliner::OnGenerateRowForOutlinerTree( FTreeItemPtr Item, const TSharedRef< STableViewBase >& OwnerTable )
 	{
 		return SNew( SSceneOutlinerTreeRow, OutlinerTreeView.ToSharedRef(), SharedThis(this) ).Item( Item );
@@ -2440,7 +3025,8 @@ namespace SceneOutliner
 			}
 
 			// Delete key: Delete selected actors (not rebindable, because it doesn't make much sense to bind.)
-			else if ( InKeyEvent.GetKey() == EKeys::Platform_Delete )
+			// Use Delete and Backspace instead of Platform_Delete because the LevelEditor default Edit Delete is bound to both 
+			else if ( InKeyEvent.GetKey() == EKeys::Delete || InKeyEvent.GetKey() == EKeys::BackSpace )
 			{
 				const FItemSelection Selection(*OutlinerTreeView);
 
@@ -2448,50 +3034,13 @@ namespace SceneOutliner
 				{
 					SharedData->CustomDelete.Execute( Selection.GetWeakActors() );
 				}
-				else if (CheckWorld())
+				else
 				{
-					const FScopedTransaction Transaction( LOCTEXT("UndoAction_DeleteSelection", "Delete selection") );
-
-					// Delete selected folders too
-					auto SelectedItems = OutlinerTreeView->GetSelectedItems();
-					
-					GEditor->SelectNone(true, true);
-
-					for (auto* Folder : Selection.Folders)
-					{
-						Folder->Delete();
-					}
-
-					for (auto* Actor : Selection.GetActorPtrs())
-					{
-						GEditor->SelectActor(Actor, true, false);
-					}
-
-					// Code from FLevelEditorActionCallbacks::Delete_CanExecute()
-					// Should this be just return FReply::Unhandled()?
-					TArray<FEdMode*> ActiveModes; 
-					GLevelEditorModeTools().GetActiveModes( ActiveModes );
-					for( int32 ModeIndex = 0; ModeIndex < ActiveModes.Num(); ++ModeIndex )
-					{
-						const EEditAction::Type CanProcess = ActiveModes[ModeIndex]->GetActionEditDelete();
-						if (CanProcess == EEditAction::Process)
-						{
-							// We don't consider the return value here, as false is assumed to mean there was an internal error processing delete, not that it should defer to other modes/default behaviour
-							ActiveModes[ModeIndex]->ProcessEditDelete();
-							return FReply::Handled();
-						}
-						else if (CanProcess == EEditAction::Halt)
-						{
-							return FReply::Unhandled();
-						}
-					}
-
-					if (GUnrealEd->CanDeleteSelectedActors( SharedData->RepresentingWorld, true, true ))
+					if (CheckWorld())
 					{
-						GEditor->edactDeleteSelected( SharedData->RepresentingWorld );
+						GUnrealEd->Exec(SharedData->RepresentingWorld, TEXT("DELETE"));
 					}
 				}
-
 				return FReply::Handled();
 			}
 		}
diff --git a/Engine/Source/Editor/SceneOutliner/Private/SSceneOutliner.h b/Engine/Source/Editor/SceneOutliner/Private/SSceneOutliner.h
--- a/Engine/Source/Editor/SceneOutliner/Private/SSceneOutliner.h
+++ b/Engine/Source/Editor/SceneOutliner/Private/SSceneOutliner.h
@@ -17,6 +17,7 @@
 #include "SceneOutlinerPublicTypes.h"
 #include "Editor/SceneOutliner/Private/SOutlinerTreeView.h"
 #include "Editor/SceneOutliner/Private/SceneOutlinerStandaloneTypes.h"
+#include "Framework/Commands/UICommandList.h"
 
 class FMenuBuilder;
 class ISceneOutlinerColumn;
@@ -130,6 +131,27 @@ namespace SceneOutliner
 		/** Sets the cached icon for this class name */
 		virtual void CacheIconForClass(FName InClassName, const FSlateBrush* InSlateBrush) override;
 		
+		/** Returns true if edit delete can be executed */
+		virtual bool Delete_CanExecute();
+
+		/** Returns true if edit rename can be executed */
+		virtual bool Rename_CanExecute();
+
+		/** Executes rename. */
+		virtual void Rename_Execute();
+
+		/** Returns true if edit cut can be executed */
+		virtual bool Cut_CanExecute();
+
+		/** Returns true if edit copy can be executed */
+		virtual bool Copy_CanExecute();
+
+		/** Returns true if edit paste can be executed */
+		virtual bool Paste_CanExecute();
+
+		/** Returns true if clipboard contains folders only */
+		bool CanPasteFoldersOnlyFromClipboard();
+		
 	private:
 		/** Methods that implement structural modification logic for the tree */
 
@@ -220,9 +242,6 @@ namespace SceneOutliner
 		/** Scroll the specified item into view */
 		void ScrollItemIntoView(FTreeItemPtr Item);
 
-		/** Open a rename for the specified tree Item */
-		void InitiateRename(TSharedRef<ITreeItem> TrreeItem);
-
 	private:
 
 		/** Synchronize the current actor selection in the world, to the tree */
@@ -237,6 +256,9 @@ namespace SceneOutliner
 		/** Get an array of selected folders */
 		TArray<FFolderTreeItem*> GetSelectedFolders() const;
 
+		/** Get an array of selected folder names */
+		TArray<FName> GetSelectedFolderNames() const;
+
 		/** Checks to see if the actor is valid for displaying in the outliner */
 		bool IsActorDisplayable( const AActor* Actor ) const;
 
@@ -312,6 +334,86 @@ namespace SceneOutliner
 		/** Called when a folder is to be deleted */
 		void OnBroadcastFolderDelete(UWorld& InWorld, FName Path);
 
+		/** Called by engine when edit cut actors begins */
+		void OnEditCutActorsBegin();
+
+		/** Called by engine when edit cut actors ends */
+		void OnEditCutActorsEnd();
+
+		/** Called by engine when edit copy actors begins */
+		void OnEditCopyActorsBegin();
+
+		/** Called by engine when edit copy actors ends */
+		void OnEditCopyActorsEnd();
+
+		/** Called by engine when edit paste actors begins */
+		void OnEditPasteActorsBegin();
+
+		/** Called by engine when edit paste actors ends */
+		void OnEditPasteActorsEnd();
+
+		/** Called by engine when edit duplicate actors begins */
+		void OnDuplicateActorsBegin();
+
+		/** Called by engine when edit duplicate actors ends */
+		void OnDuplicateActorsEnd();
+
+		/** Called by engine when edit delete actors begins */
+		void OnDeleteActorsBegin();
+
+		/** Called by engine when edit delete actors ends */
+		void OnDeleteActorsEnd();
+
+		/** Copy specified folders to clipboard, keeping current clipboard contents if they differ from previous clipboard contents (meaning actors were copied) */
+		void CopyFoldersToClipboard(const TArray<FName>& InFolders, const FString& InPrevClipboardContents);
+
+		/** Called by copy and duplicate */
+		void CopyFoldersBegin();
+
+		/** Called by copy and duplicate */
+		void CopyFoldersEnd();
+
+		/** Called by paste and duplicate */
+		void PasteFoldersBegin(TArray<FFolderTreeItem*> InFolders);
+
+		/** Called by paste and duplicate */
+		void PasteFoldersBegin(TArray<FName> InFolders);
+
+		/** Paste folders end logic */
+		void PasteFoldersEnd();
+
+		/** Called by cut and delete */
+		void DeleteFoldersBegin();
+
+		/** Called by cute and delete */
+		void DeleteFoldersEnd();
+
+		/** Get an array of folders to paste */
+		TArray<FName> GetClipboardPasteFolders() const;
+
+		/** Construct folders export string to be used in clipboard */
+		FString ExportFolderList(TArray<FName> InFolders) const;
+			
+		/** Construct array of folders to be created based on input clipboard string */
+		TArray<FName> ImportFolderList(const FString& InStrBuffer) const;
+
+	public:
+		/** Duplicates current folder and all descendants */
+		void DuplicateFoldersHierarchy();
+	
+	private: 
+		/** Cache selected folders during edit delete */
+		TArray<FFolderTreeItem*> CacheFoldersDelete;
+
+		/** Cache folders for cut/copy/paste/duplicate */
+		TArray<FName> CacheFoldersEdit;
+
+		/** Cache clipboard contents for cut/copy */
+		FString CacheClipboardContents;
+
+		/** Maps pre-existing children during paste or duplicate */
+		TMap<FName, TArray<FTreeItemID>> CachePasteFolderExistingChildrenMap;
+
 	private:
 		/** Miscellaneous bindings required by the UI */
 
@@ -382,10 +484,10 @@ namespace SceneOutliner
 	private:
 
 		/** Select the immediate children of the currently selected folders */
-		void SelectFoldersImmediateChildren() const;
+		void SelectFoldersImmediateChildren();
 
 		/** Called to select all the descendants of the currently selected folders */
-		void SelectFoldersDescendants() const;
+		void SelectFoldersDescendants();
 
 		/** Move the selected items to the specified parent */
 		void MoveSelectionTo(FTreeItemRef NewParent);
@@ -458,6 +560,9 @@ namespace SceneOutliner
 		/** Pending tree items that are yet to be added the tree */
 		FTreeItemMap PendingTreeItemMap;
 
+		/** Folders pending selection */
+		TArray<FName> PendingFoldersSelect;
+
 		/** Root level tree items */
 		TArray<FTreeItemPtr> RootTreeItems;
 
diff --git a/Engine/Source/Editor/SceneOutliner/Private/SceneOutlinerStandaloneTypes.cpp b/Engine/Source/Editor/SceneOutliner/Private/SceneOutlinerStandaloneTypes.cpp
--- a/Engine/Source/Editor/SceneOutliner/Private/SceneOutlinerStandaloneTypes.cpp
+++ b/Engine/Source/Editor/SceneOutliner/Private/SceneOutlinerStandaloneTypes.cpp
@@ -26,12 +26,10 @@ namespace SceneOutliner
 
 	FName MoveFolderTo(FName InPath, FName NewParent, UWorld& World)
 	{
-		FName NewPath = GetFolderLeafName(InPath);
+		FName LeafName = GetFolderLeafName(InPath);
 
-		if (!NewParent.IsNone())
-		{
-			NewPath = FName(*(NewParent.ToString() / NewPath.ToString()));
-		}
+		// Get unique name
+		FName NewPath = FActorFolders::Get().GetFolderName(World, NewParent, LeafName);
 
 		if (FActorFolders::Get().RenameFolderInWorld(World, InPath, NewPath))
 		{
diff --git a/Engine/Source/Editor/SceneOutliner/Public/FolderTreeItem.h b/Engine/Source/Editor/SceneOutliner/Public/FolderTreeItem.h
--- a/Engine/Source/Editor/SceneOutliner/Public/FolderTreeItem.h
+++ b/Engine/Source/Editor/SceneOutliner/Public/FolderTreeItem.h
@@ -81,13 +81,16 @@ namespace SceneOutliner
 		/** Called to drop the specified objects on this item. Only called if ValidateDrop() allows. */
 		virtual void OnDrop(FDragDropPayload& DraggedObjects, UWorld& World, const FDragValidationInfo& ValidationInfo, TSharedRef<SWidget> DroppedOnWidget) override;
 
-		/** Delete this folder */
-		void Delete();
+		/** Delete this folder, children will be reparented to provided new parent path */
+		void Delete(FName InNewParentPath);
 
 	private:
 
 		/** Create a new folder as a child of this one */
 		void CreateSubFolder(TWeakPtr<SSceneOutliner> WeakOutliner);
+
+		/** Duplicate folder hierarchy */
+		void DuplicateHierarchy(TWeakPtr<SSceneOutliner> WeakOutliner);
 	};
 
 }		// namespace SceneOutliner
diff --git a/Engine/Source/Editor/SceneOutliner/Public/ISceneOutliner.h b/Engine/Source/Editor/SceneOutliner/Public/ISceneOutliner.h
--- a/Engine/Source/Editor/SceneOutliner/Public/ISceneOutliner.h
+++ b/Engine/Source/Editor/SceneOutliner/Public/ISceneOutliner.h
@@ -36,4 +36,22 @@ public:
 
 	/** Sets the cached icon for this class name */
 	virtual void CacheIconForClass(FName InClassName, const FSlateBrush* InSlateBrush) = 0;
+
+	/** Returns true if edit delete can be executed */
+	virtual bool Delete_CanExecute() = 0;
+
+	/** Returns true if edit rename can be executed */
+	virtual bool Rename_CanExecute() = 0;
+
+	/** Executes rename. */
+	virtual void Rename_Execute() = 0;
+
+	/** Returns true if edit cut can be executed */
+	virtual bool Cut_CanExecute() = 0;
+
+	/** Returns true if edit copy can be executed */
+	virtual bool Copy_CanExecute() = 0;
+
+	/** Returns true if edit paste can be executed */
+	virtual bool Paste_CanExecute() = 0;
 };
diff --git a/Engine/Source/Editor/SceneOutliner/SceneOutliner.Build.cs b/Engine/Source/Editor/SceneOutliner/SceneOutliner.Build.cs
--- a/Engine/Source/Editor/SceneOutliner/SceneOutliner.Build.cs
+++ b/Engine/Source/Editor/SceneOutliner/SceneOutliner.Build.cs
@@ -15,19 +15,14 @@ public class SceneOutliner : ModuleRules
 			new string[] {
 				"Core", 
 				"CoreUObject",
-				"Engine", 
-				"InputCore",
+				"Engine",
+                "ApplicationCore",
+                "InputCore",
 				"Slate", 
 				"SlateCore",
 				"EditorStyle",
 				"UnrealEd",
-			}
-		);
-
-		PrivateIncludePathModuleNames.AddRange(
-			new string[] {
-				"LevelEditor"
-			}
+            }
 		);
 	}
 }
diff --git a/Engine/Source/Editor/UnrealEd/Private/Editor.cpp b/Engine/Source/Editor/UnrealEd/Private/Editor.cpp
--- a/Engine/Source/Editor/UnrealEd/Private/Editor.cpp
+++ b/Engine/Source/Editor/UnrealEd/Private/Editor.cpp
@@ -124,6 +124,17 @@ FEditorDelegates::FOnAssetsDeleted						FEditorDelegates::OnAssetsDeleted;
 FEditorDelegates::FOnAssetDragStarted					FEditorDelegates::OnAssetDragStarted;
 FSimpleMulticastDelegate								FEditorDelegates::OnActionAxisMappingsChanged;
 FEditorDelegates::FOnAddLevelToWorld					FEditorDelegates::OnAddLevelToWorld;
+FEditorDelegates::FOnEditCutActorsBegin					FEditorDelegates::OnEditCutActorsBegin;
+FEditorDelegates::FOnEditCutActorsEnd					FEditorDelegates::OnEditCutActorsEnd;
+FEditorDelegates::FOnEditCopyActorsBegin				FEditorDelegates::OnEditCopyActorsBegin;
+FEditorDelegates::FOnEditCopyActorsEnd					FEditorDelegates::OnEditCopyActorsEnd;
+FEditorDelegates::FOnEditPasteActorsBegin				FEditorDelegates::OnEditPasteActorsBegin;
+FEditorDelegates::FOnEditPasteActorsEnd					FEditorDelegates::OnEditPasteActorsEnd;
+FEditorDelegates::FOnDuplicateActorsBegin				FEditorDelegates::OnDuplicateActorsBegin;
+FEditorDelegates::FOnDuplicateActorsEnd					FEditorDelegates::OnDuplicateActorsEnd;
+FEditorDelegates::FOnDeleteActorsBegin					FEditorDelegates::OnDeleteActorsBegin;
+FEditorDelegates::FOnDeleteActorsEnd					FEditorDelegates::OnDeleteActorsEnd;
+
 
 /*-----------------------------------------------------------------------------
 	Globals.
diff --git a/Engine/Source/Editor/UnrealEd/Private/EditorActorFolders.cpp b/Engine/Source/Editor/UnrealEd/Private/EditorActorFolders.cpp
--- a/Engine/Source/Editor/UnrealEd/Private/EditorActorFolders.cpp
+++ b/Engine/Source/Editor/UnrealEd/Private/EditorActorFolders.cpp
@@ -301,14 +301,68 @@ FName FActorFolders::GetDefaultFolderNameForSelection(UWorld& InWorld)
 	return GetDefaultFolderName(InWorld, CommonParentFolder);
 }
 
+FName FActorFolders::GetFolderName(UWorld& InWorld, FName ParentPath, FName InLeafName)
+{
+	// This is potentially very slow but necessary to find a unique name
+	const auto& ExistingFolders = GetFolderPropertiesForWorld(InWorld);
+
+	// Trim any numeric suffix
+	uint32 SuffixLen = 0;
+	for (; InLeafName.ToString().Right(SuffixLen + 1).IsNumeric(); SuffixLen++) {}
+
+	uint32 Suffix = 1;
+	FString LeafNameRoot;
+	bool bHasSuffix = true;
+	if (SuffixLen > 0)
+	{
+		LeafNameRoot = InLeafName.ToString().LeftChop(SuffixLen);
+		FString LeafSuffix = InLeafName.ToString().RightChop(InLeafName.ToString().Len() - SuffixLen);
+		Suffix = LeafSuffix.IsNumeric() ? FCString::Atoi(*LeafSuffix) : 1;
+	}
+	else
+	{
+		LeafNameRoot = InLeafName.ToString();
+		bHasSuffix = false;
+	}
+
+	// Create a valid base name for this folder
+	FNumberFormattingOptions NumberFormat;
+	NumberFormat.SetUseGrouping(false);
+	NumberFormat.SetMinimumIntegralDigits(SuffixLen);
+
+	FText LeafName = FText::Format(LOCTEXT("DefaultFolderNamePattern", "{0}{1}"), FText::FromString(LeafNameRoot), bHasSuffix ? FText::AsNumber(Suffix++, &NumberFormat) : FText::GetEmpty());
+
+	FString ParentFolderPath = ParentPath.IsNone() ? TEXT("") : ParentPath.ToString();
+	if (!ParentFolderPath.IsEmpty())
+	{
+		ParentFolderPath += "/";
+	}
+
+	FName FolderName(*(ParentFolderPath + LeafName.ToString()));
+	while (ExistingFolders.Contains(FolderName))
+	{
+		LeafName = FText::Format(LOCTEXT("DefaultFolderNamePattern", "{0}{1}"), FText::FromString(LeafNameRoot), FText::AsNumber(Suffix++, &NumberFormat));
+		FolderName = FName(*(ParentFolderPath + LeafName.ToString()));
+		if (Suffix == 0)
+		{
+			// We've wrapped around a 32bit unsigned int - something must be seriously wrong!
+			return FName();
+		}
+	}
+
+	return FolderName;
+}
+
 FName FActorFolders::GetDefaultFolderName(UWorld& InWorld, FName ParentPath)
 {
 	// This is potentially very slow but necessary to find a unique name
 	const auto& ExistingFolders = GetFolderPropertiesForWorld(InWorld);
 
 	// Create a valid base name for this folder
+	FNumberFormattingOptions NumberFormat;
+	NumberFormat.SetUseGrouping(false);
 	uint32 Suffix = 1;
-	FText LeafName = FText::Format(LOCTEXT("DefaultFolderNamePattern", "NewFolder{0}"), FText::AsNumber(Suffix++));
+	FText LeafName = FText::Format(LOCTEXT("DefaultFolderNamePattern", "NewFolder{0}"), FText::AsNumber(Suffix++, &NumberFormat));
 
 	FString ParentFolderPath = ParentPath.IsNone() ? TEXT("") : ParentPath.ToString();
 	if (!ParentFolderPath.IsEmpty())
@@ -319,7 +373,7 @@ FName FActorFolders::GetDefaultFolderName(UWorld& InWorld, FName ParentPath)
 	FName FolderName(*(ParentFolderPath + LeafName.ToString()));
 	while (ExistingFolders.Contains(FolderName))
 	{
-		LeafName = FText::Format(LOCTEXT("DefaultFolderNamePattern", "NewFolder{0}"), FText::AsNumber(Suffix++));
+		LeafName = FText::Format(LOCTEXT("DefaultFolderNamePattern", "NewFolder{0}"), FText::AsNumber(Suffix++, &NumberFormat));
 		FolderName = FName(*(ParentFolderPath + LeafName.ToString()));
 		if (Suffix == 0)
 		{
diff --git a/Engine/Source/Editor/UnrealEd/Private/UnrealEdSrv.cpp b/Engine/Source/Editor/UnrealEd/Private/UnrealEdSrv.cpp
--- a/Engine/Source/Editor/UnrealEd/Private/UnrealEdSrv.cpp
+++ b/Engine/Source/Editor/UnrealEd/Private/UnrealEdSrv.cpp
@@ -1587,7 +1587,10 @@ bool UUnrealEdEngine::Exec_Edit( UWorld* InWorld, const TCHAR* Str, FOutputDevic
 		}
 		else
 		{
+			const FScopedTransaction Transaction(NSLOCTEXT("UnrealEd", "Cut", "Cut"));
+			FEditorDelegates::OnEditCutActorsBegin.Broadcast();
 			CopySelectedActorsToClipboard(InWorld, true);
+			FEditorDelegates::OnEditCutActorsEnd.Broadcast();
 		}
 	}
 	else if( FParse::Command(&Str,TEXT("COPY")) )
@@ -1608,7 +1611,9 @@ bool UUnrealEdEngine::Exec_Edit( UWorld* InWorld, const TCHAR* Str, FOutputDevic
 		}
 		else
 		{
+			FEditorDelegates::OnEditCopyActorsBegin.Broadcast();
 			CopySelectedActorsToClipboard(InWorld, false);
+			FEditorDelegates::OnEditCopyActorsEnd.Broadcast();
 		}
 	}
 	else if( FParse::Command(&Str,TEXT("PASTE")) )
@@ -1652,7 +1657,10 @@ bool UUnrealEdEngine::Exec_Edit( UWorld* InWorld, const TCHAR* Str, FOutputDevic
 				}
 			}
 
+			const FScopedTransaction Transaction(TransDescription);
+			FEditorDelegates::OnEditPasteActorsBegin.Broadcast();
 			PasteSelectedActorsFromClipboard(InWorld, TransDescription, PasteTo);
+			FEditorDelegates::OnEditPasteActorsEnd.Broadcast();
 		}
 	}
 
@@ -2550,7 +2558,9 @@ bool UUnrealEdEngine::Exec_Actor( UWorld* InWorld, const TCHAR* Str, FOutputDevi
 		if (!bHandled)
 		{
 			const FScopedTransaction Transaction( bComponentsSelected ? NSLOCTEXT("UnrealEd", "DeleteComponents", "Delete Components") : NSLOCTEXT("UnrealEd", "DeleteActors", "Delete Actors") );
+			FEditorDelegates::OnDeleteActorsBegin.Broadcast();
 			edactDeleteSelected( InWorld );
+			FEditorDelegates::OnDeleteActorsEnd.Broadcast();
 		}
 		return true;
 	}
@@ -2734,6 +2744,8 @@ bool UUnrealEdEngine::Exec_Actor( UWorld* InWorld, const TCHAR* Str, FOutputDevi
 			//@todo locked levels - if all actor levels are locked, cancel the transaction
 			const FScopedTransaction Transaction( bComponentsSelected ? NSLOCTEXT("UnrealEd", "DuplicateComponents", "Duplicate Components") : NSLOCTEXT("UnrealEd", "DuplicateActors", "Duplicate Actors") );
 
+			FEditorDelegates::OnDuplicateActorsBegin.Broadcast();
+
 			// duplicate selected
 			ABrush::SetSuppressBSPRegeneration(true);
 			edactDuplicateSelected(InWorld->GetCurrentLevel(), GetDefault<ULevelEditorViewportSettings>()->GridEnabled);
@@ -2746,6 +2758,8 @@ bool UUnrealEdEngine::Exec_Actor( UWorld* InWorld, const TCHAR* Str, FOutputDevi
 			{
 				RebuildAlteredBSP(); // Update the Bsp of any levels containing a modified brush
 			}
+
+			FEditorDelegates::OnDuplicateActorsEnd.Broadcast();
 		}
 		RedrawLevelEditingViewports();
 		return true;
diff --git a/Engine/Source/Editor/UnrealEd/Public/Editor.h b/Engine/Source/Editor/UnrealEd/Public/Editor.h
--- a/Engine/Source/Editor/UnrealEd/Public/Editor.h
+++ b/Engine/Source/Editor/UnrealEd/Public/Editor.h
@@ -92,6 +92,26 @@ struct UNREALED_API FEditorDelegates
 	DECLARE_MULTICAST_DELEGATE_OneParam(FOnAddLevelToWorld, ULevel*);
 	/** delegate type for when a texture is fit to surface  */
 	DECLARE_MULTICAST_DELEGATE_OneParam(FOnFitTextureToSurface, UWorld*);
+	/** delegate type for before edit cut actors is handled */
+	DECLARE_MULTICAST_DELEGATE(FOnEditCutActorsBegin);
+	/** delegate type for after edit cut actors is handled */
+	DECLARE_MULTICAST_DELEGATE(FOnEditCutActorsEnd);
+	/** delegate type for before edit copy actors is handled */
+	DECLARE_MULTICAST_DELEGATE(FOnEditCopyActorsBegin);
+	/** delegate type for after edit copy actors is handled */
+	DECLARE_MULTICAST_DELEGATE(FOnEditCopyActorsEnd);
+	/** delegate type for before edit paste actors is handled */
+	DECLARE_MULTICAST_DELEGATE(FOnEditPasteActorsBegin);
+	/** delegate type for after edit paste actors is handled */
+	DECLARE_MULTICAST_DELEGATE(FOnEditPasteActorsEnd);
+	/** delegate type for before edit duplicate actors is handled */
+	DECLARE_MULTICAST_DELEGATE(FOnDuplicateActorsBegin);
+	/** delegate type for after edit duplicate actors is handled */
+	DECLARE_MULTICAST_DELEGATE(FOnDuplicateActorsEnd);	
+	/** delegate type for before delete actors is handled */
+	DECLARE_MULTICAST_DELEGATE(FOnDeleteActorsBegin);
+	/** delegate type for after delete actors is handled */
+	DECLARE_MULTICAST_DELEGATE(FOnDeleteActorsEnd);
 
 	/** Called when the CurrentLevel is switched to a new level.  Note that this event won't be fired for temporary
 		changes to the current level, such as when copying/pasting actors. */
@@ -205,6 +225,26 @@ struct UNREALED_API FEditorDelegates
 	static FSimpleMulticastDelegate OnActionAxisMappingsChanged;
 	/** Called from FEditorUtils::AddLevelToWorld after the level is added successfully to the world. */
 	static FOnAddLevelToWorld OnAddLevelToWorld;
+	/** Sent before edit cut is handled */
+	static FOnEditCutActorsBegin OnEditCutActorsBegin;
+	/** Sent after edit cut is handled */
+	static FOnEditCutActorsEnd OnEditCutActorsEnd;
+	/** Sent before edit copy is handled */
+	static FOnEditCopyActorsBegin OnEditCopyActorsBegin;
+	/** Sent after edit copy is handled */
+	static FOnEditCopyActorsEnd OnEditCopyActorsEnd;
+	/** Sent before edit paste is handled */
+	static FOnEditPasteActorsBegin OnEditPasteActorsBegin;
+	/** Sent after edit paste is handled */
+	static FOnEditPasteActorsEnd OnEditPasteActorsEnd;
+	/** Sent before duplicate is handled */
+	static FOnDuplicateActorsBegin OnDuplicateActorsBegin;
+	/** Sent after duplicate is handled */
+	static FOnDuplicateActorsEnd OnDuplicateActorsEnd;
+	/** Sent when delete begin called */
+	static FOnDeleteActorsBegin OnDeleteActorsBegin;
+	/** Sent when delete end called */
+	static FOnDeleteActorsEnd OnDeleteActorsEnd;
 };
 
 /**
diff --git a/Engine/Source/Editor/UnrealEd/Public/EditorActorFolders.h b/Engine/Source/Editor/UnrealEd/Public/EditorActorFolders.h
--- a/Engine/Source/Editor/UnrealEd/Public/EditorActorFolders.h
+++ b/Engine/Source/Editor/UnrealEd/Public/EditorActorFolders.h
@@ -86,6 +86,9 @@ struct UNREALED_API FActorFolders : public FGCObject
 	/** Get a new default folder name that would apply to the current selection */
 	FName GetDefaultFolderNameForSelection(UWorld& InWorld);
 
+	/** Get folder name that is unique under specified parent path */
+	FName GetFolderName(UWorld& InWorld, FName ParentPath, FName FolderName);
+
 	/** Create a new folder in the specified world, of the specified path */
 	void CreateFolder(UWorld& InWorld, FName Path);
 
