World Outliner: Update folder path during actor attachment. Was causing deleted files to reappear when map is reopened.

diff --git a/Engine/Source/Editor/SceneOutliner/Private/ActorTreeItem.cpp b/Engine/Source/Editor/SceneOutliner/Private/ActorTreeItem.cpp
--- a/Engine/Source/Editor/SceneOutliner/Private/ActorTreeItem.cpp
+++ b/Engine/Source/Editor/SceneOutliner/Private/ActorTreeItem.cpp
@@ -175,6 +175,8 @@ void FActorDropTarget::PerformAttachment(FName SocketName, TWeakObjectPtr<AActor
 			if (GEditor->CanParentActors(ParentActor, ChildActor))
 			{
 				GEditor->ParentActors(ParentActor, ChildActor, SocketName);
+
+				ChildActor->SetFolderPath_Recursively(ParentActor->GetFolderPath());
 			}
 		}
 	}
