World Outliner: Fixed Select Descendants and Select Immediate Children to expand minimized folders/actors so selection will be correct in outliner. Select actors only, no folders (changed back to previous behavior).

diff --git a/Engine/Source/Editor/SceneOutliner/Private/SSceneOutliner.cpp b/Engine/Source/Editor/SceneOutliner/Private/SSceneOutliner.cpp
--- a/Engine/Source/Editor/SceneOutliner/Private/SSceneOutliner.cpp
+++ b/Engine/Source/Editor/SceneOutliner/Private/SSceneOutliner.cpp
@@ -1646,73 +1646,63 @@ namespace SceneOutliner
 			LOCTEXT( "AddChildrenToSelection", "Immediate Children" ),
 			LOCTEXT( "AddChildrenToSelection_ToolTip", "Select all immediate actor children of the selected folders" ),
 			FSlateIcon(),
-			FExecuteAction::CreateSP(this, &SSceneOutliner::SelectFoldersImmediateChildren));
+			FExecuteAction::CreateSP(this, &SSceneOutliner::SelectFoldersDescendants, /*bSelectImmediateChildrenOnly=*/ true));
 		MenuBuilder.AddMenuEntry(
 			LOCTEXT( "AddDescendantsToSelection", "All Descendants" ),
 			LOCTEXT( "AddDescendantsToSelection_ToolTip", "Select all actor descendants of the selected folders" ),
 			FSlateIcon(),
-			FExecuteAction::CreateSP(this, &SSceneOutliner::SelectFoldersDescendants));
+			FExecuteAction::CreateSP(this, &SSceneOutliner::SelectFoldersDescendants, /*bSelectImmediateChildrenOnly=*/ false));
 	}
 
-	void SSceneOutliner::SelectFoldersImmediateChildren()
+	void SSceneOutliner::SelectFoldersDescendants(bool bSelectImmediateChildrenOnly)
 	{
-		struct FSelectActors : ITreeItemVisitor
+		struct FExpandFoldersRecursive : IMutableTreeItemVisitor
 		{
-			virtual void Visit(const FActorTreeItem& ActorItem) const override
+			SSceneOutliner& Outliner;
+			bool bSelectImmediateChildrenOnly;
+
+			FExpandFoldersRecursive(SSceneOutliner& InOutliner, bool bInSelectImmediateChildrenOnly) :
+				Outliner(InOutliner), bSelectImmediateChildrenOnly(bInSelectImmediateChildrenOnly) {}
+
+			virtual void Visit(FActorTreeItem& ActorItem) const override
 			{
-				if (AActor* Actor = ActorItem.Actor.Get())
+				if (!Outliner.OutlinerTreeView->IsItemExpanded(ActorItem.AsShared()))
 				{
-					GEditor->SelectActor(Actor, true, /*bNotify=*/false);
+					Outliner.OutlinerTreeView->SetItemExpansion(ActorItem.AsShared(), true);
 				}
-			}
-		};
 
-		struct FSelectFolders : IMutableTreeItemVisitor
-		{
-			SSceneOutliner& Outliner;
-			FSelectFolders(SSceneOutliner& InOutliner) : Outliner(InOutliner) {}
-
-			virtual void Visit(FFolderTreeItem& FolderItem) const override
-			{
-				Outliner.OutlinerTreeView->SetItemSelection(FolderItem.AsShared(), true);
+				if (!bSelectImmediateChildrenOnly)
+				{
+					for (TWeakPtr<ITreeItem> Child : ActorItem.GetChildren())
+					{
+						Child.Pin()->Visit(*this);
+					}
+				}
 			}
-		};
-
-		TArray<FFolderTreeItem*> SelectedFolders = GetSelectedFolders();
-		OutlinerTreeView->ClearSelection();
-
-		if (SelectedFolders.Num())
-		{
-			// We'll batch selection changes instead by using BeginBatchSelectOperation()
-			GEditor->GetSelectedActors()->BeginBatchSelectOperation();
 
-			FSelectActors SelectActorsVisitor;
-			for (const FFolderTreeItem* Folder : SelectedFolders)
+			virtual void Visit(FFolderTreeItem& FolderItem) const override
 			{
-				for (TWeakPtr<ITreeItem> Child : Folder->GetChildren())
+				if (!Outliner.OutlinerTreeView->IsItemExpanded(FolderItem.AsShared()))
 				{
-					Child.Pin()->Visit(SelectActorsVisitor);
+					Outliner.OutlinerTreeView->SetItemExpansion(FolderItem.AsShared(), true);
 				}
-			}
 
-			GEditor->GetSelectedActors()->EndBatchSelectOperation();
-			GEditor->NoteSelectionChange();
-		}
+				if (!bSelectImmediateChildrenOnly)
+				{
+					for (TWeakPtr<ITreeItem> Child : FolderItem.GetChildren())
+					{
+						Child.Pin()->Visit(*this);
+					}
+				}
 
-		FSelectFolders SelectFoldersVisitor(*this);
-		for (const FFolderTreeItem* Folder : SelectedFolders)
-		{
-			for (TWeakPtr<ITreeItem> Child : Folder->GetChildren())
-			{
-				Child.Pin()->Visit(SelectFoldersVisitor);
 			}
-		}
-	}
+		};
 
-	void SSceneOutliner::SelectFoldersDescendants()
-	{
 		struct FSelectActorsRecursive : ITreeItemVisitor
 		{
+			bool bSelectImmediateChildrenOnly;
+			FSelectActorsRecursive(bool bInSelectImmediateChildrenOnly) : bSelectImmediateChildrenOnly(bInSelectImmediateChildrenOnly) {}
+
 			virtual void Visit(const FActorTreeItem& ActorItem) const override
 			{
 				if (AActor* Actor = ActorItem.Actor.Get())
@@ -1720,17 +1710,23 @@ namespace SceneOutliner
 					GEditor->SelectActor(Actor, true, /*bNotify=*/false);
 				}
 
-				for (TWeakPtr<ITreeItem> Child : ActorItem.GetChildren())
+				if (!bSelectImmediateChildrenOnly)
 				{
-					Child.Pin()->Visit(*this);
+					for (TWeakPtr<ITreeItem> Child : ActorItem.GetChildren())
+					{
+						Child.Pin()->Visit(*this);
+					}
 				}
 			}
 
 			virtual void Visit(const FFolderTreeItem& FolderItem) const override
 			{
-				for (TWeakPtr<ITreeItem> Child : FolderItem.GetChildren())
+				if (!bSelectImmediateChildrenOnly)
 				{
-					Child.Pin()->Visit(*this);
+					for (TWeakPtr<ITreeItem> Child : FolderItem.GetChildren())
+					{
+						Child.Pin()->Visit(*this);
+					}
 				}
 			}
 		};
@@ -1738,38 +1734,53 @@ namespace SceneOutliner
 		struct FSelectFoldersRecursive : IMutableTreeItemVisitor
 		{
 			SSceneOutliner& Outliner;
-			FSelectFoldersRecursive(SSceneOutliner& InOutliner) : Outliner(InOutliner) {}
+			bool bSelectImmediateChildrenOnly;
+			FSelectFoldersRecursive(SSceneOutliner& InOutliner, bool bInSelectImmediateChildrenOnly) : Outliner(InOutliner), bSelectImmediateChildrenOnly(bInSelectImmediateChildrenOnly) {}
 
 			virtual void Visit(FFolderTreeItem& FolderItem) const override
 			{
 				Outliner.OutlinerTreeView->SetItemSelection(FolderItem.AsShared(), true);
 
-				for (TWeakPtr<ITreeItem> Child : FolderItem.GetChildren())
+				if (!bSelectImmediateChildrenOnly)
 				{
-					Child.Pin()->Visit(*this);
-				}
+					for (TWeakPtr<ITreeItem> Child : FolderItem.GetChildren())
+					{
+						Child.Pin()->Visit(*this);
+					}
+				}	
 			}
 		};
 
 		TArray<FFolderTreeItem*> SelectedFolders = GetSelectedFolders();
+		OutlinerTreeView->ClearSelection();
+
+		FExpandFoldersRecursive ExpandFoldersRecursive(*this, bSelectImmediateChildrenOnly);
+		for (FFolderTreeItem* Folder : SelectedFolders)
+		{
+			Folder->Visit(ExpandFoldersRecursive);
+		}
+
 		if (SelectedFolders.Num())
 		{
 			// We'll batch selection changes instead by using BeginBatchSelectOperation()
 			GEditor->GetSelectedActors()->BeginBatchSelectOperation();
 
-			OutlinerTreeView->ClearSelection();
-
-			FSelectActorsRecursive SelectActorsRecursive;
+			FSelectActorsRecursive SelectActorsRecursive(bSelectImmediateChildrenOnly);
 			for (const FFolderTreeItem* Folder : SelectedFolders)
 			{
-				Folder->Visit(SelectActorsRecursive);
+				for (TWeakPtr<ITreeItem> Child : Folder->GetChildren())
+				{
+					Child.Pin()->Visit(SelectActorsRecursive);
+				}
 			}
 
 			GEditor->GetSelectedActors()->EndBatchSelectOperation();
 			GEditor->NoteSelectionChange();
 		}
 
-		FSelectFoldersRecursive SelectFoldersRecursiveVisitor(*this);
+		// Don't select folders, only select actors
+		/*  
+		FSelectFoldersRecursive SelectFoldersRecursiveVisitor(*this, bSelectImmediateChildrenOnly);
 		for (const FFolderTreeItem* Folder : SelectedFolders)
 		{
 			for (TWeakPtr<ITreeItem> Child : Folder->GetChildren())
@@ -1777,6 +1788,9 @@ namespace SceneOutliner
 				Child.Pin()->Visit(SelectFoldersRecursiveVisitor);
 			}
 		}
+		*/
+
+		Refresh();
 	}
 
 	void SSceneOutliner::MoveSelectionTo(FTreeItemRef NewParent)
diff --git a/Engine/Source/Editor/SceneOutliner/Private/SSceneOutliner.h b/Engine/Source/Editor/SceneOutliner/Private/SSceneOutliner.h
--- a/Engine/Source/Editor/SceneOutliner/Private/SSceneOutliner.h
+++ b/Engine/Source/Editor/SceneOutliner/Private/SSceneOutliner.h
@@ -483,11 +483,8 @@ namespace SceneOutliner
 
 	private:
 
-		/** Select the immediate children of the currently selected folders */
-		void SelectFoldersImmediateChildren();
-
-		/** Called to select all the descendants of the currently selected folders */
-		void SelectFoldersDescendants();
+		/** Called to select descendants of the currently selected folders */
+		void SelectFoldersDescendants(bool bSelectImmediateChildrenOnly = false);
 
 		/** Move the selected items to the specified parent */
 		void MoveSelectionTo(FTreeItemRef NewParent);
