Improving speed of scene outliner updates when level visibility changes
#rnx

diff --git a/Engine/Source/Editor/SceneOutliner/Private/SSceneOutliner.cpp b/Engine/Source/Editor/SceneOutliner/Private/SSceneOutliner.cpp
--- a/Engine/Source/Editor/SceneOutliner/Private/SSceneOutliner.cpp
+++ b/Engine/Source/Editor/SceneOutliner/Private/SSceneOutliner.cpp
@@ -549,7 +549,7 @@ namespace SceneOutliner
 		// @todo outliner: Might not catch some cases (see: CALLBACK_ActorPropertiesChange, CALLBACK_LayerChange, CALLBACK_LevelDirtied, CALLBACK_OnActorMoved, CALLBACK_UpdateLevelsForAllActors)
 		FEditorDelegates::MapChange.AddSP( this, &SSceneOutliner::OnMapChange );
 		FEditorDelegates::NewCurrentLevel.AddSP( this, &SSceneOutliner::OnNewCurrentLevel );
-		GEngine->OnLevelActorListChanged().AddSP( this, &SSceneOutliner::FullRefresh );
+		GEngine->OnLevelActorListChanged().AddSP( this, &SSceneOutliner::OnLevelActorListChanged );
 		FWorldDelegates::LevelAddedToWorld.AddSP( this, &SSceneOutliner::OnLevelAdded );
 		FWorldDelegates::LevelRemovedFromWorld.AddSP( this, &SSceneOutliner::OnLevelRemoved );
 
@@ -1014,6 +1014,11 @@ namespace SceneOutliner
 		Refresh();
 	}
 
+	void SSceneOutliner::OnLevelActorListChanged()
+	{
+		bDisableIntermediateSorting = true;
+		FullRefresh();
+	}
 
 	void SSceneOutliner::Populate()
 	{
@@ -1130,6 +1135,7 @@ namespace SceneOutliner
 		PendingOperations.RemoveAt(0, End);
 		SetParentsExpansionState(ExpansionStateInfo);
 
+
 		for (FName Folder : PendingFoldersSelect)
 		{
 			if (FTreeItemPtr* Item = TreeItemMap.Find(Folder))
@@ -1139,16 +1145,25 @@ namespace SceneOutliner
 		}
 		PendingFoldersSelect.Empty();
 
-		if (bMadeAnySignificantChanges && !SharedData->bRepresentingPlayWorld)
-		{
-			RequestSort();
-		}
-
+		// Check if we need to sort because we are finished with the populating operations
+		bool bFinalSort = false;
 		if (PendingOperations.Num() == 0)
 		{
 			// We're fully refreshed now.
 			NewItemActions.Empty();
 			bNeedsRefresh = false;
+			if (bDisableIntermediateSorting)
+			{
+				bDisableIntermediateSorting = false;
+				bFinalSort = true;
+			}
+		}
+
+		// If we are allowing intermediate sorts and met the conditions, or this is the final sort after all ops are complete
+		if ((bMadeAnySignificantChanges && !SharedData->bRepresentingPlayWorld && !bDisableIntermediateSorting) ||
+			bFinalSort)
+		{
+			RequestSort();
 		}
 	}
 
diff --git a/Engine/Source/Editor/SceneOutliner/Public/SSceneOutliner.h b/Engine/Source/Editor/SceneOutliner/Public/SSceneOutliner.h
--- a/Engine/Source/Editor/SceneOutliner/Public/SSceneOutliner.h
+++ b/Engine/Source/Editor/SceneOutliner/Public/SSceneOutliner.h
@@ -181,6 +181,9 @@ namespace SceneOutliner
 		/** Tells the scene outliner that it should do a full refresh, which will clear the entire tree and rebuild it from scratch. */
 		void FullRefresh();
 
+		/** Tells the scene outliner that there was a change in the level actor list. */
+		void OnLevelActorListChanged();
+
 		/** Attempts to add an item to the tree. Will add any parents if required. */
 		bool AddItemToTree(FTreeItemRef InItem);
 		
@@ -620,6 +623,9 @@ namespace SceneOutliner
 		/** true if the Scene Outliner should do a full refresh. */
 		bool bFullRefresh;
 
+		/** True if the Scene Outliner is currently responding to a level visibility change */
+		bool bDisableIntermediateSorting;
+
 		/** true when the actor selection state in the world does not match the selection state of the tree */
 		bool bActorSelectionDirty;
 
