Update search results after pressing ESC in World Outliner

diff --git a/Engine/Source/Editor/SceneOutliner/Private/SSceneOutliner.cpp b/Engine/Source/Editor/SceneOutliner/Private/SSceneOutliner.cpp
--- a/Engine/Source/Editor/SceneOutliner/Private/SSceneOutliner.cpp
+++ b/Engine/Source/Editor/SceneOutliner/Private/SSceneOutliner.cpp
@@ -3053,6 +3053,10 @@ namespace SceneOutliner
 				}
 			}
 		}
+		else if (CommitInfo == ETextCommit::OnCleared)
+		{
+			OnFilterTextChanged(InFilterText);
+		}
 	}
 
 	EVisibility SSceneOutliner::GetFilterStatusVisibility() const
