PR #5552: UE-20016: Respect RF_Transient object flag in World Outliner (Contributed by [REDACTED_PERSON])

diff --git a/Engine/Source/Editor/LevelEditor/Private/SLevelEditor.cpp b/Engine/Source/Editor/LevelEditor/Private/SLevelEditor.cpp
--- a/Engine/Source/Editor/LevelEditor/Private/SLevelEditor.cpp
+++ b/Engine/Source/Editor/LevelEditor/Private/SLevelEditor.cpp
@@ -668,6 +668,7 @@ TSharedRef<SDockTab> SLevelEditor::SpawnLevelEditorTab( const FSpawnTabArgs& Arg
 	else if( TabIdentifier == TEXT("LevelEditorSceneOutliner") )
 	{
 		SceneOutliner::FInitializationOptions InitOptions;
+		InitOptions.bShowTransient = true;
 		InitOptions.Mode = ESceneOutlinerMode::ActorBrowsing;
 		{
 			TWeakPtr<SLevelEditor> WeakLevelEditor = SharedThis(this);
diff --git a/Engine/Source/Editor/SceneOutliner/Private/SSceneOutliner.cpp b/Engine/Source/Editor/SceneOutliner/Private/SSceneOutliner.cpp
--- a/Engine/Source/Editor/SceneOutliner/Private/SSceneOutliner.cpp
+++ b/Engine/Source/Editor/SceneOutliner/Private/SSceneOutliner.cpp
@@ -118,7 +118,7 @@ namespace SceneOutliner
 	TSharedPtr< FOutlinerFilter > CreateHideTemporaryActorsFilter()
 	{
 		return MakeShareable( new FOutlinerPredicateFilter( FActorFilterPredicate::CreateStatic( []( const AActor* InActor ){			
-			return (InActor->GetWorld() && InActor->GetWorld()->WorldType != EWorldType::PIE) || GEditor->ObjectsThatExistInEditorWorld.Get(InActor);
+			return ((InActor->GetWorld() && InActor->GetWorld()->WorldType != EWorldType::PIE) || GEditor->ObjectsThatExistInEditorWorld.Get(InActor)) && !InActor->HasAnyFlags(EObjectFlags::RF_Transient);
 		} ), EDefaultFilterBehaviour::Pass ) );
 	}
 
