PR #5732: Error message when trying to dragdrop assets in the world outliner (Contributed by [REDACTED_PERSON])

diff --git a/Engine/Source/Editor/SceneOutliner/Private/SOutlinerTreeView.cpp b/Engine/Source/Editor/SceneOutliner/Private/SOutlinerTreeView.cpp
--- a/Engine/Source/Editor/SceneOutliner/Private/SOutlinerTreeView.cpp
+++ b/Engine/Source/Editor/SceneOutliner/Private/SOutlinerTreeView.cpp
@@ -33,7 +33,7 @@ namespace SceneOutliner
 			else if (Operation->IsOfType<FDecoratedDragDropOp>())
 			{
 				auto* DecoratedOp = static_cast<FDecoratedDragDropOp*>(Operation);
-				DecoratedOp->SetToolTip(ValidationInfo.ValidationText, Icon);
+				DecoratedOp->SetToolTip(LOCTEXT("SceneOutlinerInvalidDragDropOp", "Invalid drag and drop operation, Drag into the Viewport."), Icon);
 			}
 		}
 	}
@@ -134,14 +134,9 @@ namespace SceneOutliner
 		FFolderDropTarget DropTarget(NAME_None);
 
 		auto Reply = HandleDrop(SceneOutlinerWeak, DragDropEvent, DropTarget, ValidationInfo);
-
-		if (Reply.IsEventHandled())
-		{
-			UpdateOperationDecorator(DragDropEvent, ValidationInfo);
-			return FReply::Handled();
-		}
-
-		return FReply::Unhandled();
+		UpdateOperationDecorator(DragDropEvent, ValidationInfo);
+		
+		return Reply;
 	}
 	
 	void SOutlinerTreeView::OnDragLeave(const FDragDropEvent& DragDropEvent)
