Fix for Scene Outliner not represhing added actors during PIE

diff --git a/Engine/Source/Editor/SceneOutliner/Private/SSceneOutliner.cpp b/Engine/Source/Editor/SceneOutliner/Private/SSceneOutliner.cpp
--- a/Engine/Source/Editor/SceneOutliner/Private/SSceneOutliner.cpp
+++ b/Engine/Source/Editor/SceneOutliner/Private/SSceneOutliner.cpp
@@ -1160,8 +1160,7 @@ namespace SceneOutliner
 		}
 
 		// If we are allowing intermediate sorts and met the conditions, or this is the final sort after all ops are complete
-		if ((bMadeAnySignificantChanges && !SharedData->bRepresentingPlayWorld && !bDisableIntermediateSorting) ||
-			bFinalSort)
+		if ((bMadeAnySignificantChanges && !bDisableIntermediateSorting) || bFinalSort)
 		{
 			RequestSort();
 		}
