Adding the ability to hide folders in World Outliner if all nested actors are hidden.

diff --git a/Engine/Source/Editor/SceneOutliner/Private/SSceneOutliner.cpp b/Engine/Source/Editor/SceneOutliner/Private/SSceneOutliner.cpp
--- a/Engine/Source/Editor/SceneOutliner/Private/SSceneOutliner.cpp
+++ b/Engine/Source/Editor/SceneOutliner/Private/SSceneOutliner.cpp
@@ -762,6 +762,19 @@ namespace SceneOutliner
 					EUserInterfaceActionType::ToggleButton
 				);
 
+				MenuBuilder.AddMenuEntry(
+					LOCTEXT("ToggleHideFoldersContainingHiddenActors", "Hide Folders with Only Hidden Actors"),
+					LOCTEXT("ToggleHideFoldersContainingHiddenActorsToolTip", "When enabled, only shows Folders containing non-hidden Actors."),
+					FSlateIcon(),
+					FUIAction(
+						FExecuteAction::CreateSP(this, &SSceneOutliner::ToggleHideFoldersContainingOnlyHiddenActors),
+						FCanExecuteAction(),
+						FIsActionChecked::CreateSP(this, &SSceneOutliner::IsHidingFoldersContainingOnlyHiddenActors)
+					),
+					NAME_None,
+					EUserInterfaceActionType::ToggleButton
+				);
+
 				if (bActorComponentsEnabled)
 				{
 					MenuBuilder.AddMenuEntry(
@@ -940,6 +953,17 @@ namespace SceneOutliner
 		}
 	}
 
+	void SSceneOutliner::ToggleHideFoldersContainingOnlyHiddenActors()
+	{
+		const bool bEnableFlag = !IsHidingFoldersContainingOnlyHiddenActors();
+
+		USceneOutlinerSettings* Settings = GetMutableDefault<USceneOutlinerSettings>();
+		Settings->bHideFoldersContainingHiddenActors = bEnableFlag;
+		Settings->PostEditChange();
+
+		FullRefresh();
+	}
+
 	bool SSceneOutliner::IsShowingActorComponents()
 	{
 		return bActorComponentsEnabled && GetDefault<USceneOutlinerSettings>()->bShowActorComponents;
@@ -978,6 +1002,11 @@ namespace SceneOutliner
 		return GetDefault<USceneOutlinerSettings>()->bShowOnlyActorsInCurrentLevel;
 	}
 
+	bool SSceneOutliner::IsHidingFoldersContainingOnlyHiddenActors() const
+	{
+		return GetDefault<USceneOutlinerSettings>()->bHideFoldersContainingHiddenActors;
+	}
+
 	/** END FILTERS */
 
 
@@ -1008,6 +1037,10 @@ namespace SceneOutliner
 
 	void SSceneOutliner::Refresh()
 	{
+		if (IsHidingFoldersContainingOnlyHiddenActors())
+		{
+			bFullRefresh = true;
+		}
 		bNeedsRefresh = true;
 	}
 
@@ -1160,6 +1193,8 @@ namespace SceneOutliner
 				bDisableIntermediateSorting = false;
 				bFinalSort = true;
 			}
+
+			HideFoldersContainingOnlyHiddenActors();
 		}
 
 		// If we are allowing intermediate sorts and met the conditions, or this is the final sort after all ops are complete
@@ -1455,6 +1490,74 @@ namespace SceneOutliner
 		}
 	}
 
+	void SSceneOutliner::HideFoldersContainingOnlyHiddenActors()
+	{
+		if (IsHidingFoldersContainingOnlyHiddenActors())
+		{
+			for (const FTreeItemPtr & TreeItem : RootTreeItems)
+			{
+				HideFoldersContainingOnlyHiddenActors(TreeItem, true);
+			}
+		}
+	}
+
+	bool SSceneOutliner::HideFoldersContainingOnlyHiddenActors(FTreeItemPtr Parent, bool bIsRoot)
+	{
+		TArray<FTreeItemPtr> ItemsToRemove;
+
+		bool bActorsHidden = true;
+		bool bFoldersHidden = true;
+
+		const TArray<TWeakPtr<ITreeItem>>& Children = Parent->GetChildren();
+
+		if (Children.Num())
+		{
+
+			for (const TWeakPtr<ITreeItem> & ChildItem : Children)
+			{
+				FTreeItemPtr TreeItem = ChildItem.Pin();
+
+				TWeakObjectPtr<AActor> TrueTreeActor = StaticCastSharedPtr<FActorTreeItem>(TreeItem)->Actor;
+
+				if (!bIsRoot && TrueTreeActor.IsValid())
+				{
+					if (bActorsHidden)
+					{
+						if (!(TrueTreeActor->IsTemporarilyHiddenInEditor()))
+						{
+							bActorsHidden = false;
+						}
+					}
+				}
+				else
+				{
+					bool bCurrentFolderHidden = HideFoldersContainingOnlyHiddenActors(TreeItem);
+
+					if (bCurrentFolderHidden)
+					{
+						ItemsToRemove.Add(TreeItem);
+					}
+
+					bFoldersHidden = bCurrentFolderHidden & bFoldersHidden;
+
+				}
+
+			}
+		}
+		else
+		{
+			return false;
+		}
+
+		for (const FTreeItemPtr & Item : ItemsToRemove)
+		{
+			FTreeItemRef RemoveItem = Item.ToSharedRef();
+			Parent->RemoveChild(RemoveItem);
+		}
+
+		return bActorsHidden && bFoldersHidden;
+	}
+
 	void SSceneOutliner::PopulateSearchStrings(const ITreeItem& Item, TArray< FString >& OutSearchStrings) const
 	{
 		for (const auto& Pair : Columns)
diff --git a/Engine/Source/Editor/SceneOutliner/Private/SceneOutlinerGutter.cpp b/Engine/Source/Editor/SceneOutliner/Private/SceneOutlinerGutter.cpp
--- a/Engine/Source/Editor/SceneOutliner/Private/SceneOutlinerGutter.cpp
+++ b/Engine/Source/Editor/SceneOutliner/Private/SceneOutlinerGutter.cpp
@@ -308,11 +308,15 @@ private:
 	/** Set the actor this widget is responsible for to be hidden or shown */
 	void SetIsVisible(const bool bVisible)
 	{
-		auto TreeItem = WeakTreeItem.Pin();
-		if (TreeItem.IsValid() && IsVisible() != bVisible)
+		TSharedPtr<ITreeItem> TreeItem = WeakTreeItem.Pin();
+		TSharedPtr<ISceneOutliner> Outliner = WeakOutliner.Pin();
+
+		if (TreeItem.IsValid() && Outliner.IsValid() && IsVisible() != bVisible)
 		{
 			FSetVisibilityVisitor Visitor(bVisible);
 			TreeItem->Visit(Visitor);
+			
+			Outliner->Refresh();
 
 			GEditor->RedrawAllViewports();
 		}
diff --git a/Engine/Source/Editor/SceneOutliner/Private/SceneOutlinerSettings.h b/Engine/Source/Editor/SceneOutliner/Private/SceneOutlinerSettings.h
--- a/Engine/Source/Editor/SceneOutliner/Private/SceneOutlinerSettings.h
+++ b/Engine/Source/Editor/SceneOutliner/Private/SceneOutlinerSettings.h
@@ -25,6 +25,10 @@ class USceneOutlinerSettings
 	UPROPERTY(config)
 	uint32 bShowOnlyActorsInCurrentLevel:1;
 
+	/** True when the Scene Outliner is showing only Folders containing  non-hidden Actors */
+	UPROPERTY(config)
+	uint32 bHideFoldersContainingHiddenActors : 1;
+
 	/** True when the Scene Outliner is only displaying selected Actors */
 	UPROPERTY(config)
 	uint32 bShowOnlySelectedActors:1;
diff --git a/Engine/Source/Editor/SceneOutliner/Public/SSceneOutliner.h b/Engine/Source/Editor/SceneOutliner/Public/SSceneOutliner.h
--- a/Engine/Source/Editor/SceneOutliner/Public/SSceneOutliner.h
+++ b/Engine/Source/Editor/SceneOutliner/Public/SSceneOutliner.h
@@ -543,6 +543,11 @@ namespace SceneOutliner
 		/** Enables/Disables whether the ShowOnlyActorsInCurrentLevelFilter is applied */
 		void ApplyShowOnlyCurrentLevelFilter(bool bShowOnlyActorsInCurrentLevel);
 
+		/** @return whether we are hiding Folders with hidden actors */
+		bool IsHidingFoldersContainingOnlyHiddenActors() const;
+		/** Toggles whether we are hiding Folders with hidden actors */
+		void ToggleHideFoldersContainingOnlyHiddenActors();
+
 		/** @return whether we are showing the components of the Actors */
 		bool IsShowingActorComponents();
 		/** Toggles whether we are showing the components of the Actors */
@@ -611,6 +616,10 @@ namespace SceneOutliner
 		/** Updates the expansion state of parent items after a repopulate, according to the previous state */
 		void SetParentsExpansionState(const FParentsExpansionState& ExpansionStateInfo) const;
 
+		/** Pair of functions to Hide Folders in Outliner when HideFoldersContainingHiddenActors filter is active*/
+		void HideFoldersContainingOnlyHiddenActors();
+		bool HideFoldersContainingOnlyHiddenActors(FTreeItemPtr Parent, bool bIsRoot = false);
+
 	private:
 
 		/** Number of actors that passed the search filter */
