PR #5961: Fix some instances of FName churn in the engine (Contributed by [REDACTED_PERSON])

diff --git a/Engine/Source/Editor/ContentBrowser/Private/SourcesViewWidgets.cpp b/Engine/Source/Editor/ContentBrowser/Private/SourcesViewWidgets.cpp
--- a/Engine/Source/Editor/ContentBrowser/Private/SourcesViewWidgets.cpp
+++ b/Engine/Source/Editor/ContentBrowser/Private/SourcesViewWidgets.cpp
@@ -297,7 +297,9 @@ FText SAssetTreeItem::GetToolTipText() const
 
 const FSlateBrush* SAssetTreeItem::GetBorderImage() const
 {
-	return bDraggedOver ? FEditorStyle::GetBrush("Menu.Background") : FEditorStyle::GetBrush("NoBorder");
+	static const FName NAME_DraggedBorderImage = TEXT("Menu.Background");
+	static const FName NAME_NoBorderImage = TEXT("NoBorder");
+	return bDraggedOver ? FEditorStyle::GetBrush(NAME_DraggedBorderImage) : FEditorStyle::GetBrush(NAME_NoBorderImage);
 }
 
 
diff --git a/Engine/Source/Editor/SceneOutliner/Private/SceneOutlinerGutter.cpp b/Engine/Source/Editor/SceneOutliner/Private/SceneOutlinerGutter.cpp
--- a/Engine/Source/Editor/SceneOutliner/Private/SceneOutlinerGutter.cpp
+++ b/Engine/Source/Editor/SceneOutliner/Private/SceneOutlinerGutter.cpp
@@ -283,13 +283,17 @@ private:
 	{
 		if (IsVisible())
 		{
-			return IsHovered() ? FEditorStyle::GetBrush( "Level.VisibleHighlightIcon16x" ) :
-				FEditorStyle::GetBrush( "Level.VisibleIcon16x" );
+			static const FName NAME_VisibleHoveredBrush = TEXT("Level.VisibleHighlightIcon16x");
+			static const FName NAME_VisibleNotHoveredBrush = TEXT("Level.VisibleIcon16x");
+			return IsHovered() ? FEditorStyle::GetBrush(NAME_VisibleHoveredBrush) :
+				FEditorStyle::GetBrush(NAME_VisibleNotHoveredBrush);
 		}
 		else
 		{
-			return IsHovered() ? FEditorStyle::GetBrush( "Level.NotVisibleHighlightIcon16x" ) :
-				FEditorStyle::GetBrush( "Level.NotVisibleIcon16x" );
+			static const FName NAME_NotVisibleHoveredBrush = TEXT("Level.NotVisibleHighlightIcon16x");
+			static const FName NAME_NotVisibleNotHoveredBrush = TEXT("Level.NotVisibleIcon16x");
+			return IsHovered() ? FEditorStyle::GetBrush(NAME_NotVisibleHoveredBrush) :
+				FEditorStyle::GetBrush(NAME_NotVisibleNotHoveredBrush);
 		}
 	}
 
diff --git a/Engine/Source/Runtime/Slate/Private/Widgets/Views/SExpanderArrow.cpp b/Engine/Source/Runtime/Slate/Private/Widgets/Views/SExpanderArrow.cpp
--- a/Engine/Source/Runtime/Slate/Private/Widgets/Views/SExpanderArrow.cpp
+++ b/Engine/Source/Runtime/Slate/Private/Widgets/Views/SExpanderArrow.cpp
@@ -56,8 +56,9 @@ int32 SExpanderArrow::OnPaint(const FPaintArgs& Args, const FGeometry& AllottedG
 	//   
 	//
 
+	static const FName NAME_VerticalBarBrush = TEXT("WhiteBrush");
 	const float Indent = IndentAmount.Get(10.f);
-	const FSlateBrush* VerticalBarBrush = (StyleSet == nullptr) ? nullptr : StyleSet->GetBrush("WhiteBrush");
+	const FSlateBrush* VerticalBarBrush = (StyleSet == nullptr) ? nullptr : StyleSet->GetBrush(NAME_VerticalBarBrush);
 
 	if (ShouldDrawWires.Get() == true && VerticalBarBrush != nullptr)
 	{
