Fix for world outliner's customizable column allowing multiselect, which could cause a crash when no options are selected

diff --git a/Engine/Source/Editor/SceneOutliner/Private/SceneOutlinerActorInfoColumn.cpp b/Engine/Source/Editor/SceneOutliner/Private/SceneOutlinerActorInfoColumn.cpp
--- a/Engine/Source/Editor/SceneOutliner/Private/SceneOutlinerActorInfoColumn.cpp
+++ b/Engine/Source/Editor/SceneOutliner/Private/SceneOutlinerActorInfoColumn.cpp
@@ -169,6 +169,7 @@ SHeaderRow::FColumn::FArguments FActorInfoColumn::ConstructHeaderRowColumn()
 			[
 				SNew(SListView<TSharedPtr<ECustomColumnMode::Type>>)
 				.ListItemsSource(&ModeOptions)
+				.SelectionMode(ESelectionMode::Single)
 				.OnGenerateRow( this, &FActorInfoColumn::MakeComboButtonItemWidget )
 				.OnSelectionChanged( this, &FActorInfoColumn::OnModeChanged )
 			]
