PR #6097: Fix typo in Scene Outliner view options menu (Contributed by [REDACTED_PERSON])

diff --git a/Engine/Source/Editor/SceneOutliner/Private/SSceneOutliner.cpp b/Engine/Source/Editor/SceneOutliner/Private/SSceneOutliner.cpp
--- a/Engine/Source/Editor/SceneOutliner/Private/SSceneOutliner.cpp
+++ b/Engine/Source/Editor/SceneOutliner/Private/SSceneOutliner.cpp
@@ -779,7 +779,7 @@ namespace SceneOutliner
 
 				MenuBuilder.AddMenuEntry(
 					LOCTEXT("ToggleShowActorComponents", "Show Actor Components"),
-					LOCTEXT("ToggleShowActorComponentsToolTip", "When enabled, shows components beloging to actors."),
+					LOCTEXT("ToggleShowActorComponentsToolTip", "When enabled, shows components belonging to actors."),
 					FSlateIcon(),
 					FUIAction(
 						FExecuteAction::CreateSP(this, &SSceneOutliner::ToggleShowActorComponents),
diff --git a/Engine/Source/Editor/SceneOutliner/Public/SceneOutlinerPublicTypes.h b/Engine/Source/Editor/SceneOutliner/Public/SceneOutlinerPublicTypes.h
--- a/Engine/Source/Editor/SceneOutliner/Public/SceneOutlinerPublicTypes.h
+++ b/Engine/Source/Editor/SceneOutliner/Public/SceneOutlinerPublicTypes.h
@@ -112,7 +112,7 @@ namespace SceneOutliner
 		/** Map of column types available to the scene outliner, along with default ordering */
 		TMap<FName, FColumnInfo> ColumnMap;
 		
-		/** Whether the Scene Outliner should display parenta actors in a Tree */
+		/** Whether the Scene Outliner should display parent actors in a Tree */
 		bool bShowParentTree : 1;
 
 		/** True to only show folders in this outliner */
