Added enhancement JIRA UE-76300: The World Outliner scrolls to the last selected item after writing/deleting/clearing its search box.

diff --git a/Engine/Source/Editor/SceneOutliner/Private/SSceneOutliner.cpp b/Engine/Source/Editor/SceneOutliner/Private/SSceneOutliner.cpp
--- a/Engine/Source/Editor/SceneOutliner/Private/SSceneOutliner.cpp
+++ b/Engine/Source/Editor/SceneOutliner/Private/SSceneOutliner.cpp
@@ -1855,7 +1855,7 @@ namespace SceneOutliner
 					const FName& Folder = Actor->GetFolderPath();
 					if (!Folder.IsNone() && !ExcludedParents.Contains(Folder))
 					{
-						auto FolderItem = TreeItemMap.FindRef(Folder);
+						FTreeItemPtr FolderItem = TreeItemMap.FindRef(Folder);
 						if (FolderItem.IsValid() && !FolderItem->GetChildren().ContainsByPredicate(&ItemHasSubFolders))
 						{
 							ExcludedParents.Add(Folder);
@@ -2218,7 +2218,7 @@ namespace SceneOutliner
 			return;
 		}
 
-		auto Item = TreeItemMap.FindRef(OldPath);
+		FTreeItemPtr Item = TreeItemMap.FindRef(OldPath);
 		if (Item.IsValid())
 		{
 			// Remove it from the map under the old ID (which is derived from the folder path)
@@ -2867,7 +2867,7 @@ namespace SceneOutliner
 				// Scroll last item into view - this means if we are multi-selecting, we show newest selection. @TODO Not perfect though
 				if (AActor* LastSelectedActor = GEditor->GetSelectedActors()->GetBottom<AActor>())
 				{
-					auto TreeItem = TreeItemMap.FindRef(LastSelectedActor);
+					FTreeItemPtr TreeItem = TreeItemMap.FindRef(LastSelectedActor);
 					if (TreeItem.IsValid())
 					{
 						if (!OutlinerTreeView->IsItemVisible(TreeItem))
@@ -3140,7 +3140,7 @@ namespace SceneOutliner
 			return;
 		}
 		
-		auto TreeItem = TreeItemMap.FindRef(ChangedActor);
+		FTreeItemPtr TreeItem = TreeItemMap.FindRef(ChangedActor);
 		if (TreeItem.IsValid())
 		{
 			if (SearchBoxFilter->PassesFilter(*TreeItem))
@@ -3175,6 +3175,31 @@ namespace SceneOutliner
 	{
 		SearchBoxFilter->SetRawFilterText( InFilterText );
 		FilterTextBoxWidget->SetError( SearchBoxFilter->GetFilterErrorText() );
+
+		// Scroll last item (if it passes the filter) into view - this means if we are multi-selecting, we show newest selection that passes the filter
+		if (AActor* LastSelectedActor = GEditor->GetSelectedActors()->GetBottom<AActor>())
+		{
+			// This part is different than that of OnLevelSelectionChanged(nullptr) because IsItemVisible(TreeItem) & ScrollItemIntoView(TreeItem) are applied to
+			// the current visual state, not to the one after applying the filter. Thus, the scroll would go to the place where the object was located
+			// before applying the FilterText
+
+			// If the object is already in the list, but it does not passes the filter, then we do not want to re-add it, because it will be removed by the filter
+			const FTreeItemPtr TreeItem = TreeItemMap.FindRef(LastSelectedActor);
+			if (TreeItem.IsValid() && !SearchBoxFilter->PassesFilter(*TreeItem))
+			{
+				return;
+			}
+
+			// If the object is not in the list, and it does not passes the filter, then we should not re-add it, because it would be removed by the filter again. Unfortunately,
+			// there is no code to check if a future element (i.e., one that is currently not in the TreeItemMap list) will pass the filter. Therefore, we kind of overkill it
+			// by re-adding that element (even though it will be removed). However, AddItemToTree(FTreeItemRef Item) and similar functions already check the element before
+			// adding it. So this solution is fine.
+			// This solution might affect the performance of the World Outliner when a key is pressed, but it will still work properly when the remove/del keys are pressed. Not
+			// updating the filter when !TreeItem.IsValid() would result in the focus not being updated when the remove/del keys are pressed.
+
+			// In any other case (i.e., if the object passes the current filter), re-add it
+			OnItemAdded(LastSelectedActor, ENewItemAction::ScrollIntoView);
+		}
 	}
 
 	void SSceneOutliner::OnFilterTextCommitted( const FText& InFilterText, ETextCommit::Type CommitInfo )
