Scene Outliner Icons and Visibility States

diff --git a/Engine/Source/Editor/SceneOutliner/Private/SceneOutlinerGutter.cpp b/Engine/Source/Editor/SceneOutliner/Private/SceneOutlinerGutter.cpp
--- a/Engine/Source/Editor/SceneOutliner/Private/SceneOutlinerGutter.cpp
+++ b/Engine/Source/Editor/SceneOutliner/Private/SceneOutlinerGutter.cpp
@@ -166,15 +166,17 @@ public:
 	SLATE_END_ARGS()
 
 	/** Construct this widget */
-	void Construct(const FArguments& InArgs, TWeakPtr<FSceneOutlinerGutter> InWeakColumn, TWeakPtr<ISceneOutliner> InWeakOutliner, TWeakPtr<ITreeItem> InWeakTreeItem)
+	void Construct(const FArguments& InArgs, TWeakPtr<FSceneOutlinerGutter> InWeakColumn, TWeakPtr<ISceneOutliner> InWeakOutliner, TWeakPtr<ITreeItem> InWeakTreeItem, const STableRow<FTreeItemPtr>* InRow)
 	{
 		WeakTreeItem = InWeakTreeItem;
 		WeakOutliner = InWeakOutliner;
 		WeakColumn = InWeakColumn;
 
+		Row = InRow;
+
 		SImage::Construct(
 			SImage::FArguments()
-			.ColorAndOpacity(FSlateColor::UseForeground())
+			.ColorAndOpacity(this, &SVisibilityWidget::GetForegroundColor)
 			.Image(this, &SVisibilityWidget::GetBrush)
 		);
 	}
@@ -284,20 +286,42 @@ private:
 	{
 		if (IsVisible())
 		{
+
 			static const FName NAME_VisibleHoveredBrush = TEXT("Level.VisibleHighlightIcon16x");
 			static const FName NAME_VisibleNotHoveredBrush = TEXT("Level.VisibleIcon16x");
-			return IsHovered() ? FEditorStyle::GetBrush(NAME_VisibleHoveredBrush) :
-				FEditorStyle::GetBrush(NAME_VisibleNotHoveredBrush);
+			return IsHovered() ? FAppStyle::Get().GetBrush(NAME_VisibleHoveredBrush) :
+				FAppStyle::Get().GetBrush(NAME_VisibleNotHoveredBrush);
 		}
 		else
 		{
 			static const FName NAME_NotVisibleHoveredBrush = TEXT("Level.NotVisibleHighlightIcon16x");
 			static const FName NAME_NotVisibleNotHoveredBrush = TEXT("Level.NotVisibleIcon16x");
-			return IsHovered() ? FEditorStyle::GetBrush(NAME_NotVisibleHoveredBrush) :
-				FEditorStyle::GetBrush(NAME_NotVisibleNotHoveredBrush);
+			return IsHovered() ? FAppStyle::Get().GetBrush(NAME_NotVisibleHoveredBrush) :
+				FAppStyle::Get().GetBrush(NAME_NotVisibleNotHoveredBrush);
 		}
 	}
 
+	virtual FSlateColor GetForegroundColor() const
+	{
+
+		auto Outliner = WeakOutliner.Pin();
+		auto TreeItem = WeakTreeItem.Pin();
+
+		const bool bIsSelected = Outliner->GetTree().IsItemSelected(TreeItem.ToSharedRef());
+
+		// make the foreground brush transparent if it is not selected and it is visible
+		if (IsVisible() && !Row->IsHovered() && !bIsSelected)
+		{
+			return FLinearColor::Transparent;
+		}
+		else if (IsHovered() && !bIsSelected)
+		{
+			return FAppStyle::Get().GetColor("Colors.ForegroundHover");
+		}
+
+		return FSlateColor::UseForeground();
+	}
+
 	/** Check if the specified item is visible */
 	static bool IsVisible(const FTreeItemPtr& Item, const TSharedPtr<FSceneOutlinerGutter>& Column)
 	{
@@ -336,6 +360,9 @@ private:
 	/** Weak pointer back to the column */
 	TWeakPtr<FSceneOutlinerGutter> WeakColumn;
 
+	/** Weak pointer back to the row */
+	const STableRow<FTreeItemPtr>* Row;
+
 	/** Scoped undo transaction */
 	TUniquePtr<FScopedTransaction> UndoTransaction;
 };
@@ -377,7 +404,7 @@ const TSharedRef<SWidget> FSceneOutlinerGutter::ConstructRowWidget(FTreeItemRef
 		.AutoWidth()
 		.VAlign(VAlign_Center)
 		[
-			SNew(SVisibilityWidget, SharedThis(this), WeakOutliner, TreeItem)
+			SNew(SVisibilityWidget, SharedThis(this), WeakOutliner, TreeItem, &Row)
 		];
 }
 
