Fix a few places that were accessing slate style every frame instead of just caching the brushes

diff --git a/Engine/Source/Editor/PlacementMode/Private/SPlacementModeTools.cpp b/Engine/Source/Editor/PlacementMode/Private/SPlacementModeTools.cpp
--- a/Engine/Source/Editor/PlacementMode/Private/SPlacementModeTools.cpp
+++ b/Engine/Source/Editor/PlacementMode/Private/SPlacementModeTools.cpp
@@ -297,6 +297,11 @@ void SPlacementModeTools::Construct( const FArguments& InArgs )
 	bRecentlyPlacedRefreshRequested = false;
 	bNeedsUpdate = true;
 
+
+	static const FName PlacementBrowserActiveTabBarBrushName("PlacementBrowser.ActiveTabBar");
+	PlacementBrowserActiveTabBarBrush = FEditorStyle::GetBrush(PlacementBrowserActiveTabBarBrushName);
+
+
 	FPlacementMode* PlacementEditMode = (FPlacementMode*)GLevelEditorModeTools().GetActiveMode( FBuiltinEditorModes::EM_Placement );
 	if (PlacementEditMode)
 	{
@@ -548,8 +553,7 @@ const FSlateBrush* SPlacementModeTools::PlacementGroupBorderImage( FName Categor
 {
 	if ( ActiveTabName == CategoryName )
 	{
-		static FName PlacementBrowserActiveTabBarBrush( "PlacementBrowser.ActiveTabBar" );
-		return FEditorStyle::GetBrush( PlacementBrowserActiveTabBarBrush );
+		return PlacementBrowserActiveTabBarBrush;
 	}
 	else
 	{
diff --git a/Engine/Source/Editor/PlacementMode/Private/SPlacementModeTools.h b/Engine/Source/Editor/PlacementMode/Private/SPlacementModeTools.h
--- a/Engine/Source/Editor/PlacementMode/Private/SPlacementModeTools.h
+++ b/Engine/Source/Editor/PlacementMode/Private/SPlacementModeTools.h
@@ -150,4 +150,6 @@ private:
 
 	/** List view that shows placeable items */
 	TSharedPtr<SListView<TSharedPtr<FPlaceableItem>>> ListView;
+
+	const FSlateBrush* PlacementBrowserActiveTabBarBrush;
 };
diff --git a/Engine/Source/Editor/SceneOutliner/Private/SSceneOutliner.cpp b/Engine/Source/Editor/SceneOutliner/Private/SSceneOutliner.cpp
--- a/Engine/Source/Editor/SceneOutliner/Private/SSceneOutliner.cpp
+++ b/Engine/Source/Editor/SceneOutliner/Private/SSceneOutliner.cpp
@@ -3795,7 +3795,7 @@ namespace SceneOutliner
 	{
 		if ( !IsFilterActive() )
 		{
-			 return FAppStyle::Get().GetSlateColor("Colors.Foreground");
+			return FSlateColor::UseForeground();
 		}
 		else if( FilteredActorCount == 0 )
 		{
diff --git a/Engine/Source/Editor/SceneOutliner/Private/SceneOutlinerGutter.cpp b/Engine/Source/Editor/SceneOutliner/Private/SceneOutlinerGutter.cpp
--- a/Engine/Source/Editor/SceneOutliner/Private/SceneOutlinerGutter.cpp
+++ b/Engine/Source/Editor/SceneOutliner/Private/SceneOutlinerGutter.cpp
@@ -179,6 +179,19 @@ public:
 			.ColorAndOpacity(this, &SVisibilityWidget::GetForegroundColor)
 			.Image(this, &SVisibilityWidget::GetBrush)
 		);
+
+
+		static const FName NAME_VisibleHoveredBrush = TEXT("Level.VisibleHighlightIcon16x");
+		static const FName NAME_VisibleNotHoveredBrush = TEXT("Level.VisibleIcon16x");
+		static const FName NAME_NotVisibleHoveredBrush = TEXT("Level.NotVisibleHighlightIcon16x");
+		static const FName NAME_NotVisibleNotHoveredBrush = TEXT("Level.NotVisibleIcon16x");
+
+		VisibleHoveredBrush = FAppStyle::Get().GetBrush(NAME_VisibleHoveredBrush);
+		VisibleNotHoveredBrush = FAppStyle::Get().GetBrush(NAME_VisibleNotHoveredBrush);
+
+		NotVisibleHoveredBrush = FAppStyle::Get().GetBrush(NAME_NotVisibleHoveredBrush);
+		NotVisibleNotHoveredBrush = FAppStyle::Get().GetBrush(NAME_NotVisibleNotHoveredBrush);
+
 	}
 
 private:
@@ -286,18 +299,11 @@ private:
 	{
 		if (IsVisible())
 		{
-
-			static const FName NAME_VisibleHoveredBrush = TEXT("Level.VisibleHighlightIcon16x");
-			static const FName NAME_VisibleNotHoveredBrush = TEXT("Level.VisibleIcon16x");
-			return IsHovered() ? FAppStyle::Get().GetBrush(NAME_VisibleHoveredBrush) :
-				FAppStyle::Get().GetBrush(NAME_VisibleNotHoveredBrush);
+			return IsHovered() ? VisibleHoveredBrush : VisibleNotHoveredBrush;
 		}
 		else
 		{
-			static const FName NAME_NotVisibleHoveredBrush = TEXT("Level.NotVisibleHighlightIcon16x");
-			static const FName NAME_NotVisibleNotHoveredBrush = TEXT("Level.NotVisibleIcon16x");
-			return IsHovered() ? FAppStyle::Get().GetBrush(NAME_NotVisibleHoveredBrush) :
-				FAppStyle::Get().GetBrush(NAME_NotVisibleNotHoveredBrush);
+			return IsHovered() ? NotVisibleHoveredBrush : NotVisibleNotHoveredBrush;
 		}
 	}
 
@@ -365,6 +371,12 @@ private:
 
 	/** Scoped undo transaction */
 	TUniquePtr<FScopedTransaction> UndoTransaction;
+
+	/** Visbility brushes for the various states */
+	const FSlateBrush* VisibleHoveredBrush;
+	const FSlateBrush* VisibleNotHoveredBrush;
+	const FSlateBrush* NotVisibleHoveredBrush;
+	const FSlateBrush* NotVisibleNotHoveredBrush;
 };
 
 FSceneOutlinerGutter::FSceneOutlinerGutter(ISceneOutliner& Outliner)
