Scene Outliner Mobility - moved Mobility to an optional info column, removed icon badges

diff --git a/Engine/Source/Editor/SceneOutliner/Private/SceneOutlinerActorInfoColumn.cpp b/Engine/Source/Editor/SceneOutliner/Private/SceneOutlinerActorInfoColumn.cpp
--- a/Engine/Source/Editor/SceneOutliner/Private/SceneOutlinerActorInfoColumn.cpp
+++ b/Engine/Source/Editor/SceneOutliner/Private/SceneOutlinerActorInfoColumn.cpp
@@ -64,6 +64,23 @@ struct FGetInfo : TTreeItemGetter<FString>
 				}
 				return Result;
 			}
+		case ECustomColumnMode::Mobility:
+			{
+				FString Result;
+				USceneComponent* RootComponent = Actor->GetRootComponent();
+				if (RootComponent)
+				{
+					if (RootComponent->Mobility == EComponentMobility::Stationary)
+					{
+						Result = FString(TEXT("Stationary"));
+					}
+					else if (RootComponent->Mobility == EComponentMobility::Movable)
+					{
+						Result = FString(TEXT("Movable"));
+					}
+				}
+				return Result;
+			}
 
 		default:
 			return FString();
@@ -380,6 +397,10 @@ FText FActorInfoColumn::MakeComboText( const ECustomColumnMode::Type& Mode ) con
 		ModeName = LOCTEXT("CustomColumnMode_UncachedLights", "# Uncached Lights");
 		break;
 
+	case ECustomColumnMode::Mobility:
+		ModeName = LOCTEXT("CustomColumnMode_Mobility", "Mobility");
+		break;
+
 	default:
 		ensure(0);
 		break;
@@ -423,6 +444,11 @@ FText FActorInfoColumn::MakeComboToolTipText( const ECustomColumnMode::Type& Mod
 		ToolTipText = LOCTEXT("CustomColumnModeToolTip_UncachedLights", "Shows the number of uncached static lights (missing in lightmap)");
 		break;
 
+	case ECustomColumnMode::Mobility:
+		ToolTipText = LOCTEXT("CustomColumnModeToolTip_Mobility", "Shows the mobility of each actor");
+		break;
+
+
 	default:
 		ensure(0);
 		break;
diff --git a/Engine/Source/Editor/SceneOutliner/Private/SceneOutlinerActorInfoColumn.h b/Engine/Source/Editor/SceneOutliner/Private/SceneOutlinerActorInfoColumn.h
--- a/Engine/Source/Editor/SceneOutliner/Private/SceneOutlinerActorInfoColumn.h
+++ b/Engine/Source/Editor/SceneOutliner/Private/SceneOutlinerActorInfoColumn.h
@@ -28,6 +28,9 @@ namespace ECustomColumnMode
 		/** Class */
 		Class,
 
+		/** Mobility */
+		Mobility,
+
 		/** Level */
 		Level,
 
diff --git a/Engine/Source/Editor/SceneOutliner/Private/SceneOutlinerItemLabelColumn.cpp b/Engine/Source/Editor/SceneOutliner/Private/SceneOutlinerItemLabelColumn.cpp
--- a/Engine/Source/Editor/SceneOutliner/Private/SceneOutlinerItemLabelColumn.cpp
+++ b/Engine/Source/Editor/SceneOutliner/Private/SceneOutlinerItemLabelColumn.cpp
@@ -89,10 +89,6 @@ struct SActorTreeLabel : FCommonLabelData, public SCompoundWidget
 
 		HighlightText = SceneOutliner.GetFilterHighlightText();
 
-		MobilityStaticBrush = FEditorStyle::GetBrush( "ClassIcon.ComponentMobilityStaticPip" );
-		MobilityStationaryBrush = FEditorStyle::GetBrush( "ClassIcon.ComponentMobilityStationaryPip" );
-		MobilityMovableBrush = FEditorStyle::GetBrush( "ClassIcon.ComponentMobilityMovablePip" );
-
 		TSharedPtr<SInlineEditableTextBlock> InlineTextBlock;
 
 		auto MainContent = SNew(SHorizontalBox)
@@ -126,35 +122,8 @@ struct SActorTreeLabel : FCommonLabelData, public SCompoundWidget
 				.HighlightText(HighlightText)
 			];
 
-		TSharedRef<SOverlay> IconContent = SNew(SOverlay)
-			+ SOverlay::Slot()
-			.HAlign(HAlign_Right)
-			.VAlign(VAlign_Center)
-			[
-				SNew(SImage)
-				.Image(this, &SActorTreeLabel::GetIcon)
-				.ToolTipText(this, &SActorTreeLabel::GetIconTooltip)
-			]
-
-			+ SOverlay::Slot()
-			.HAlign(HAlign_Right)
-			.VAlign(VAlign_Center)
-			[
-				SNew(SImage)
-				.Image(this, &SActorTreeLabel::GetIconOverlay)
-			];
-
-
 		if (ActorItem.GetSharedData().Mode == ESceneOutlinerMode::ActorBrowsing)
 		{
-			// Add the component mobility icon
-			IconContent->AddSlot()
-			.HAlign(HAlign_Left)
-			[
-				SNew(SImage)
-				.Image(this, &SActorTreeLabel::GetBrushForComponentMobilityIcon)
-			];
-
 			ActorItem.RenameRequestEvent.BindSP( InlineTextBlock.Get(), &SInlineEditableTextBlock::EnterEditingMode );
 		}
 		
@@ -171,7 +140,9 @@ struct SActorTreeLabel : FCommonLabelData, public SCompoundWidget
 				.WidthOverride(FDefaultTreeItemMetrics::IconSize())
 				.HeightOverride(FDefaultTreeItemMetrics::IconSize())
 				[
-					IconContent
+					SNew(SImage)
+					.Image(this, &SActorTreeLabel::GetIcon)
+					.ToolTipText(this, &SActorTreeLabel::GetIconTooltip)
 				]
 			]
 
@@ -190,11 +161,6 @@ private:
 	TWeakObjectPtr<AActor> ActorPtr;
 	TAttribute<FText> HighlightText;
 
-	/** The component mobility brushes */
-	const FSlateBrush* MobilityStaticBrush;
-	const FSlateBrush* MobilityStationaryBrush;
-	const FSlateBrush* MobilityMovableBrush;
-
 	FText GetDisplayText() const
 	{
 		const AActor* Actor = ActorPtr.Get();
@@ -370,29 +336,6 @@ private:
 			}
 		}
 	}
-
-	const FSlateBrush* GetBrushForComponentMobilityIcon() const
-	{
-		const FSlateBrush* IconBrush = MobilityStaticBrush;
-
-		if (auto* Actor = ActorPtr.Get())
-		{
-			USceneComponent* RootComponent = ActorPtr->GetRootComponent();
-			if (RootComponent)
-			{
-				if (RootComponent->Mobility == EComponentMobility::Stationary)
-				{
-					IconBrush = MobilityStationaryBrush;
-				}
-				else if (RootComponent->Mobility == EComponentMobility::Movable)
-				{
-					IconBrush = MobilityMovableBrush;
-				}
-			}
-		}
-
-		return IconBrush;
-	}
 };
 
 struct SWorldTreeLabel : FCommonLabelData, public SCompoundWidget
@@ -662,8 +605,6 @@ struct SComponentTreeLabel : FCommonLabelData, public SCompoundWidget
 		TreeItemPtr = StaticCastSharedRef<FComponentTreeItem>(ComponentItem.AsShared());
 		WeakSceneOutliner = StaticCastSharedRef<ISceneOutliner>(SceneOutliner.AsShared());
 
-		MobilityStaticBrush = FEditorStyle::GetBrush("ClassIcon.ComponentMobilityStaticPip");
-
 		ComponentPtr = ComponentItem.Component;
 		
 		HighlightText = SceneOutliner.GetFilterHighlightText();
@@ -689,35 +630,8 @@ struct SComponentTreeLabel : FCommonLabelData, public SCompoundWidget
 				.IsSelected(FIsSelected::CreateSP(&InRow, &STableRow<FTreeItemPtr>::IsSelectedExclusively))
 			];
 
-		TSharedRef<SOverlay> IconContent = SNew(SOverlay)
-			+ SOverlay::Slot()
-			.HAlign(HAlign_Right)
-			.VAlign(VAlign_Center)
-			[
-				SNew(SImage)
-				.Image(this, &SComponentTreeLabel::GetIcon)
-				.ToolTipText(this, &SComponentTreeLabel::GetIconTooltip)
-			]
-
-		+ SOverlay::Slot()
-			.HAlign(HAlign_Right)
-			.VAlign(VAlign_Center)
-			[
-				SNew(SImage)
-				.Image(this, &SComponentTreeLabel::GetIconOverlay)
-			];
-
-
 		if (ComponentItem.GetSharedData().Mode == ESceneOutlinerMode::ActorBrowsing)
 		{
-			// Add the component mobility icon
-			IconContent->AddSlot()
-				.HAlign(HAlign_Left)
-				[
-					SNew(SImage)
-					.Image(this, &SComponentTreeLabel::GetBrushForComponentMobilityIcon)
-				];
-
 			ComponentItem.RenameRequestEvent.BindSP(InlineTextBlock.Get(), &SInlineEditableTextBlock::EnterEditingMode);
 		}
 
@@ -726,25 +640,27 @@ struct SComponentTreeLabel : FCommonLabelData, public SCompoundWidget
 				SNew(SHorizontalBox)
 
 				+ SHorizontalBox::Slot()
-			.AutoWidth()
-			.VAlign(VAlign_Center)
-			.Padding(FDefaultTreeItemMetrics::IconPadding())
-			[
-				SNew(SBox)
-				.WidthOverride(FDefaultTreeItemMetrics::IconSize())
-			.HeightOverride(FDefaultTreeItemMetrics::IconSize())
-			[
-				IconContent
-			]
-			]
+				.AutoWidth()
+				.VAlign(VAlign_Center)
+				.Padding(FDefaultTreeItemMetrics::IconPadding())
+				[
+					SNew(SBox)
+					.WidthOverride(FDefaultTreeItemMetrics::IconSize())
+					.HeightOverride(FDefaultTreeItemMetrics::IconSize())
+					[
+						SNew(SImage)
+						.Image(this, &SComponentTreeLabel::GetIcon)
+						.ToolTipText(this, &SComponentTreeLabel::GetIconTooltip)
+					]
+				]
 
-		+ SHorizontalBox::Slot()
-			.FillWidth(1.0f)
-			.VAlign(VAlign_Center)
-			.Padding(0.0f, 2.0f)
-			[
-				MainContent
-			]
+				+ SHorizontalBox::Slot()
+				.FillWidth(1.0f)
+				.VAlign(VAlign_Center)
+				.Padding(0.0f, 2.0f)
+				[
+					MainContent
+				]
 			];
 
 	}
@@ -753,8 +669,6 @@ private:
 	TWeakPtr<FComponentTreeItem> TreeItemPtr;
 	TWeakObjectPtr<UActorComponent> ComponentPtr;
 	TAttribute<FText> HighlightText;
-	const FSlateBrush* MobilityStaticBrush;
-
 
 	FText GetDisplayText() const
 	{
@@ -819,24 +733,12 @@ private:
 		}
 	}
 
-	const FSlateBrush* GetIconOverlay() const
-	{
-		const FSlateBrush* IconBrush = MobilityStaticBrush;
-		return IconBrush;
-	}
-
 	FText GetIconTooltip() const
 	{
 		FText ToolTipText;
 		return ToolTipText;
 	}
 
-	const FSlateBrush* GetBrushForComponentMobilityIcon() const
-	{
-		const FSlateBrush* IconBrush = MobilityStaticBrush;
-		return IconBrush;
-	}
-
 	FSlateColor GetForegroundColor() const
 	{
 		if (auto BaseColor = FCommonLabelData::GetForegroundColor(TreeItemPtr.Pin()))
@@ -858,7 +760,6 @@ struct SSubComponentTreeLabel : FCommonLabelData, public SCompoundWidget
 		TreeItemPtr = StaticCastSharedRef<FSubComponentTreeItem>(SubComponentItem.AsShared());
 		WeakSceneOutliner = StaticCastSharedRef<ISceneOutliner>(SceneOutliner.AsShared());
 
-		MobilityStaticBrush = FEditorStyle::GetBrush("ClassIcon.ComponentMobilityStaticPip");
 		ComponentPtr = SubComponentItem.ParentComponent;
 
 		TSharedPtr<SInlineEditableTextBlock> InlineTextBlock;
@@ -882,34 +783,8 @@ struct SSubComponentTreeLabel : FCommonLabelData, public SCompoundWidget
 				})
 			];
 
-		TSharedRef<SOverlay> IconContent = SNew(SOverlay)
-			+ SOverlay::Slot()
-			.HAlign(HAlign_Right)
-			.VAlign(VAlign_Center)
-			[
-				SNew(SImage)
-				.Image(this, &SSubComponentTreeLabel::GetIcon)
-				.ToolTipText(this, &SSubComponentTreeLabel::GetIconTooltip)
-			]
-
-			+ SOverlay::Slot()
-			.HAlign(HAlign_Right)
-			.VAlign(VAlign_Center)
-			[
-				SNew(SImage)
-				.Image(this, &SSubComponentTreeLabel::GetIconOverlay)
-			];
-
 		if (SubComponentItem.GetSharedData().Mode == ESceneOutlinerMode::ActorBrowsing)
 		{
-			// Add the component mobility icon
-			IconContent->AddSlot()
-				.HAlign(HAlign_Left)
-				[
-					SNew(SImage)
-					.Image(this, &SSubComponentTreeLabel::GetBrushForComponentMobilityIcon)
-				];
-
 			SubComponentItem.RenameRequestEvent.BindSP(InlineTextBlock.Get(), &SInlineEditableTextBlock::EnterEditingMode);
 		}
 
@@ -926,7 +801,9 @@ struct SSubComponentTreeLabel : FCommonLabelData, public SCompoundWidget
 				.WidthOverride(FDefaultTreeItemMetrics::IconSize())
 				.HeightOverride(FDefaultTreeItemMetrics::IconSize())
 				[
-					IconContent
+					SNew(SImage)
+					.Image(this, &SSubComponentTreeLabel::GetIcon)
+					.ToolTipText(this, &SSubComponentTreeLabel::GetIconTooltip)
 				]
 			]
 
@@ -943,7 +820,6 @@ struct SSubComponentTreeLabel : FCommonLabelData, public SCompoundWidget
 private:
 	TWeakPtr<FSubComponentTreeItem> TreeItemPtr;
 	TWeakObjectPtr<UActorComponent> ComponentPtr;
-	const FSlateBrush* MobilityStaticBrush;
 
 	FText GetDisplayText() const
 	{
@@ -1003,24 +879,12 @@ private:
 		}
 	}
 
-	const FSlateBrush* GetIconOverlay() const
-	{
-		const FSlateBrush* IconBrush = MobilityStaticBrush;
-		return IconBrush;
-	}
-
 	FText GetIconTooltip() const
 	{
 		FText ToolTipText;
 		return ToolTipText;
 	}
 
-	const FSlateBrush* GetBrushForComponentMobilityIcon() const
-	{
-		const FSlateBrush* IconBrush = MobilityStaticBrush;
-		return IconBrush;
-	}
-
 	virtual bool OnVerifyItemLabelChanged(const FText& InLabel, FText& OutErrorMessage)
 	{
 		bool IsValidName = false;
