Scene Outliner - Adjusts Temp Actor display

diff --git a/Engine/Source/Editor/SceneOutliner/Private/SceneOutlinerItemLabelColumn.cpp b/Engine/Source/Editor/SceneOutliner/Private/SceneOutlinerItemLabelColumn.cpp
--- a/Engine/Source/Editor/SceneOutliner/Private/SceneOutlinerItemLabelColumn.cpp
+++ b/Engine/Source/Editor/SceneOutliner/Private/SceneOutlinerItemLabelColumn.cpp
@@ -302,7 +302,8 @@ private:
 		if (TreeItem->GetSharedData().bRepresentingPlayWorld && !TreeItem->bExistsInCurrentWorldAndPIE)
 		{
 			// Highlight actors that are exclusive to PlayWorld
-			return FLinearColor(0.9f, 0.8f, 0.4f);
+			return FLinearColor::FromSRGBColor(FColor::FromHex("FFE65A"));
+
 		}
 
 		// also darken items that are non selectable in the active mode(s)
