Scene Outliner - Adjusts padding on Info Column and adds Static Label for Mobility Info Column

diff --git a/Engine/Source/Editor/SceneOutliner/Private/SceneOutlinerActorInfoColumn.cpp b/Engine/Source/Editor/SceneOutliner/Private/SceneOutlinerActorInfoColumn.cpp
--- a/Engine/Source/Editor/SceneOutliner/Private/SceneOutlinerActorInfoColumn.cpp
+++ b/Engine/Source/Editor/SceneOutliner/Private/SceneOutlinerActorInfoColumn.cpp
@@ -70,6 +70,10 @@ struct FGetInfo : TTreeItemGetter<FString>
 				USceneComponent* RootComponent = Actor->GetRootComponent();
 				if (RootComponent)
 				{
+					if( RootComponent->Mobility == EComponentMobility::Static )
+					{
+						Result = FString(TEXT("Static"));
+					}
 					if (RootComponent->Mobility == EComponentMobility::Stationary)
 					{
 						Result = FString(TEXT("Stationary"));
@@ -221,6 +225,7 @@ const TSharedRef< SWidget > FActorInfoColumn::ConstructRowWidget( FTreeItemRef T
 	HorizontalBox->AddSlot()
 	.AutoWidth()
 	.VAlign(VAlign_Center)
+	.Padding(8, 0, 0, 0)
 	[
 		MainText
 	];
@@ -235,6 +240,7 @@ const TSharedRef< SWidget > FActorInfoColumn::ConstructRowWidget( FTreeItemRef T
 		HorizontalBox->AddSlot()
 		.VAlign(VAlign_Center)
 		.AutoWidth()
+		.Padding(8, 0, 0, 0)
 		[
 			// Make sure that the hyperlink shows as black (by multiplying black * desired color) when selected so it is readable against the orange background even if blue/green/etc... normally
 			SNew(SBorder)
