SceneOutliner refactor

diff --git a/Engine/Source/Editor/SceneOutliner/Private/ActorBrowsingMode.cpp b/Engine/Source/Editor/SceneOutliner/Private/ActorBrowsingMode.cpp
new file mode 100644
--- /dev/null
+++ b/Engine/Source/Editor/SceneOutliner/Private/ActorBrowsingMode.cpp
@@ -0,0 +1,1230 @@
+// Copyright Epic Games, Inc. All Rights Reserved.
+
+#include "ActorBrowsingMode.h"
+#include "SceneOutlinerFilters.h"
+#include "SceneOutlinerModule.h"
+#include "SceneOutlinerMenuContext.h"
+#include "ActorHierarchy.h"
+#include "SceneOutlinerDelegates.h"
+#include "Editor.h"
+#include "Editor/GroupActor.h"
+#include "UnrealEdGlobals.h"
+#include "ToolMenus.h"
+#include "Engine/Selection.h"
+#include "Editor/UnrealEdEngine.h"
+#include "HAL/PlatformApplicationMisc.h"
+#include "SceneOutlinerDragDrop.h"
+#include "EditorActorFolders.h"
+#include "ActorEditorUtils.h"
+#include "DragAndDrop/ActorDragDropOp.h"
+#include "Logging/MessageLog.h"
+#include "SSocketChooser.h"
+#include "ActorFolderPickingMode.h"
+#include "ActorTreeItem.h"
+#include "ActorFolderTreeItem.h"
+#include "WorldTreeItem.h"
+#include "ComponentTreeItem.h"
+#include "ActorBrowsingModeSettings.h"
+#include "ScopedTransaction.h"
+
+DEFINE_LOG_CATEGORY_STATIC(LogSceneOutliner, Log, All);
+
+#define LOCTEXT_NAMESPACE "SceneOutliner_ActorBrowsingMode"
+
+UActorBrowsingModeSettings::UActorBrowsingModeSettings(const FObjectInitializer& ObjectInitializer)
+	: Super(ObjectInitializer)
+{ }
+
+namespace SceneOutliner
+{
+	using FActorFilter = TOutlinerPredicateFilter<FActorTreeItem>;
+
+	FActorBrowsingMode::FActorBrowsingMode(SSceneOutliner* InSceneOutliner, TWeakObjectPtr<UWorld> InSpecifiedWorldToDisplay)
+		: FActorModeInteractive(InSceneOutliner, true, InSpecifiedWorldToDisplay)
+		, FilteredActorCount(0)
+	{
+		// Capture selection changes of bones from mesh selection in fracture tools
+		FSceneOutlinerDelegates::Get().OnComponentsUpdated.AddRaw(this, &FActorBrowsingMode::OnComponentsUpdated);
+
+		GEngine->OnLevelActorDeleted().AddRaw(this, &FActorBrowsingMode::OnLevelActorDeleted);
+
+		FEditorDelegates::OnEditCutActorsBegin.AddRaw(this, &FActorBrowsingMode::OnEditCutActorsBegin);
+		FEditorDelegates::OnEditCutActorsEnd.AddRaw(this, &FActorBrowsingMode::OnEditCutActorsEnd);
+		FEditorDelegates::OnEditCopyActorsBegin.AddRaw(this, &FActorBrowsingMode::OnEditCopyActorsBegin);
+		FEditorDelegates::OnEditCopyActorsEnd.AddRaw(this, &FActorBrowsingMode::OnEditCopyActorsEnd);
+		FEditorDelegates::OnEditPasteActorsBegin.AddRaw(this, &FActorBrowsingMode::OnEditPasteActorsBegin);
+		FEditorDelegates::OnEditPasteActorsEnd.AddRaw(this, &FActorBrowsingMode::OnEditPasteActorsEnd);
+		FEditorDelegates::OnDuplicateActorsBegin.AddRaw(this, &FActorBrowsingMode::OnDuplicateActorsBegin);
+		FEditorDelegates::OnDuplicateActorsEnd.AddRaw(this, &FActorBrowsingMode::OnDuplicateActorsEnd);
+		FEditorDelegates::OnDeleteActorsBegin.AddRaw(this, &FActorBrowsingMode::OnDeleteActorsBegin);
+		FEditorDelegates::OnDeleteActorsEnd.AddRaw(this, &FActorBrowsingMode::OnDeleteActorsEnd);
+
+		UActorBrowsingModeSettings* SharedSettings = GetMutableDefault<UActorBrowsingModeSettings>();
+		// Get the OutlinerModule to register FilterInfos with the FilterInfoMap
+		FOutlinerFilterInfo ShowOnlySelectedActorsInfo(LOCTEXT("ToggleShowOnlySelected", "Only Selected"), LOCTEXT("ToggleShowOnlySelectedToolTip", "When enabled, only displays actors that are currently selected."), SharedSettings->bShowOnlySelectedActors, FCreateSceneOutlinerFilter::CreateStatic(&FActorBrowsingMode::CreateShowOnlySelectedActorsFilter));
+		ShowOnlySelectedActorsInfo.OnToggle().AddLambda([this](bool bIsActive)
+			{
+				UActorBrowsingModeSettings* Settings = GetMutableDefault<UActorBrowsingModeSettings>();
+				Settings->bShowOnlySelectedActors = bIsActive;
+				Settings->PostEditChange();
+			});
+		FilterInfoMap.Add(TEXT("ShowOnlySelectedActors"), ShowOnlySelectedActorsInfo);
+
+		FOutlinerFilterInfo HideTemporaryActorsInfo(LOCTEXT("ToggleHideTemporaryActors", "Hide Temporary Actors"), LOCTEXT("ToggleHideTemporaryActorsToolTip", "When enabled, hides temporary/run-time Actors."), SharedSettings->bHideTemporaryActors, FCreateSceneOutlinerFilter::CreateStatic(&FActorBrowsingMode::CreateHideTemporaryActorsFilter));
+		HideTemporaryActorsInfo.OnToggle().AddLambda([this](bool bIsActive)
+			{
+				UActorBrowsingModeSettings* Settings = GetMutableDefault<UActorBrowsingModeSettings>();
+				Settings->bHideTemporaryActors = bIsActive;
+				Settings->PostEditChange();
+			});
+		FilterInfoMap.Add(TEXT("HideTemporaryActors"), HideTemporaryActorsInfo);
+
+		FOutlinerFilterInfo OnlyCurrentLevelInfo(LOCTEXT("ToggleShowOnlyCurrentLevel", "Only in Current Level"), LOCTEXT("ToggleShowOnlyCurrentLevelToolTip", "When enabled, only shows Actors that are in the Current Level."), SharedSettings->bShowOnlyActorsInCurrentLevel, FCreateSceneOutlinerFilter::CreateStatic(&FActorBrowsingMode::CreateIsInCurrentLevelFilter));
+		OnlyCurrentLevelInfo.OnToggle().AddLambda([this](bool bIsActive)
+			{
+				UActorBrowsingModeSettings* Settings = GetMutableDefault<UActorBrowsingModeSettings>();
+				Settings->bShowOnlyActorsInCurrentLevel = bIsActive;
+				Settings->PostEditChange();
+			});
+		FilterInfoMap.Add(TEXT("ShowOnlyCurrentLevel"), OnlyCurrentLevelInfo);
+
+		bHideComponents = SharedSettings->bHideActorComponents;
+		FOutlinerFilterInfo HideComponentsInfo(LOCTEXT("ToggleHideActorComponents", "Hide Actor Components"), LOCTEXT("ToggleHideActorComponentsToolTip", "When Enabled, hides components belonging to actors."), SharedSettings->bHideActorComponents, FCreateSceneOutlinerFilter::CreateStatic(&FActorBrowsingMode::CreateHideComponentsFilter));
+		HideComponentsInfo.OnToggle().AddLambda([this](bool bIsActive)
+			{
+				UActorBrowsingModeSettings* Settings = GetMutableDefault<UActorBrowsingModeSettings>();
+				Settings->bHideActorComponents = bHideComponents = bIsActive;
+				Settings->PostEditChange();
+
+				if (auto ActorHierarchy = StaticCast<FActorHierarchy*>(Hierarchy.Get()))
+				{
+					ActorHierarchy->SetShowingComponents(!bIsActive);
+				}
+			});
+
+		FilterInfoMap.Add(TEXT("HideComponentsFilter"), HideComponentsInfo);
+
+		Rebuild();
+	}
+
+	FActorBrowsingMode::~FActorBrowsingMode()
+	{
+		FSceneOutlinerDelegates::Get().OnComponentsUpdated.RemoveAll(this);
+
+		GEngine->OnLevelActorDeleted().RemoveAll(this);
+		
+		FEditorDelegates::OnEditCutActorsBegin.RemoveAll(this);
+		FEditorDelegates::OnEditCutActorsEnd.RemoveAll(this);
+		FEditorDelegates::OnEditCopyActorsBegin.RemoveAll(this);
+		FEditorDelegates::OnEditCopyActorsEnd.RemoveAll(this);
+		FEditorDelegates::OnEditPasteActorsBegin.RemoveAll(this);
+		FEditorDelegates::OnEditPasteActorsEnd.RemoveAll(this);
+		FEditorDelegates::OnDuplicateActorsBegin.RemoveAll(this);
+		FEditorDelegates::OnDuplicateActorsEnd.RemoveAll(this);
+		FEditorDelegates::OnDeleteActorsBegin.RemoveAll(this);
+		FEditorDelegates::OnDeleteActorsEnd.RemoveAll(this);
+	}
+
+	void FActorBrowsingMode::Rebuild()
+	{
+		FActorMode::Rebuild();
+
+		FilteredActorCount = 0;
+		ApplicableActors.Empty();
+	}
+
+	FText FActorBrowsingMode::GetStatusText() const 
+	{
+		if (!RepresentingWorld.IsValid())
+		{
+			return FText();
+		}
+
+		const int32 TotalActorCount = ApplicableActors.Num();
+		const int32 SelectedActorCount = SceneOutliner->GetSelection().Num<FActorTreeItem>();
+
+		if (!SceneOutliner->IsTextFilterActive())
+		{
+			if (SelectedActorCount == 0) //-V547
+			{
+				return FText::Format(LOCTEXT("ShowingAllActorsFmt", "{0} actors"), FText::AsNumber(FilteredActorCount));
+			}
+			else
+			{
+				return FText::Format(LOCTEXT("ShowingAllActorsSelectedFmt", "{0} actors ({1} selected)"), FText::AsNumber(FilteredActorCount), FText::AsNumber(SelectedActorCount));
+			}
+		}
+		else if (SceneOutliner->IsTextFilterActive() && FilteredActorCount == 0)
+		{
+			return FText::Format(LOCTEXT("ShowingNoActorsFmt", "No matching actors ({0} total)"), FText::AsNumber(TotalActorCount));
+		}
+		else if (SelectedActorCount != 0) //-V547
+		{
+			return FText::Format(LOCTEXT("ShowingOnlySomeActorsSelectedFmt", "Showing {0} of {1} actors ({2} selected)"), FText::AsNumber(FilteredActorCount), FText::AsNumber(TotalActorCount), FText::AsNumber(SelectedActorCount));
+		}
+		else
+		{
+			return FText::Format(LOCTEXT("ShowingOnlySomeActorsFmt", "Showing {0} of {1} actors"), FText::AsNumber(FilteredActorCount), FText::AsNumber(TotalActorCount));
+		}
+	}
+
+	FSlateColor FActorBrowsingMode::GetStatusTextColor() const
+	{
+		if (!SceneOutliner->IsTextFilterActive())
+		{
+			return FSlateColor::UseForeground();
+		}
+		else if (FilteredActorCount == 0)
+		{
+			return FAppStyle::Get().GetSlateColor("Colors.AccentRed");
+		}
+		else
+		{
+			return FAppStyle::Get().GetSlateColor("Colors.AccentGreen");
+		}
+	}
+
+	void FActorBrowsingMode::CreateViewContent(FMenuBuilder& MenuBuilder)
+	{
+		MenuBuilder.BeginSection("AssetThumbnails", LOCTEXT("ShowWorldHeading", "World"));
+		{
+			MenuBuilder.AddSubMenu(
+				LOCTEXT("ChooseWorldSubMenu", "Choose World"),
+				LOCTEXT("ChooseWorldSubMenuToolTip", "Choose the world to display in the outliner."),
+				FNewMenuDelegate::CreateRaw(this, &FActorMode::BuildWorldPickerMenu)
+			);
+		}
+		MenuBuilder.EndSection();
+	}
+
+	TSharedRef<FOutlinerFilter> FActorBrowsingMode::CreateShowOnlySelectedActorsFilter()
+	{
+		auto IsActorSelected = [](const AActor* InActor)
+		{
+			return InActor && InActor->IsSelected();
+		};
+		return MakeShareable(new FActorFilter(FActorTreeItem::FFilterPredicate::CreateStatic(IsActorSelected), EDefaultFilterBehaviour::Fail, FActorTreeItem::FFilterPredicate::CreateStatic(IsActorSelected)));
+	}
+
+	TSharedRef<FOutlinerFilter> FActorBrowsingMode::CreateHideTemporaryActorsFilter()
+	{
+		return MakeShareable(new FActorFilter(FActorTreeItem::FFilterPredicate::CreateStatic([](const AActor* InActor)
+			{
+				return ((InActor->GetWorld() && InActor->GetWorld()->WorldType != EWorldType::PIE) || GEditor->ObjectsThatExistInEditorWorld.Get(InActor)) && !InActor->HasAnyFlags(EObjectFlags::RF_Transient);
+			}), EDefaultFilterBehaviour::Pass));
+	}
+
+	TSharedRef<FOutlinerFilter> FActorBrowsingMode::CreateIsInCurrentLevelFilter()
+	{
+		return MakeShareable(new FActorFilter(FActorTreeItem::FFilterPredicate::CreateStatic([](const AActor* InActor)
+			{
+				if (InActor->GetWorld())
+				{
+					return InActor->GetLevel() == InActor->GetWorld()->GetCurrentLevel();
+				}
+
+				return false;
+			}), EDefaultFilterBehaviour::Pass));
+	}
+
+	TSharedRef<FOutlinerFilter> FActorBrowsingMode::CreateHideComponentsFilter()
+	{
+		return MakeShared<TOutlinerPredicateFilter<FComponentTreeItem>>(TOutlinerPredicateFilter<FComponentTreeItem>(FComponentTreeItem::FFilterPredicate::CreateStatic([](const UActorComponent* Component)
+			{
+				// Filter out all components
+				return false;
+			}), EDefaultFilterBehaviour::Pass));
+	}
+
+	static const FName DefaultContextBaseMenuName("SceneOutliner.DefaultContextMenuBase");
+	static const FName DefaultContextMenuName("SceneOutliner.DefaultContextMenu");
+
+	void FActorBrowsingMode::RegisterContextMenu()
+	{
+		UToolMenus* ToolMenus = UToolMenus::Get();
+
+		if (!ToolMenus->IsMenuRegistered(DefaultContextBaseMenuName))
+		{
+			UToolMenu* Menu = ToolMenus->RegisterMenu(DefaultContextBaseMenuName);
+
+			Menu->AddDynamicSection("DynamicSection1", FNewToolMenuDelegate::CreateLambda([](UToolMenu* InMenu)
+				{
+					USceneOutlinerMenuContext* Context = InMenu->FindContext<USceneOutlinerMenuContext>();
+					if (!Context || !Context->SceneOutliner.IsValid())
+					{
+						return;
+					}
+
+					SSceneOutliner* SceneOutliner = Context->SceneOutliner.Pin().Get();
+					if (Context->bShowParentTree)
+					{
+						if (Context->NumSelectedItems == 0)
+						{
+							InMenu->FindOrAddSection("Section").AddMenuEntry(
+								"CreateFolder",
+								LOCTEXT("CreateFolder", "Create Folder"),
+								FText(),
+								FSlateIcon(FEditorStyle::GetStyleSetName(), "SceneOutliner.NewFolderIcon"),
+								FUIAction(FExecuteAction::CreateSP(SceneOutliner, &SSceneOutliner::CreateFolder)));
+						}
+						else
+						{
+							if (Context->NumSelectedItems == 1)
+							{
+								SceneOutliner->GetTree().GetSelectedItems()[0]->GenerateContextMenu(InMenu, *SceneOutliner);
+							}
+
+							// If we've only got folders selected, show the selection and edit sub menus
+							if (Context->NumSelectedItems > 0 && Context->NumSelectedFolders == Context->NumSelectedItems)
+							{
+								InMenu->FindOrAddSection("Section").AddSubMenu(
+									"SelectSubMenu",
+									LOCTEXT("SelectSubmenu", "Select"),
+									LOCTEXT("SelectSubmenu_Tooltip", "Select the contents of the current selection"),
+									FNewToolMenuDelegate::CreateSP(SceneOutliner, &SSceneOutliner::FillSelectionSubMenu));
+							}
+						}
+					}
+				}));
+
+			Menu->AddDynamicSection("DynamicMainSection", FNewToolMenuDelegate::CreateLambda([](UToolMenu* InMenu)
+				{
+					// We always create a section here, even if there is no parent so that clients can still extend the menu
+					FToolMenuSection& Section = InMenu->AddSection("MainSection");
+
+					if (USceneOutlinerMenuContext* Context = InMenu->FindContext<USceneOutlinerMenuContext>())
+					{
+						// Don't add any of these menu items if we're not showing the parent tree
+						// Can't move worlds or level blueprints
+						if (Context->bShowParentTree && Context->NumSelectedItems > 0 && Context->NumWorldsSelected == 0 && Context->SceneOutliner.IsValid())
+						{
+							Section.AddSubMenu(
+								"MoveActorsTo",
+								LOCTEXT("MoveActorsTo", "Move To"),
+								LOCTEXT("MoveActorsTo_Tooltip", "Move selection to another folder"),
+								FNewToolMenuDelegate::CreateSP(Context->SceneOutliner.Pin().Get(), &SSceneOutliner::FillFoldersSubMenu));
+						}
+					}
+				}));
+		}
+
+		if (!ToolMenus->IsMenuRegistered(DefaultContextMenuName))
+		{
+			ToolMenus->RegisterMenu(DefaultContextMenuName, DefaultContextBaseMenuName);
+		}
+	}
+
+	TSharedPtr<SWidget> FActorBrowsingMode::BuildContextMenu()
+	{
+		RegisterContextMenu();
+
+		FItemSelection ItemSelection(SceneOutliner->GetSelection());
+
+		USceneOutlinerMenuContext* ContextObject = NewObject<USceneOutlinerMenuContext>();
+		ContextObject->SceneOutliner = StaticCastSharedRef<SSceneOutliner>(SceneOutliner->AsShared());
+		ContextObject->bShowParentTree = SceneOutliner->GetSharedData().bShowParentTree;
+		ContextObject->NumSelectedItems = ItemSelection.Num();
+		ContextObject->NumSelectedFolders = ItemSelection.Num<FFolderTreeItem>();
+		ContextObject->NumWorldsSelected = ItemSelection.Num<FWorldTreeItem>();
+		FToolMenuContext Context(ContextObject);
+
+		FName MenuName = DefaultContextMenuName;
+		SceneOutliner->GetSharedData().ModifyContextMenu.ExecuteIfBound(MenuName, Context);
+
+		// Build up the menu for a selection
+		UToolMenus* ToolMenus = UToolMenus::Get();
+		UToolMenu* Menu = ToolMenus->GenerateMenu(MenuName, Context);
+
+		for (const FToolMenuSection& Section : Menu->Sections)
+		{
+			if (Section.Blocks.Num() > 0)
+			{
+				return ToolMenus->GenerateWidget(Menu);
+			}
+		}
+
+		return nullptr;
+	}
+
+	TSharedPtr<SWidget> FActorBrowsingMode::CreateContextMenu()
+	{
+		TArray<AActor*> SelectedActors;
+		GEditor->GetSelectedActors()->GetSelectedObjects<AActor>(SelectedActors);
+
+		// Make sure that no components are selected
+		if (GEditor->GetSelectedComponentCount() > 0)
+		{
+			// We want to be able to undo to regain the previous component selection
+			const FScopedTransaction Transaction(NSLOCTEXT("UnrealEd", "ClickingOnActorsContextMenu", "Clicking on Actors (context menu)"));
+			USelection* ComponentSelection = GEditor->GetSelectedComponents();
+			ComponentSelection->Modify(false);
+			ComponentSelection->DeselectAll();
+
+			GUnrealEd->UpdatePivotLocationForSelection();
+			GEditor->RedrawLevelEditingViewports(false);
+		}
+
+		return BuildContextMenu();
+	}
+
+	void FActorBrowsingMode::OnItemAdded(FTreeItemPtr Item)
+	{
+		if (const FActorTreeItem* ActorItem = Item->CastTo<FActorTreeItem>())
+		{
+			if (!Item->Flags.bIsFilteredOut)
+			{
+				++FilteredActorCount;
+
+				// Synchronize selection
+				if (GEditor->GetSelectedActors()->IsSelected(ActorItem->Actor.Get()))
+				{
+					SceneOutliner->SetItemSelection(Item, true);
+				}
+			}
+		}
+		else if (FActorFolderTreeItem* FolderItem = Item->CastTo<FActorFolderTreeItem>())
+		{
+			if (FolderItem->World.IsValid())
+			{
+				if (FActorFolderProps* Props = FActorFolders::Get().GetFolderProperties(*FolderItem->World, FolderItem->Path))
+				{
+					FolderItem->Flags.bIsExpanded = Props->bIsExpanded;
+				}
+			}
+		}
+	}
+
+	void FActorBrowsingMode::OnItemRemoved(FTreeItemPtr Item)
+	{
+		if (const FActorTreeItem* ActorItem = Item->CastTo<FActorTreeItem>())
+		{
+			if (!ActorItem->Flags.bIsFilteredOut)
+			{
+				--FilteredActorCount;
+			}
+		}
+	}
+
+	void FActorBrowsingMode::OnComponentsUpdated()
+	{
+		SceneOutliner->FullRefresh();
+	}
+
+	void FActorBrowsingMode::OnLevelActorDeleted(AActor* Actor)
+	{
+		ApplicableActors.Remove(Actor);
+	}
+
+	void FActorBrowsingMode::OnItemSelectionChanged(FTreeItemPtr TreeItem, ESelectInfo::Type SelectionType, const FItemSelection& Selection)
+	{
+		TArray<AActor*> SelectedActors = Selection.GetData<AActor*>(FActorSelector());
+
+		bool bChanged = false;
+		bool bAnyInPIE = false;
+		for (auto* Actor : SelectedActors)
+		{
+			if (!bAnyInPIE && Actor && Actor->GetOutermost()->HasAnyPackageFlags(PKG_PlayInEditor))
+			{
+				bAnyInPIE = true;
+			}
+			if (!GEditor->GetSelectedActors()->IsSelected(Actor))
+			{
+				bChanged = true;
+				break;
+			}
+		}
+
+		for (FSelectionIterator SelectionIt(*GEditor->GetSelectedActors()); SelectionIt && !bChanged; ++SelectionIt)
+		{
+			const AActor* Actor = CastChecked< AActor >(*SelectionIt);
+			if (!bAnyInPIE && Actor->GetOutermost()->HasAnyPackageFlags(PKG_PlayInEditor))
+			{
+				bAnyInPIE = true;
+			}
+			if (!SelectedActors.Contains(Actor))
+			{
+				// Actor has been deselected
+				bChanged = true;
+
+				// If actor was a group actor, remove its members from the ActorsToSelect list
+				const AGroupActor* DeselectedGroupActor = Cast<AGroupActor>(Actor);
+				if (DeselectedGroupActor)
+				{
+					TArray<AActor*> GroupActors;
+					DeselectedGroupActor->GetGroupActors(GroupActors);
+
+					for (auto* GroupActor : GroupActors)
+					{
+						SelectedActors.Remove(GroupActor);
+					}
+
+				}
+			}
+		}
+
+		// If there's a discrepancy, update the selected actors to reflect this list.
+		if (bChanged)
+		{
+			const FScopedTransaction Transaction(NSLOCTEXT("UnrealEd", "ClickingOnActors", "Clicking on Actors"), !bAnyInPIE);
+			GEditor->GetSelectedActors()->Modify();
+
+			// Clear the selection.
+			GEditor->SelectNone(false, true, true);
+
+			// We'll batch selection changes instead by using BeginBatchSelectOperation()
+			GEditor->GetSelectedActors()->BeginBatchSelectOperation();
+
+			const bool bShouldSelect = true;
+			const bool bNotifyAfterSelect = false;
+			const bool bSelectEvenIfHidden = true;	// @todo outliner: Is this actually OK?
+			for (auto* Actor : SelectedActors)
+			{
+				UE_LOG(LogSceneOutliner, Verbose, TEXT("Clicking on Actor (world outliner): %s (%s)"), *Actor->GetClass()->GetName(), *Actor->GetActorLabel());
+				GEditor->SelectActor(Actor, bShouldSelect, bNotifyAfterSelect, bSelectEvenIfHidden);
+			}
+
+			// Commit selection changes
+			GEditor->GetSelectedActors()->EndBatchSelectOperation(/*bNotify*/false);
+
+			// Fire selection changed event
+			GEditor->NoteSelectionChange();
+		}
+
+		SceneOutliner->RefreshSelection();
+	}
+
+	void FActorBrowsingMode::OnItemDoubleClick(FTreeItemPtr Item)
+	{
+		if (Item->CanInteract())
+		{
+			if (Item->IsA<FActorTreeItem>())
+			{
+				FItemSelection Selection(SceneOutliner->GetSelection());
+				if (Selection.Has<FActorTreeItem>())
+				{
+					const bool bActiveViewportOnly = false;
+					GEditor->MoveViewportCamerasToActor(Selection.GetData<AActor*>(FActorSelector()), bActiveViewportOnly);
+				}
+			}
+		}
+		else
+		{
+			if (const FActorTreeItem* ActorItem = Item->CastTo<FActorTreeItem>())
+			{
+				if (AActor* Actor = ActorItem->Actor.Get())
+				{
+					const bool bActiveViewportOnly = false;
+					GEditor->MoveViewportCamerasToActor(*Actor, bActiveViewportOnly);
+				}
+			}
+			else if (Item->IsA<FFolderTreeItem>())
+			{
+				SceneOutliner->SetItemExpansion(Item, !SceneOutliner->IsItemExpanded(Item));
+			}
+		}
+	}
+
+	void FActorBrowsingMode::OnFilterTextCommited(FItemSelection& Selection, ETextCommit::Type CommitType)
+	{
+		// Start batching selection changes
+		GEditor->GetSelectedActors()->BeginBatchSelectOperation();
+
+		// Select actors (and only the actors) that match the filter text
+		const bool bNoteSelectionChange = false;
+		const bool bDeselectBSPSurfs = false;
+		const bool WarnAboutManyActors = true;
+		GEditor->SelectNone(bNoteSelectionChange, bDeselectBSPSurfs, WarnAboutManyActors);
+		for (AActor* Actor : Selection.GetData<AActor*>(FActorSelector()))
+		{
+			const bool bShouldSelect = true;
+			const bool bSelectEvenIfHidden = false;
+			GEditor->SelectActor(Actor, bShouldSelect, bNoteSelectionChange, bSelectEvenIfHidden);
+		}
+
+		// Commit selection changes
+		GEditor->GetSelectedActors()->EndBatchSelectOperation(/*bNotify*/false);
+
+		// Fire selection changed event
+		GEditor->NoteSelectionChange();
+
+		// Set keyboard focus to the SceneOutliner, so the user can perform keyboard commands that interact
+		// with selected actors (such as Delete, to delete selected actors.)
+		SceneOutliner->SetKeyboardFocus();
+	}
+
+	void FActorBrowsingMode::OnItemPassesFilters(const ITreeItem& Item)
+	{
+		if (const FActorTreeItem* ActorItem = Item.CastTo<FActorTreeItem>())
+		{
+			ApplicableActors.Add(ActorItem->Actor);
+		}
+	}
+
+	FReply FActorBrowsingMode::OnKeyDown(const FKeyEvent& InKeyEvent)
+	{
+		const FItemSelection& Selection = SceneOutliner->GetSelection();
+
+
+		// Rename key: Rename selected actors (not rebindable, because it doesn't make much sense to bind.)
+		if (InKeyEvent.GetKey() == EKeys::F2)
+		{
+			if (Selection.Num() == 1)
+			{
+				FTreeItemPtr ItemToRename = Selection.SelectedItems[0].Pin();
+
+				if (ItemToRename.IsValid() && CanRenameItem(*ItemToRename) && ItemToRename->CanInteract())
+				{
+					SceneOutliner->SetPendingRenameItem(ItemToRename);
+					SceneOutliner->ScrollItemIntoView(ItemToRename);
+				}
+
+				return FReply::Handled();
+			}
+		}
+
+		// F5 forces a full refresh
+		else if (InKeyEvent.GetKey() == EKeys::F5)
+		{
+			SceneOutliner->FullRefresh();
+			return FReply::Handled();
+		}
+
+		// Delete key: Delete selected actors (not rebindable, because it doesn't make much sense to bind.)
+		// Use Delete and Backspace instead of Platform_Delete because the LevelEditor default Edit Delete is bound to both 
+		else if (InKeyEvent.GetKey() == EKeys::Delete || InKeyEvent.GetKey() == EKeys::BackSpace)
+		{
+			if (SceneOutliner->GetSharedData().CustomDelete.IsBound())
+			{
+				SceneOutliner->GetSharedData().CustomDelete.Execute(Selection.SelectedItems);
+			}
+			else
+			{
+				if (RepresentingWorld.IsValid())
+				{
+					GUnrealEd->Exec(RepresentingWorld.Get(), TEXT("DELETE"));
+				}
+			}
+			return FReply::Handled();
+			
+		}
+		return FReply::Unhandled();
+	}
+
+	bool FActorBrowsingMode::CanRename() const
+	{
+		const FItemSelection ItemSelection = SceneOutliner->GetSelection();
+		const uint32 NumberOfFolders = ItemSelection.Num<FFolderTreeItem>();
+		return (NumberOfFolders == 1 && NumberOfFolders == ItemSelection.Num());
+	}
+
+	bool FActorBrowsingMode::CanRenameItem(const ITreeItem& Item) const
+	{
+		// Can only rename actor and folder items when in actor browsing mode
+		return (Item.IsValid() && (Item.IsA<FActorTreeItem>() || Item.IsA<FFolderTreeItem>()));
+	}
+
+	bool FActorBrowsingMode::CanCut() const
+	{
+		const FItemSelection ItemSelection = SceneOutliner->GetSelection();
+		const uint32 NumberOfFolders = ItemSelection.Num<FFolderTreeItem>();
+		return (NumberOfFolders > 0 && NumberOfFolders == ItemSelection.Num());
+	}
+
+	bool FActorBrowsingMode::CanCopy() const
+	{
+		const FItemSelection ItemSelection = SceneOutliner->GetSelection();
+		const uint32 NumberOfFolders = ItemSelection.Num<FFolderTreeItem>();
+		return (NumberOfFolders > 0 && NumberOfFolders == ItemSelection.Num());
+	}
+
+	bool FActorBrowsingMode::CanPaste() const
+	{
+		return CanPasteFoldersOnlyFromClipboard();
+	}
+
+
+	bool FActorBrowsingMode::CanPasteFoldersOnlyFromClipboard() const
+	{
+		// Intentionally not checking if the level is locked/hidden here, as it's better feedback for the user if they attempt to paste
+		// and get the message explaining why it's failed, than just not having the option available to them.
+		FString PasteString;
+		FPlatformApplicationMisc::ClipboardPaste(PasteString);
+		return PasteString.StartsWith("BEGIN FOLDERLIST");
+	}
+
+	TSharedPtr<FDragDropOperation> FActorBrowsingMode::CreateDragDropOperation(const TArray<FTreeItemPtr>& InTreeItems) const
+	{
+		FDragDropPayload DraggedObjects(InTreeItems);
+
+		TSharedPtr<FSceneOutlinerDragDropOp> OutlinerOp = MakeShareable(new FSceneOutlinerDragDropOp());
+
+		if (DraggedObjects.Has<FActorTreeItem>())
+		{
+			TSharedPtr<FActorDragDropOp> ActorOperation = MakeShareable(new FActorDragDropOp);
+			ActorOperation->Init(DraggedObjects.GetData<TWeakObjectPtr<AActor>>(FWeakActorSelector()));
+			OutlinerOp->AddSubOp(ActorOperation);
+		}
+
+		if (DraggedObjects.Has<FFolderTreeItem>())
+		{
+			TSharedPtr<FFolderDragDropOp> FolderOperation = MakeShareable(new FFolderDragDropOp);
+			FolderOperation->Init(DraggedObjects.GetData<FName>(FFolderPathSelector()));
+			OutlinerOp->AddSubOp(FolderOperation);
+		}
+		OutlinerOp->Construct();
+		return OutlinerOp;
+	}
+	
+	bool FActorBrowsingMode::ParseDragDrop(FDragDropPayload& OutPayload, const FDragDropOperation& Operation) const
+	{
+		if (Operation.IsOfType<FSceneOutlinerDragDropOp>())
+		{
+			const auto& OutlinerOp = static_cast<const FSceneOutlinerDragDropOp&>(Operation);
+			if (const auto& FolderOp = OutlinerOp.GetSubOp<FFolderDragDropOp>())
+			{
+				for (const auto& Folder : FolderOp->Folders)
+				{
+					OutPayload.DraggedItems.Add(SceneOutliner->GetTreeItem(Folder));
+				}
+			}
+			if (const auto& ActorOp = OutlinerOp.GetSubOp<FActorDragDropOp>())
+			{
+				for (const auto& Actor : ActorOp->Actors)
+				{
+					OutPayload.DraggedItems.Add(SceneOutliner->GetTreeItem(Actor.Get()));
+				}
+			}
+			return true;
+		}
+		else if (Operation.IsOfType<FActorDragDropOp>())
+		{
+			for (const TWeakObjectPtr<AActor>& Actor : static_cast<const FActorDragDropOp&>(Operation).Actors)
+			{
+				OutPayload.DraggedItems.Add(SceneOutliner->GetTreeItem(Actor.Get()));
+			}
+			return true;
+		}
+		
+		return false;
+	}
+
+	FDragValidationInfo FActorBrowsingMode::ValidateDrop(const ITreeItem& DropTarget, const FDragDropPayload& Payload) const
+	{
+		if (const FActorTreeItem* ActorItem = DropTarget.CastTo<FActorTreeItem>())
+		{
+			if (Payload.Has<FFolderTreeItem>())
+			{
+				return FDragValidationInfo(DropCompatibility::Drop_IncompatibleGeneric, LOCTEXT("FoldersOnActorError", "Cannot attach folders to actors"));
+			}
+
+			const AActor* ActorTarget = ActorItem->Actor.Get();
+
+			if (!ActorTarget || !Payload.Has<FActorTreeItem>())
+			{
+				return FDragValidationInfo(DropCompatibility::Drop_IncompatibleGeneric, FText());
+			}
+
+			FText AttachErrorMsg;
+			bool bCanAttach = true;
+			bool bDraggedOntoAttachmentParent = true;
+
+			const auto& DragActors = Payload.GetData<TWeakObjectPtr<AActor>>(FWeakActorSelector());
+			for (const auto& DragActorPtr : DragActors)
+			{
+				AActor* DragActor = DragActorPtr.Get();
+				if (DragActor)
+				{
+					if (bCanAttach)
+					{
+						if (DragActor->IsChildActor())
+						{
+							AttachErrorMsg = FText::Format(LOCTEXT("Error_AttachChildActor", "Cannot move {0} as it is a child actor."), FText::FromString(DragActor->GetActorLabel()));
+							bCanAttach = bDraggedOntoAttachmentParent = false;
+							break;
+						}
+						if (!GEditor->CanParentActors(ActorTarget, DragActor, &AttachErrorMsg))
+						{
+							bCanAttach = false;
+						}
+					}
+
+					if (DragActor->GetAttachParentActor() != ActorTarget)
+					{
+						bDraggedOntoAttachmentParent = false;
+					}
+				}
+			}
+
+			const FText ActorLabel = FText::FromString(ActorTarget->GetActorLabel());
+			if (bDraggedOntoAttachmentParent)
+			{
+				if (DragActors.Num() == 1)
+				{
+					return FDragValidationInfo(DropCompatibility::Drop_CompatibleDetach, ActorLabel);
+				}
+				else
+				{
+					return FDragValidationInfo(DropCompatibility::Drop_CompatibleMultipleDetach, ActorLabel);
+				}
+			}
+			else if (bCanAttach)
+			{
+				if (DragActors.Num() == 1)
+				{
+					return FDragValidationInfo(DropCompatibility::Drop_CompatibleAttach, ActorLabel);
+				}
+				else
+				{
+					return FDragValidationInfo(DropCompatibility::Drop_CompatibleMultipleAttach, ActorLabel);
+				}
+			}
+			else
+			{
+				if (DragActors.Num() == 1)
+				{
+					return FDragValidationInfo(DropCompatibility::Drop_IncompatibleGeneric, AttachErrorMsg);
+				}
+				else
+				{
+					const FText ReasonText = FText::Format(LOCTEXT("DropOntoText", "{0}. {1}"), ActorLabel, AttachErrorMsg);
+					return FDragValidationInfo(DropCompatibility::Drop_IncompatibleMultipleAttach, ReasonText);
+				}
+			}
+		}
+		else if (DropTarget.IsA<FFolderTreeItem>() || DropTarget.IsA<FWorldTreeItem>())
+		{
+			const FFolderTreeItem* FolderItem = DropTarget.CastTo<FFolderTreeItem>();
+			// World items are treated as folders with path = none
+			const FName& DestinationPath = FolderItem ? FolderItem->Path : NAME_None;
+			if (Payload.Has<FFolderTreeItem>())
+			{
+				// Iterate over all the folders that have been dragged
+				for (FName DraggedFolder : Payload.GetData<FName>(FFolderPathSelector()))
+				{
+					const FName Leaf = GetFolderLeafName(DraggedFolder);
+					const FName Parent = GetParentPath(DraggedFolder);
+
+					if (Parent == DestinationPath)
+					{
+						FFormatNamedArguments Args;
+						Args.Add(TEXT("SourceName"), FText::FromName(Leaf));
+
+						FText Text;
+						if (DestinationPath.IsNone())
+						{
+							Text = FText::Format(LOCTEXT("FolderAlreadyAssignedRoot", "{SourceName} is already assigned to root"), Args);
+						}
+						else
+						{
+							Args.Add(TEXT("DestPath"), FText::FromName(DestinationPath));
+							Text = FText::Format(LOCTEXT("FolderAlreadyAssigned", "{SourceName} is already assigned to {DestPath}"), Args);
+						}
+
+						return FDragValidationInfo(DropCompatibility::Drop_IncompatibleGeneric, Text);
+					}
+
+					const FString DragFolderPath = DraggedFolder.ToString();
+					const FString LeafName = Leaf.ToString();
+					const FString DstFolderPath = DestinationPath.IsNone() ? FString() : DestinationPath.ToString();
+					const FString NewPath = DstFolderPath / LeafName;
+
+					if (FActorFolders::Get().GetFolderProperties(*RepresentingWorld, FName(*NewPath)))
+					{
+						// The folder already exists
+						FFormatNamedArguments Args;
+						Args.Add(TEXT("DragName"), FText::FromString(LeafName));
+						return FDragValidationInfo(DropCompatibility::Drop_IncompatibleGeneric,
+							FText::Format(LOCTEXT("FolderAlreadyExistsRoot", "A folder called \"{DragName}\" already exists at this level"), Args));
+					}
+					else if (DragFolderPath == DstFolderPath || DstFolderPath.StartsWith(DragFolderPath + "/"))
+					{
+						// Cannot drag as a child of itself
+						FFormatNamedArguments Args;
+						Args.Add(TEXT("FolderPath"), FText::FromName(DraggedFolder));
+						return FDragValidationInfo(
+							DropCompatibility::Drop_IncompatibleGeneric,
+							FText::Format(LOCTEXT("ChildOfItself", "Cannot move \"{FolderPath}\" to be a child of itself"), Args));
+					}
+				}
+			}
+
+			if (Payload.Has<FActorTreeItem>())
+			{
+				// Iterate over all the folders that have been dragged
+				for (const TWeakObjectPtr<AActor>& WeakActor : Payload.GetData<TWeakObjectPtr<AActor>>(FWeakActorSelector()))
+				{
+					const AActor* Actor = WeakActor.Get();
+					if (Actor->IsChildActor())
+					{
+						return FDragValidationInfo(DropCompatibility::Drop_IncompatibleGeneric, FText::Format(LOCTEXT("Error_AttachChildActor", "Cannot move {0} as it is a child actor."), FText::FromString(Actor->GetActorLabel())));
+					}
+					else if (Actor->GetFolderPath() == DestinationPath && !Actor->GetAttachParentActor())
+					{
+						FFormatNamedArguments Args;
+						Args.Add(TEXT("SourceName"), FText::FromString(Actor->GetActorLabel()));
+
+						FText Text;
+						if (DestinationPath.IsNone())
+						{
+							Text = FText::Format(LOCTEXT("FolderAlreadyAssignedRoot", "{SourceName} is already assigned to root"), Args);
+						}
+						else
+						{
+							Args.Add(TEXT("DestPath"), FText::FromName(DestinationPath));
+							Text = FText::Format(LOCTEXT("FolderAlreadyAssigned", "{SourceName} is already assigned to {DestPath}"), Args);
+						}
+
+						return FDragValidationInfo(DropCompatibility::Drop_IncompatibleGeneric, Text);
+					}
+				}
+			}
+
+			// Everything else is a valid operation
+			if (DestinationPath.IsNone())
+			{
+				return FDragValidationInfo(DropCompatibility::Drop_CompatibleGeneric, LOCTEXT("MoveToRoot", "Move to root"));
+			}
+			else
+			{
+				FFormatNamedArguments Args;
+				Args.Add(TEXT("DestPath"), FText::FromName(DestinationPath));
+				return FDragValidationInfo(DropCompatibility::Drop_CompatibleGeneric, FText::Format(LOCTEXT("MoveInto", "Move into \"{DestPath}\""), Args));
+			}
+		}
+		else if (DropTarget.IsA<FComponentTreeItem>())
+		{
+			// we don't allow drag and drop on components for now
+			return FDragValidationInfo(DropCompatibility::Drop_IncompatibleGeneric, FText());
+		}
+		return FDragValidationInfo::Invalid();
+	}
+
+	void FActorBrowsingMode::OnDrop(ITreeItem& DropTarget, const FDragDropPayload& Payload, const FDragValidationInfo& ValidationInfo) const
+	{
+		if (const FActorTreeItem* ActorItem = DropTarget.CastTo<FActorTreeItem>())
+		{
+			AActor* DropActor = ActorItem->Actor.Get();
+			if (!DropActor)
+			{
+				return;
+			}
+
+			TArray<TWeakObjectPtr<AActor>> DraggedActors = Payload.GetData<TWeakObjectPtr<AActor>>(FWeakActorSelector());
+
+			FMessageLog EditorErrors("EditorErrors");
+			EditorErrors.NewPage(LOCTEXT("ActorAttachmentsPageLabel", "Actor attachment"));
+
+			if (ValidationInfo.CompatibilityType == DropCompatibility::Drop_CompatibleMultipleDetach || ValidationInfo.CompatibilityType == DropCompatibility::Drop_CompatibleDetach)
+			{
+				const FScopedTransaction Transaction(LOCTEXT("UndoAction_DetachActors", "Detach actors"));
+
+				for (const auto& WeakActor : DraggedActors)
+				{
+					if (auto* DragActor = WeakActor.Get())
+					{
+						// Detach from parent
+						USceneComponent* RootComp = DragActor->GetRootComponent();
+						if (RootComp && RootComp->GetAttachParent())
+						{
+							AActor* OldParent = RootComp->GetAttachParent()->GetOwner();
+							OldParent->Modify();
+							RootComp->DetachFromComponent(FDetachmentTransformRules::KeepWorldTransform);
+
+							DragActor->SetFolderPath_Recursively(OldParent->GetFolderPath());
+						}
+					}
+				}
+			}
+			else if (ValidationInfo.CompatibilityType == DropCompatibility::Drop_CompatibleMultipleAttach || ValidationInfo.CompatibilityType == DropCompatibility::Drop_CompatibleAttach)
+			{
+				// Show socket chooser if we have sockets to select
+
+				auto PerformAttachment = [](FName SocketName, TWeakObjectPtr<AActor> Parent, const TArray<TWeakObjectPtr<AActor>> NewAttachments)
+				{
+					AActor* ParentActor = Parent.Get();
+					if (ParentActor)
+					{
+						// modify parent and child
+						const FScopedTransaction Transaction(LOCTEXT("UndoAction_PerformAttachment", "Attach actors"));
+
+						// Attach each child
+						bool bAttached = false;
+						for (auto& Child : NewAttachments)
+						{
+							AActor* ChildActor = Child.Get();
+							if (GEditor->CanParentActors(ParentActor, ChildActor))
+							{
+								GEditor->ParentActors(ParentActor, ChildActor, SocketName);
+
+								ChildActor->SetFolderPath_Recursively(ParentActor->GetFolderPath());
+							}
+						}
+					}
+				};
+
+				//@TODO: Should create a menu for each component that contains sockets, or have some form of disambiguation within the menu (like a fully qualified path)
+				// Instead, we currently only display the sockets on the root component
+				USceneComponent* Component = DropActor->GetRootComponent();
+				if ((Component != NULL) && (Component->HasAnySockets()))
+				{
+					// Create the popup
+					FSlateApplication::Get().PushMenu(
+						SceneOutliner->AsShared(),
+						FWidgetPath(),
+						SNew(SSocketChooserPopup)
+						.SceneComponent(Component)
+						.OnSocketChosen_Lambda(PerformAttachment, DropActor, MoveTemp(DraggedActors)),
+						FSlateApplication::Get().GetCursorPos(),
+						FPopupTransitionEffect(FPopupTransitionEffect::TypeInPopup)
+					);
+				}
+				else
+				{
+					PerformAttachment(NAME_None, DropActor, MoveTemp(DraggedActors));
+				}
+			}
+
+			// Report errors
+			EditorErrors.Notify(NSLOCTEXT("ActorAttachmentError", "AttachmentsFailed", "Attachments Failed!"));
+		}
+		else if (DropTarget.IsA<FFolderTreeItem>() || DropTarget.IsA<FWorldTreeItem>())
+		{
+			const FFolderTreeItem* FolderItem = DropTarget.CastTo<FFolderTreeItem>();
+			// If the cast fails, the item must be a WorldTreeItem and we set the path to None to reflect this
+			const FName& DestinationPath = FolderItem ? FolderItem->Path : FName(NAME_None);
+
+			const FScopedTransaction Transaction(LOCTEXT("MoveOutlinerItems", "Move World Outliner Items"));
+			
+			auto MoveToDestination = [&DestinationPath](FFolderTreeItem& Item)
+			{
+				Item.MoveTo(DestinationPath);
+			};
+			Payload.ForEachItem<FFolderTreeItem>(MoveToDestination);
+
+			// Set the folder path on all the dragged actors, and detach any that need to be moved
+			if (Payload.Has<FActorTreeItem>())
+			{
+				TSet<const AActor*> ParentActors;
+				TSet<const AActor*> ChildActors;
+
+				for (const auto& WeakActor : Payload.GetData<TWeakObjectPtr<AActor>>(FWeakActorSelector()))
+				{
+					AActor* Actor = WeakActor.Get();
+					if (Actor)
+					{
+						// First mark this object as a parent, then set its children's path
+						ParentActors.Add(Actor);
+						Actor->SetFolderPath(DestinationPath);
+
+						FActorEditorUtils::TraverseActorTree_ParentFirst(Actor, [&](AActor* InActor) {
+							ChildActors.Add(InActor);
+							InActor->SetFolderPath(DestinationPath);
+							return true;
+							}, false);
+					}
+				}
+
+				// Detach parent actors
+				for (const AActor* Parent : ParentActors)
+				{
+					auto* RootComp = Parent->GetRootComponent();
+
+					// We don't detach if it's a child of another that's been dragged
+					if (RootComp && RootComp->GetAttachParent() && !ChildActors.Contains(Parent))
+					{
+						if (AActor* OldParentActor = RootComp->GetAttachParent()->GetOwner())
+						{
+							OldParentActor->Modify();
+						}
+						RootComp->DetachFromComponent(FDetachmentTransformRules::KeepWorldTransform);
+					}
+				}
+			}
+		}
+	}
+
+	FName FActorBrowsingMode::CreateNewFolder()
+	{
+		const FScopedTransaction Transaction(LOCTEXT("UndoAction_CreateFolder", "Create Folder"));
+
+		const FName NewFolderName = FActorFolders::Get().GetDefaultFolderNameForSelection(*RepresentingWorld);
+		FActorFolders::Get().CreateFolderContainingSelection(*RepresentingWorld, NewFolderName);
+
+		return NewFolderName;
+	}
+
+	FName FActorBrowsingMode::CreateFolder(const FName& ParentPath, const FName& LeafName)
+	{
+		const FName NewPath = FActorFolders::Get().GetFolderName(*RepresentingWorld, ParentPath, LeafName);
+		FActorFolders::Get().CreateFolder(*RepresentingWorld, NewPath);
+		return NewPath;
+	}
+
+	bool FActorBrowsingMode::ReparentItemToFolder(const FName& FolderPath, const FTreeItemPtr& Item)
+	{
+		if (FActorTreeItem* ActorItem = Item->CastTo<FActorTreeItem>())
+		{
+			ActorItem->Actor->SetFolderPath_Recursively(FolderPath);
+			return true;
+		}
+		return false;
+	}
+
+	void FActorBrowsingMode::SelectFoldersDescendants(const TArray<FFolderTreeItem*>& FolderItems, bool bSelectImmediateChildrenOnly)
+	{
+		// Ensure that all folder descendants are expanded
+		TFunctionRef<void(const FTreeItemPtr&)> RecursiveExpand = [&](const FTreeItemPtr&  Item)
+		{
+			SceneOutliner->SetItemExpansion(Item, true);
+			
+			if (!bSelectImmediateChildrenOnly)
+			{
+				for (const TWeakPtr<ITreeItem> Child : Item->GetChildren())
+				{
+					RecursiveExpand(Child.Pin());
+				}
+			}
+		};
+
+		// Selects all actor descendants of a folder
+		TFunctionRef<void(const FTreeItemPtr&)> RecursiveActorSelect = [&](const FTreeItemPtr& Item)
+		{
+			// If the current item is an actor, ensure to select it as well
+			if (FActorTreeItem* ActorItem = Item->CastTo<FActorTreeItem>())
+			{
+				if (AActor* Actor = ActorItem->Actor.Get())
+				{
+					GEditor->SelectActor(Actor, true, false);
+				}
+			}
+			// Select all children
+			for (const TWeakPtr<ITreeItem> Child : Item->GetChildren())
+			{
+				if (FActorTreeItem* ActorItem = Child.Pin()->CastTo<FActorTreeItem>())
+				{
+					if (AActor* Actor = ActorItem->Actor.Get())
+					{
+						GEditor->SelectActor(Actor, true, false);
+					}
+				}
+				else if (FFolderTreeItem* FolderItem = Child.Pin()->CastTo<FFolderTreeItem>())
+				{
+					SceneOutliner->SetItemSelection(FolderItem->AsShared(), true);
+				}
+
+				if (!bSelectImmediateChildrenOnly)
+				{
+					for (const TWeakPtr<ITreeItem>& Grandchild : Child.Pin()->GetChildren())
+					{
+						RecursiveActorSelect(Grandchild.Pin());
+					}
+				}
+			}
+		};
+
+		// Expand everything before beginning selection
+		for (FFolderTreeItem* Folder : FolderItems)
+		{
+			RecursiveExpand(Folder->AsShared());
+		}
+
+		// batch selection
+		GEditor->GetSelectedActors()->BeginBatchSelectOperation();
+
+		for (FFolderTreeItem* Folder : FolderItems)
+		{
+			RecursiveActorSelect(Folder->AsShared());
+		}
+
+		GEditor->GetSelectedActors()->EndBatchSelectOperation(/*bNotify*/false);
+		GEditor->NoteSelectionChange();
+	}
+
+	FCreateOutlinerMode FActorBrowsingMode::CreateFolderPickerMode() const
+	{
+		auto MoveSelectionTo = [this](const FTreeItemRef& NewParent)
+		{
+			if (NewParent->IsA<FWorldTreeItem>())
+			{
+				SceneOutliner->MoveSelectionTo(FName());
+			}
+			else if (FFolderTreeItem* FolderItem = NewParent->CastTo<FFolderTreeItem>())
+			{
+				SceneOutliner->MoveSelectionTo(FolderItem->Path);
+			}
+		};
+
+		return FCreateOutlinerMode::CreateLambda([this, MoveSelectionTo](SceneOutliner::SSceneOutliner* Outliner)
+			{
+				return new FActorFolderPickingMode(Outliner, FOnSceneOutlinerItemPicked::CreateLambda(MoveSelectionTo));
+			});
+	}
+
+	void FActorBrowsingMode::OnDuplicateSelected()
+	{
+		GUnrealEd->Exec(RepresentingWorld.Get(), TEXT("DUPLICATE"));
+	}
+
+	void FActorBrowsingMode::OnEditCutActorsBegin()
+	{
+		// Only a callback in actor browsing mode
+		SceneOutliner->CopyFoldersBegin();
+		SceneOutliner->DeleteFoldersBegin();
+	}
+
+	void FActorBrowsingMode::OnEditCutActorsEnd()
+	{
+		// Only a callback in actor browsing mode
+		SceneOutliner->CopyFoldersEnd();
+		SceneOutliner->DeleteFoldersEnd();
+	}
+
+	void FActorBrowsingMode::OnEditCopyActorsBegin()
+	{
+		// Only a callback in actor browsing mode
+		SceneOutliner->CopyFoldersBegin();
+	}
+
+	void FActorBrowsingMode::OnEditCopyActorsEnd()
+	{
+		// Only a callback in actor browsing mode
+		SceneOutliner->CopyFoldersEnd();
+	}
+
+	void FActorBrowsingMode::OnEditPasteActorsBegin()
+	{
+		// Only a callback in actor browsing mode
+		const TArray<FName> FolderPaths = SceneOutliner->GetClipboardPasteFolders();
+		SceneOutliner->PasteFoldersBegin(FolderPaths);
+	}
+
+	void FActorBrowsingMode::OnEditPasteActorsEnd()
+	{
+		// Only a callback in actor browsing mode
+		SceneOutliner->PasteFoldersEnd();
+	}
+
+	void FActorBrowsingMode::OnDuplicateActorsBegin()
+	{
+		// Only a callback in actor browsing mode
+		const TArray<FName> SelectedFolderPaths = SceneOutliner->GetSelection().GetData<FName>(FFolderPathSelector());
+		SceneOutliner->PasteFoldersBegin(SelectedFolderPaths);
+	}
+
+	void FActorBrowsingMode::OnDuplicateActorsEnd()
+	{
+		// Only a callback in actor browsing mode
+		SceneOutliner->PasteFoldersEnd();
+	}
+
+	void FActorBrowsingMode::OnDeleteActorsBegin()
+	{
+		SceneOutliner->DeleteFoldersBegin();
+	}
+
+	void FActorBrowsingMode::OnDeleteActorsEnd()
+	{
+		SceneOutliner->DeleteFoldersEnd();
+	}
+}
+#undef LOCTEXT_NAMESPACE
\ No newline at end of file
diff --git a/Engine/Source/Editor/SceneOutliner/Private/ActorBrowsingModeSettings.h b/Engine/Source/Editor/SceneOutliner/Private/ActorBrowsingModeSettings.h
new file mode 100644
--- /dev/null
+++ b/Engine/Source/Editor/SceneOutliner/Private/ActorBrowsingModeSettings.h
@@ -0,0 +1,35 @@
+// Copyright Epic Games, Inc. All Rights Reserved.
+
+#pragma once
+
+#include "CoreMinimal.h"
+#include "UObject/ObjectMacros.h"
+#include "UObject/Object.h"
+#include "ActorBrowsingModeSettings.generated.h"
+
+
+/**
+ * Implements the settings for the Actor Browsing Outliner (also known as the World Outliner).
+ */
+UCLASS(config = EditorPerProjectUserSettings)
+class UActorBrowsingModeSettings
+	: public UObject
+{
+	GENERATED_UCLASS_BODY()
+
+	/** True when the Scene Outliner is hiding temporary/run-time Actors */
+	UPROPERTY(config)
+	uint32 bHideTemporaryActors : 1;
+
+	/** True when the Scene Outliner is showing only Actors that exist in the current level */
+	UPROPERTY(config)
+	uint32 bShowOnlyActorsInCurrentLevel : 1;
+
+	/** True when the Scene Outliner is only displaying selected Actors */
+	UPROPERTY(config)
+	uint32 bShowOnlySelectedActors : 1;
+
+	/** True when the Scene Outliner is only displaying Actor Components*/
+	UPROPERTY(config)
+	uint32 bHideActorComponents : 1;
+};
diff --git a/Engine/Source/Editor/SceneOutliner/Private/ActorFolderHierarchy.cpp b/Engine/Source/Editor/SceneOutliner/Private/ActorFolderHierarchy.cpp
new file mode 100644
--- /dev/null
+++ b/Engine/Source/Editor/SceneOutliner/Private/ActorFolderHierarchy.cpp
@@ -0,0 +1,151 @@
+#include "ActorFolderHierarchy.h"
+
+#include "ISceneOutlinerMode.h"
+#include "WorldTreeItem.h"
+#include "ActorFolderTreeItem.h"
+#include "EditorActorFolders.h"
+
+namespace SceneOutliner
+{
+	FActorFolderHierarchy::FActorFolderHierarchy(ISceneOutlinerMode* InMode, const TWeakObjectPtr<UWorld>& World)
+		: ISceneOutlinerHierarchy(InMode)
+		, RepresentingWorld(World)
+	{
+		// ActorFolderHierarchy should only be used with a mode which is showing folders
+		check(Mode->ShouldShowFolders());
+	}
+
+	FTreeItemPtr FActorFolderHierarchy::FindParent(const ITreeItem& Item, const TMap<FTreeItemID, FTreeItemPtr>& Items) const
+	{
+		if (Item.IsA<FWorldTreeItem>())
+		{
+			return nullptr;
+		}
+		else if (const FActorFolderTreeItem* ActorFolderItem = Item.CastTo<FActorFolderTreeItem>())
+		{
+			const FName ParentPath = GetParentPath(ActorFolderItem->Path);
+
+			const FTreeItemPtr* ParentItem = nullptr;
+			// If the folder has no parent path, it must be parented to the root world
+			if (ParentPath.IsNone())
+			{
+				ParentItem = Items.Find(ActorFolderItem->World.Get());
+			}
+			else
+			{
+				ParentItem = Items.Find(GetParentPath(ActorFolderItem->Path));
+			}
+
+			if (ParentItem)
+			{
+				return *ParentItem;
+			}
+		}
+		return nullptr;
+	}
+
+	void FActorFolderHierarchy::FindChildren(const ITreeItem& Item, const TMap<FTreeItemID, FTreeItemPtr>& Items, TArray<FTreeItemPtr>& OutChildren) const
+	{
+		if (const FWorldTreeItem* WorldTreeItem = Item.CastTo<FWorldTreeItem>())
+		{
+			// Find folders which are located at the root
+			for (const auto& Pair : FActorFolders::Get().GetFolderPropertiesForWorld(*WorldTreeItem->World))
+			{
+				if (const FTreeItemPtr* PotentialChild = Items.Find(Pair.Key))
+				{
+					if (const FFolderTreeItem* FolderItem = (*PotentialChild)->CastTo<FFolderTreeItem>())
+					{
+						if (GetParentPath(FolderItem->Path).IsNone())
+						{
+							OutChildren.Add(*PotentialChild);
+						}
+					}
+				}
+			}
+		}
+		else if (const FFolderTreeItem* FolderItem = Item.CastTo<FFolderTreeItem>())
+		{
+			// Search through all items and see if there is an item with a path which is a child of this one
+			for (const auto& Pair : Items)
+			{
+				if (const FFolderTreeItem* PotentialChild = Pair.Value->CastTo<FFolderTreeItem>())
+				{
+					if (PathIsChildOf(PotentialChild->Path, FolderItem->Path))
+					{
+						OutChildren.Add(Pair.Value);
+					}
+				}
+			}
+		}
+	}
+
+	void FActorFolderHierarchy::CreateWorldChildren(UWorld* World, TArray<FTreeItemPtr>& OutItems) const
+	{
+		for (const auto& Pair : FActorFolders::Get().GetFolderPropertiesForWorld(*World))
+		{
+			if (FTreeItemPtr FolderItem = Mode->CreateItemFor<FActorFolderTreeItem>(FActorFolderTreeItem(Pair.Key, World)))
+			{
+				OutItems.Add(FolderItem);
+			}
+		}
+	}
+
+	void FActorFolderHierarchy::CreateItems(TArray<FTreeItemPtr>& OutItems) const
+	{
+		check(RepresentingWorld.IsValid());
+
+		if (FTreeItemPtr WorldItem = Mode->CreateItemFor<FWorldTreeItem>(RepresentingWorld.Get()))
+		{
+			OutItems.Add(WorldItem);
+		}
+
+		CreateWorldChildren(RepresentingWorld.Get(), OutItems);
+	}
+
+	void FActorFolderHierarchy::CreateChildren(const FTreeItemPtr& Item, TArray<FTreeItemPtr>& OutChildren) const
+	{
+		if (const FWorldTreeItem* WorldItem = Item->CastTo<FWorldTreeItem>())
+		{
+			check(WorldItem->World == RepresentingWorld);
+			CreateWorldChildren(WorldItem->World.Get(), OutChildren);
+		}
+		else if (FActorFolderTreeItem* FolderItem = Item->CastTo<FActorFolderTreeItem>())
+		{
+			// Since no map of folder->children exist for ActorFolders, must iterate through all
+			// and manually check the path to know if it a child
+			for (const auto& Pair : FActorFolders::Get().GetFolderPropertiesForWorld(*(FolderItem->World)))
+			{
+				if (PathIsChildOf(Pair.Key, FolderItem->Path))
+				{
+					if (FTreeItemPtr NewFolderItem = Mode->CreateItemFor<FActorFolderTreeItem>(FActorFolderTreeItem(Pair.Key, FolderItem->World), true))
+					{
+						OutChildren.Add(NewFolderItem);
+					}
+				}
+			}
+		}
+	}
+
+	FTreeItemPtr FActorFolderHierarchy::CreateParentItem(const FTreeItemPtr& Item) const
+	{
+		if (Item->IsA<FWorldTreeItem>())
+		{
+			return nullptr;
+		}
+		else if (const FActorFolderTreeItem* FolderTreeItem = Item->CastTo<FActorFolderTreeItem>())
+		{
+			const FName ParentPath = GetParentPath(FolderTreeItem->Path);
+			if (ParentPath.IsNone())
+			{
+				UWorld* OwningWorld = FolderTreeItem->World.Get();
+				check(OwningWorld);
+				return Mode->CreateItemFor<FWorldTreeItem>(OwningWorld, true);
+			}
+			else
+			{
+				return Mode->CreateItemFor<FActorFolderTreeItem>(FActorFolderTreeItem(ParentPath, FolderTreeItem->World));
+			}
+		}
+		return nullptr;
+	}
+}
\ No newline at end of file
diff --git a/Engine/Source/Editor/SceneOutliner/Private/ActorFolderPickingMode.cpp b/Engine/Source/Editor/SceneOutliner/Private/ActorFolderPickingMode.cpp
new file mode 100644
--- /dev/null
+++ b/Engine/Source/Editor/SceneOutliner/Private/ActorFolderPickingMode.cpp
@@ -0,0 +1,58 @@
+// Copyright Epic Games, Inc. All Rights Reserved.
+
+#include "ActorFolderPickingMode.h"
+#include "ActorFolderTreeItem.h"
+#include "ActorFolderHierarchy.h"
+
+#define LOCTEXT_NAMESPACE "SceneOutliner_ActoFolderPickingMode"
+
+namespace SceneOutliner
+{
+	FActorFolderPickingMode::FActorFolderPickingMode(SSceneOutliner* InSceneOutliner, FOnSceneOutlinerItemPicked InOnItemPicked, TWeakObjectPtr<UWorld> InSpecifiedWorldToDisplay)
+		: FActorMode(InSceneOutliner, true, InSpecifiedWorldToDisplay)
+		, OnItemPicked(InOnItemPicked)
+	{
+		Rebuild();
+	}
+
+	TUniquePtr<ISceneOutlinerHierarchy> FActorFolderPickingMode::CreateHierarchy()
+	{
+		return MakeUnique<FActorFolderHierarchy>(this, RepresentingWorld);
+	}
+
+	void FActorFolderPickingMode::OnItemSelectionChanged(FTreeItemPtr Item, ESelectInfo::Type SelectionType, const FItemSelection& Selection)
+	{
+		// fire off the notification to whoever is listening.
+		// This may often cause the widget itself to be enqueued for destruction
+		auto SelectedItems = SceneOutliner->GetSelectedItems();
+		if (SelectedItems.Num() > 0)
+		{
+			auto FirstItem = SelectedItems[0];
+			if (FirstItem->CanInteract())
+			{
+				OnItemPicked.ExecuteIfBound(FirstItem.ToSharedRef());
+			}
+		}
+	}
+
+	void FActorFolderPickingMode::OnFilterTextCommited(FItemSelection& Selection, ETextCommit::Type CommitType)
+	{
+		// In folder picking mode, we check to see if we have a selected folder, and if so, fire
+		// off the notification to whoever is listening.  This may often cause the widget itself
+		// to be enqueued for destruction
+		TArray<FActorFolderTreeItem*> Folders;
+		Selection.Get(Folders);
+		if (Folders.Num() == 1 && Folders[0])
+		{
+			// Signal that a folder was selected. We assume it is valid as it wouldn't not have been added to the selection if not
+			SceneOutliner->SetItemSelection(Folders[0]->AsShared(), true, ESelectInfo::OnKeyPress);
+		}
+	}
+
+	void FActorFolderPickingMode::CreateViewContent(FMenuBuilder& MenuBuilder)
+	{
+		FActorMode::BuildWorldPickerMenu(MenuBuilder);
+	}
+}
+
+#undef LOCTEXT_NAMESPACE
\ No newline at end of file
diff --git a/Engine/Source/Editor/SceneOutliner/Private/ActorFolderTreeItem.cpp b/Engine/Source/Editor/SceneOutliner/Private/ActorFolderTreeItem.cpp
new file mode 100644
--- /dev/null
+++ b/Engine/Source/Editor/SceneOutliner/Private/ActorFolderTreeItem.cpp
@@ -0,0 +1,277 @@
+#include "ActorFolderTreeItem.h"
+#include "ActorTreeItem.h"
+#include "SSceneOutliner.h"
+#include "ActorEditorUtils.h"
+#include "EditorActorFolders.h"
+#include "Widgets/Text/SInlineEditableTextBlock.h"
+#include "ISceneOutliner.h"
+#include "ISceneOutlinerMode.h"
+#include "ScopedTransaction.h"
+
+#define LOCTEXT_NAMESPACE "SceneOutliner_ActorFolderTreeItem"
+
+namespace SceneOutliner
+{
+	const FTreeItemType FActorFolderTreeItem::Type(&FFolderTreeItem::Type);
+
+	struct SActorFolderTreeLabel : FCommonLabelData, public SCompoundWidget
+	{
+		SLATE_BEGIN_ARGS(SActorFolderTreeLabel) {}
+		SLATE_END_ARGS()
+
+		void Construct(const FArguments& InArgs, FActorFolderTreeItem& FolderItem, ISceneOutliner& SceneOutliner, const STableRow<FTreeItemPtr>& InRow)
+		{
+			TreeItemPtr = StaticCastSharedRef<FActorFolderTreeItem>(FolderItem.AsShared());
+			WeakSceneOutliner = StaticCastSharedRef<ISceneOutliner>(SceneOutliner.AsShared());
+
+			TSharedPtr<SInlineEditableTextBlock> InlineTextBlock = SNew(SInlineEditableTextBlock)
+				.Text(this, &SActorFolderTreeLabel::GetDisplayText)
+				.HighlightText(SceneOutliner.GetFilterHighlightText())
+				.ColorAndOpacity(this, &SActorFolderTreeLabel::GetForegroundColor)
+				.OnTextCommitted(this, &SActorFolderTreeLabel::OnLabelCommitted)
+				.OnVerifyTextChanged(this, &SActorFolderTreeLabel::OnVerifyItemLabelChanged)
+				.IsSelected(FIsSelected::CreateSP(&InRow, &STableRow<FTreeItemPtr>::IsSelectedExclusively))
+				.IsReadOnly_Lambda([Item = FolderItem.AsShared(), this]()
+			{
+				return !CanExecuteRenameRequest(Item.Get());
+			});
+
+			if (WeakSceneOutliner.Pin()->GetMode()->IsInteractive())
+			{
+				FolderItem.RenameRequestEvent.BindSP(InlineTextBlock.Get(), &SInlineEditableTextBlock::EnterEditingMode);
+			}
+
+			ChildSlot
+				[
+					SNew(SHorizontalBox)
+
+					+ SHorizontalBox::Slot()
+				.AutoWidth()
+				.VAlign(VAlign_Center)
+				.Padding(FDefaultTreeItemMetrics::IconPadding())
+				[
+					SNew(SBox)
+					.WidthOverride(FDefaultTreeItemMetrics::IconSize())
+				[
+					SNew(SImage)
+					.Image(this, &SActorFolderTreeLabel::GetIcon)
+				]
+				]
+
+			+ SHorizontalBox::Slot()
+				.FillWidth(1.0f)
+				.VAlign(VAlign_Center)
+				.Padding(0.0f, 2.0f)
+				[
+					InlineTextBlock.ToSharedRef()
+				]
+				];
+		}
+
+	private:
+		TWeakPtr<FActorFolderTreeItem> TreeItemPtr;
+
+		FText GetDisplayText() const
+		{
+			auto Folder = TreeItemPtr.Pin();
+			return Folder.IsValid() ? FText::FromName(Folder->LeafName) : FText();
+		}
+
+		const FSlateBrush* GetIcon() const
+		{
+			auto TreeItem = TreeItemPtr.Pin();
+			if (!TreeItem.IsValid())
+			{
+				return FEditorStyle::Get().GetBrush(TEXT("SceneOutliner.FolderClosed"));
+			}
+
+			if (TreeItem->Flags.bIsExpanded && TreeItem->GetChildren().Num())
+			{
+				return FEditorStyle::Get().GetBrush(TEXT("SceneOutliner.FolderOpen"));
+			}
+			else
+			{
+				return FEditorStyle::Get().GetBrush(TEXT("SceneOutliner.FolderClosed"));
+			}
+		}
+
+		FSlateColor GetForegroundColor() const
+		{
+			if (auto BaseColor = FCommonLabelData::GetForegroundColor(*TreeItemPtr.Pin()))
+			{
+				return BaseColor.GetValue();
+			}
+
+			return FSlateColor::UseForeground();
+		}
+
+		bool OnVerifyItemLabelChanged(const FText& InLabel, FText& OutErrorMessage)
+		{
+			auto TreeItem = TreeItemPtr.Pin();
+			if (!TreeItem.IsValid())
+			{
+				OutErrorMessage = LOCTEXT("RenameFailed_TreeItemDeleted", "Tree item no longer exists");
+				return false;
+			}
+
+			const FText TrimmedLabel = FText::TrimPrecedingAndTrailing(InLabel);
+
+			if (TrimmedLabel.IsEmpty())
+			{
+				OutErrorMessage = LOCTEXT("RenameFailed_LeftBlank", "Names cannot be left blank");
+				return false;
+			}
+
+			if (TrimmedLabel.ToString().Len() >= NAME_SIZE)
+			{
+				FFormatNamedArguments Arguments;
+				Arguments.Add(TEXT("CharCount"), NAME_SIZE);
+				OutErrorMessage = FText::Format(LOCTEXT("RenameFailed_TooLong", "Names must be less than {CharCount} characters long."), Arguments);
+				return false;
+			}
+
+			const FString LabelString = TrimmedLabel.ToString();
+			if (TreeItem->LeafName.ToString() == LabelString)
+			{
+				return true;
+			}
+
+			int32 Dummy = 0;
+			if (LabelString.FindChar('/', Dummy) || LabelString.FindChar('\\', Dummy))
+			{
+				OutErrorMessage = LOCTEXT("RenameFailed_InvalidChar", "Folder names cannot contain / or \\.");
+				return false;
+			}
+
+			// Validate that this folder doesn't exist already
+			FName NewPath = GetParentPath(TreeItem->Path);
+			if (NewPath.IsNone())
+			{
+				NewPath = FName(*LabelString);
+			}
+			else
+			{
+				NewPath = FName(*(NewPath.ToString() / LabelString));
+			}
+
+			if (FActorFolders::Get().GetFolderProperties(*GWorld, NewPath))
+			{
+				OutErrorMessage = LOCTEXT("RenameFailed_AlreadyExists", "A folder with this name already exists at this level");
+				return false;
+			}
+
+			return true;
+		}
+
+		void OnLabelCommitted(const FText& InLabel, ETextCommit::Type InCommitInfo)
+		{
+			auto TreeItem = TreeItemPtr.Pin();
+			if (TreeItem.IsValid() && !InLabel.ToString().Equals(TreeItem->LeafName.ToString(), ESearchCase::CaseSensitive))
+			{
+				// Rename the item
+				FName NewPath = GetParentPath(TreeItem->Path);
+				if (NewPath.IsNone())
+				{
+					NewPath = FName(*InLabel.ToString());
+				}
+				else
+				{
+					NewPath = FName(*(NewPath.ToString() / InLabel.ToString()));
+				}
+
+				FActorFolders::Get().RenameFolderInWorld(*GWorld, TreeItem->Path, NewPath);
+
+				auto Outliner = WeakSceneOutliner.Pin();
+				if (Outliner.IsValid())
+				{
+					Outliner->SetKeyboardFocus();
+				}
+			}
+		}
+	};
+
+	FActorFolderTreeItem::FActorFolderTreeItem(FName InPath, TWeakObjectPtr<UWorld> InWorld)
+		: FFolderTreeItem(InPath, Type)
+		, World(InWorld)
+	{
+	}
+
+	void FActorFolderTreeItem::OnExpansionChanged()
+	{
+		if (!World.IsValid())
+		{
+			return;
+		}
+
+		// Update the central store of folder properties with this folder's new expansion state
+		if (FActorFolderProps* Props = FActorFolders::Get().GetFolderProperties(*World, Path))
+		{
+			Props->bIsExpanded = Flags.bIsExpanded;
+		}
+	}
+
+	void FActorFolderTreeItem::Delete(FName InNewParentPath)
+	{
+		if (!World.IsValid())
+		{
+			return;
+		}
+		const FScopedTransaction Transaction(LOCTEXT("DeleteFolderTransaction", "Delete Folder"));
+
+		for (const TWeakPtr<ITreeItem>& ChildPtr : GetChildren())
+		{
+			FTreeItemPtr Child = ChildPtr.Pin();
+			if (const FActorTreeItem* ActorItem = Child->CastTo<FActorTreeItem>())
+			{
+				if (AActor* Actor = ActorItem->Actor.Get())
+				{
+					Actor->SetFolderPath_Recursively(InNewParentPath);
+				}
+			}
+			else if (FFolderTreeItem* FolderItem = Child->CastTo<FFolderTreeItem>())
+			{
+				FolderItem->MoveTo(InNewParentPath);
+			}
+		}
+
+		FActorFolders::Get().DeleteFolder(*World, Path);
+	}
+
+	FName FActorFolderTreeItem::MoveTo(const FName& NewParent)
+	{
+		if (World.IsValid())
+		{
+			// Get unique name
+			const FName NewPath = FActorFolders::Get().GetFolderName(*World, NewParent, LeafName);
+
+			if (FActorFolders::Get().RenameFolderInWorld(*World, Path, NewPath))
+			{
+				return NewPath;
+			}
+		}
+		return FName();
+	}
+
+	void FActorFolderTreeItem::CreateSubFolder(TWeakPtr<SSceneOutliner> WeakOutliner)
+	{
+		auto Outliner = WeakOutliner.Pin();
+
+		if (Outliner.IsValid() && World.IsValid())
+		{
+			const FScopedTransaction Transaction(LOCTEXT("UndoAction_CreateFolder", "Create Folder"));
+
+			const FName NewFolderName = FActorFolders::Get().GetDefaultFolderName(*World, Path);
+			FActorFolders::Get().CreateFolder(*World, NewFolderName);
+
+			// At this point the new folder will be in our newly added list, so select it and open a rename when it gets refreshed
+			Outliner->OnItemAdded(NewFolderName, ENewItemAction::Select | ENewItemAction::Rename);
+		}
+	}
+
+	TSharedRef<SWidget> FActorFolderTreeItem::GenerateLabelWidget(ISceneOutliner& Outliner, const STableRow<FTreeItemPtr>& InRow)
+	{
+		return SNew(SActorFolderTreeLabel, *this, Outliner, InRow);
+	}
+}
+
+#undef LOCTEXT_NAMESPACE
\ No newline at end of file
diff --git a/Engine/Source/Editor/SceneOutliner/Private/ActorHierarchy.cpp b/Engine/Source/Editor/SceneOutliner/Private/ActorHierarchy.cpp
new file mode 100644
--- /dev/null
+++ b/Engine/Source/Editor/SceneOutliner/Private/ActorHierarchy.cpp
@@ -0,0 +1,512 @@
+// Copyright Epic Games, Inc. All Rights Reserved.
+
+#include "ActorHierarchy.h"
+
+#include "Engine/World.h"
+#include "EngineUtils.h"
+#include "WorldTreeItem.h"
+#include "ActorTreeItem.h"
+#include "ComponentTreeItem.h"
+#include "ActorFolderTreeItem.h"
+#include "ISceneOutlinerMode.h"
+#include "ActorEditorUtils.h"
+#include "LevelUtils.h"
+#include "GameFramework/WorldSettings.h"
+#include "EditorActorFolders.h"
+
+namespace SceneOutliner
+{
+	TUniquePtr<FActorHierarchy> FActorHierarchy::Create(ISceneOutlinerMode* Mode, const TWeakObjectPtr<UWorld>& World, TFunction<bool(const AActor*)> InIsActorDisplayable)
+	{
+		FActorHierarchy* Hierarchy = new FActorHierarchy(Mode, World, InIsActorDisplayable);
+		
+		GEngine->OnLevelActorAdded().AddRaw(Hierarchy, &FActorHierarchy::OnLevelActorAdded);
+		GEngine->OnLevelActorDeleted().AddRaw(Hierarchy, &FActorHierarchy::OnLevelActorDeleted);
+		GEngine->OnLevelActorDetached().AddRaw(Hierarchy, &FActorHierarchy::OnLevelActorDetached);
+		GEngine->OnLevelActorAttached().AddRaw(Hierarchy, &FActorHierarchy::OnLevelActorAttached);
+		GEngine->OnLevelActorFolderChanged().AddRaw(Hierarchy, &FActorHierarchy::OnLevelActorFolderChanged);
+		GEngine->OnLevelActorListChanged().AddRaw(Hierarchy, &FActorHierarchy::OnLevelActorListChanged);
+
+		FWorldDelegates::LevelAddedToWorld.AddRaw(Hierarchy, &FActorHierarchy::OnLevelAdded);
+		FWorldDelegates::LevelRemovedFromWorld.AddRaw(Hierarchy, &FActorHierarchy::OnLevelRemoved);
+
+		auto& Folders = FActorFolders::Get();
+		Folders.OnFolderCreate.AddRaw(Hierarchy, &FActorHierarchy::OnBroadcastFolderCreate);
+		Folders.OnFolderMove.AddRaw(Hierarchy, &FActorHierarchy::OnBroadcastFolderMove);
+		Folders.OnFolderDelete.AddRaw(Hierarchy, &FActorHierarchy::OnBroadcastFolderDelete);
+
+		return TUniquePtr<FActorHierarchy>(Hierarchy);
+	}
+
+	FActorHierarchy::FActorHierarchy(ISceneOutlinerMode* Mode, const TWeakObjectPtr<UWorld>& World, TFunction<bool(const AActor*)> InIsActorDisplayable)
+		: ISceneOutlinerHierarchy(Mode)
+		, RepresentingWorld(World)
+		, IsActorDisplayableCallback(InIsActorDisplayable)
+	{
+	}
+
+	FActorHierarchy::~FActorHierarchy()
+	{
+		if (GEngine)
+		{
+			GEngine->OnLevelActorAdded().RemoveAll(this);
+			GEngine->OnLevelActorDeleted().RemoveAll(this);
+			GEngine->OnLevelActorDetached().RemoveAll(this);
+			GEngine->OnLevelActorAttached().RemoveAll(this);
+			GEngine->OnLevelActorFolderChanged().RemoveAll(this);
+			GEngine->OnLevelActorListChanged().RemoveAll(this);
+		}
+
+		FWorldDelegates::LevelAddedToWorld.RemoveAll(this);
+		FWorldDelegates::LevelRemovedFromWorld.RemoveAll(this);
+
+		
+		if (FActorFolders::IsAvailable())
+		{
+			auto& Folders = FActorFolders::Get();
+			Folders.OnFolderCreate.RemoveAll(this);
+			Folders.OnFolderMove.RemoveAll(this);
+			Folders.OnFolderDelete.RemoveAll(this);
+		}
+	}
+
+	void FActorHierarchy::FindChildren(const ITreeItem& Item, const TMap<FTreeItemID, FTreeItemPtr>& Items, TArray<FTreeItemPtr>& OutChildren) const
+	{
+		if (const FWorldTreeItem* WorldTreeItem = Item.CastTo<FWorldTreeItem>())
+		{
+			// All actors in the world are children of the world
+			for (FActorIterator ActorIt(WorldTreeItem->World.Get()); ActorIt; ++ActorIt)
+			{
+				// Actors which have parents are not direct children of a world
+				if (ActorIt->GetAttachParentActor() == nullptr)
+				{
+					if (const FTreeItemPtr* ChildItem = Items.Find(*ActorIt))
+					{
+						OutChildren.Add(*ChildItem);
+					}
+				}
+			}
+
+			if (Mode->ShouldShowFolders())
+			{
+				// Find folders which are located at the root
+				for (const auto& Pair : FActorFolders::Get().GetFolderPropertiesForWorld(*WorldTreeItem->World))
+				{
+					if (const FTreeItemPtr* PotentialChild = Items.Find(Pair.Key))
+					{
+						if (const FFolderTreeItem* FolderItem = (*PotentialChild)->CastTo<FFolderTreeItem>())
+						{
+							if (GetParentPath(FolderItem->Path).IsNone())
+							{
+								OutChildren.Add(*PotentialChild);
+							}
+						}
+					}
+				}
+			}
+		}
+		else if (const FActorTreeItem* ActorTreeItem = Item.CastTo<FActorTreeItem>())
+		{
+			// All attached actors are children of this actor
+			TArray<AActor*> AttachedActors;
+			ActorTreeItem->Actor->GetAttachedActors(AttachedActors);
+			for (const AActor* AttachedActor : AttachedActors)
+			{
+				if (const FTreeItemPtr* AttachedActorItem = Items.Find(AttachedActor))
+				{
+					OutChildren.Add(*AttachedActorItem);
+				}
+			}
+			if (bShowingComponents)
+			{
+				for (const UActorComponent* Component : ActorTreeItem->Actor->GetComponents())
+				{
+					if (const FTreeItemPtr* ComponentItem = Items.Find(Component))
+					{
+						OutChildren.Add(*ComponentItem);
+					}
+				}
+			}
+		}
+		else if (const FFolderTreeItem* FolderItem = Item.CastTo<FFolderTreeItem>())
+		{
+			// We should never call FindParents on a folder item if folders are not being shown
+			check(Mode->ShouldShowFolders());
+
+			// Search through all folders and find any folders with paths which are children of this item
+			for (const auto& Pair : FActorFolders::Get().GetFolderPropertiesForWorld(*RepresentingWorld))
+			{
+				if (PathIsChildOf(Pair.Key, FolderItem->Path))
+				{
+					if (const FTreeItemPtr* Child = Items.Find(Pair.Key))
+					{
+						OutChildren.Add(*Child);
+					}
+				}
+			}
+		}
+	}
+
+	FTreeItemPtr FActorHierarchy::FindParent(const ITreeItem& Item, const TMap<FTreeItemID, FTreeItemPtr>& Items) const
+	{
+		if (Item.IsA<FWorldTreeItem>())
+		{
+			return nullptr;
+		}
+		else if (const FActorTreeItem* ActorTreeItem = Item.CastTo<FActorTreeItem>())
+		{
+			if (ActorTreeItem->Actor.IsValid())
+			{
+				const FTreeItemPtr* ParentItem;
+				if (const AActor* ParentActor = ActorTreeItem->Actor->GetAttachParentActor())
+				{
+					if (ParentItem = Items.Find(ParentActor))
+					{
+						return *ParentItem;
+					}
+				}
+				else if (Mode->ShouldShowFolders() && !ActorTreeItem->Actor->GetFolderPath().IsNone())
+				{
+					if (ParentItem = Items.Find(ActorTreeItem->Actor->GetFolderPath()))
+					{
+						return *ParentItem;
+					}
+				}
+				else if (ParentItem = Items.Find(ActorTreeItem->Actor->GetWorld()))
+				{
+					return *ParentItem;
+				}
+			}
+		}
+		else if (const FFolderTreeItem* FolderItem = Item.CastTo<FFolderTreeItem>())
+		{
+			// We should never call FindParents on a folder item if folders are not being shown
+			check(Mode->ShouldShowFolders());
+
+			const FName ParentPath = GetParentPath(FolderItem->Path);
+
+			const FTreeItemPtr* ParentItem = nullptr;
+			// If the folder has no parent path, it must be parented to the root world
+			if (ParentPath.IsNone())
+			{
+				ParentItem = Items.Find(RepresentingWorld.Get());
+			}
+			else
+			{
+				ParentItem = Items.Find(GetParentPath(FolderItem->Path));
+			}
+			
+			if (ParentItem)
+			{
+				return *ParentItem;
+			}
+		}
+		else if (const FComponentTreeItem* ComponentTreeItem = Item.CastTo<FComponentTreeItem>())
+		{
+			const AActor* Owner = ComponentTreeItem->Component->GetOwner();
+			if (Owner)
+			{
+				if (const FTreeItemPtr* ParentItem = Items.Find(Owner))
+				{
+					return *ParentItem;
+				}
+			}
+		}
+		return nullptr;
+	}
+
+	void FActorHierarchy::CreateComponentItems(const AActor* Actor, TArray<FTreeItemPtr>& OutItems) const
+	{
+		// Add all this actors components if showing components and the owning actor was created
+		if (bShowingComponents)
+		{
+			for (UActorComponent* Component : Actor->GetComponents())
+			{
+				if (FTreeItemPtr ComponentItem = Mode->CreateItemFor<FComponentTreeItem>(Component))
+				{
+					OutItems.Add(ComponentItem);
+				}
+			}
+		}
+	}
+
+	void FActorHierarchy::CreateWorldChildren(UWorld* World, TArray<FTreeItemPtr>& OutItems) const
+	{
+		// Create all actor items
+		for (FActorIterator ActorIt(World); ActorIt; ++ActorIt)
+		{
+			if (IsActorDisplayableCallback(*ActorIt))
+			{
+				if (FTreeItemPtr ActorItem = Mode->CreateItemFor<FActorTreeItem>(*ActorIt))
+				{
+					OutItems.Add(ActorItem);
+
+					// Create all component items
+					CreateComponentItems(*ActorIt, OutItems);
+				}
+			}
+		}
+
+		// Create all folder items
+		if (Mode->ShouldShowFolders())
+		{
+			for (const auto& Pair : FActorFolders::Get().GetFolderPropertiesForWorld(*RepresentingWorld))
+			{
+				if (FTreeItemPtr FolderItem = Mode->CreateItemFor<FActorFolderTreeItem>(FActorFolderTreeItem(Pair.Key, RepresentingWorld)))
+				{
+					OutItems.Add(FolderItem);
+				}
+			}
+		}
+	}
+
+	void FActorHierarchy::CreateItems(TArray<FTreeItemPtr>& OutItems) const
+	{
+		check(RepresentingWorld.IsValid());
+
+		if (FTreeItemPtr WorldItem = Mode->CreateItemFor<FWorldTreeItem>(RepresentingWorld))
+		{
+			OutItems.Add(WorldItem);
+		}
+
+		// Create world children regardless of if a world item was created
+		CreateWorldChildren(RepresentingWorld.Get(), OutItems);
+	}
+
+	void FActorHierarchy::CreateChildren(const FTreeItemPtr& Item, TArray<FTreeItemPtr>& OutChildren) const
+	{
+		if (FWorldTreeItem* WorldItem = Item->CastTo<FWorldTreeItem>())
+		{
+			check(WorldItem->World == RepresentingWorld);
+			CreateWorldChildren(WorldItem->World.Get(), OutChildren);
+		}
+		else if (const FActorTreeItem* ParentActorItem = Item->CastTo<FActorTreeItem>())
+		{
+			const AActor* ParentActor = ParentActorItem->Actor.Get();
+
+			if (IsActorDisplayableCallback(ParentActor))
+			{
+				CreateComponentItems(ParentActor, OutChildren);
+
+				TArray<AActor*> ChildActors;
+				TFunction<bool(AActor*)> GetAttachedActors = [&ChildActors, &GetAttachedActors](AActor* Child)
+				{
+					ChildActors.Add(Child);
+					Child->ForEachAttachedActors(GetAttachedActors);
+
+					// Always continue
+					return true;
+				};
+
+				// Grab all direct/indirect children of an actor
+				ParentActor->ForEachAttachedActors(GetAttachedActors);
+
+				for (auto ChildActor : ChildActors)
+				{
+					if (FTreeItemPtr ChildActorItem = Mode->CreateItemFor<FActorTreeItem>(ChildActor))
+					{
+						OutChildren.Add(ChildActorItem);
+
+						CreateComponentItems(ChildActor, OutChildren);
+					}
+				}
+			}
+		}
+		else if (FActorFolderTreeItem* FolderItem = Item->CastTo<FActorFolderTreeItem>())
+		{
+			check(Mode->ShouldShowFolders());
+
+			for (const auto& Pair : FActorFolders::Get().GetFolderPropertiesForWorld(*FolderItem->World))
+			{
+				if (PathIsChildOf(Pair.Key, FolderItem->Path))
+				{
+					if (FTreeItemPtr NewFolderItem = Mode->CreateItemFor<FActorFolderTreeItem>(FActorFolderTreeItem(Pair.Key, FolderItem->World)))
+					{
+						OutChildren.Add(NewFolderItem);
+					}
+				}
+			}
+		}
+	}
+
+	FTreeItemPtr FActorHierarchy::CreateParentItem(const FTreeItemPtr& Item) const
+	{
+		if (Item->IsA<FWorldTreeItem>())
+		{
+			return nullptr;
+		}
+		else if (const FActorTreeItem* ActorTreeItem = Item->CastTo<FActorTreeItem>())
+		{
+			if (ActorTreeItem->Actor.IsValid())
+			{
+				AActor* ParentActor = ActorTreeItem->Actor->GetAttachParentActor();
+				if (ParentActor)
+				{
+					return Mode->CreateItemFor<FActorTreeItem>(ParentActor, true);
+				}
+				// if this item does not belong to a folder (or folders are disabled)
+				else if (!Mode->ShouldShowFolders() || ActorTreeItem->Actor->GetFolderPath().IsNone())
+				{
+					UWorld* OwningWorld = ActorTreeItem->Actor->GetWorld();
+					check(OwningWorld);
+					return Mode->CreateItemFor<FWorldTreeItem>(OwningWorld, true);
+				}
+				else
+				{
+					return Mode->CreateItemFor<FActorFolderTreeItem>(FActorFolderTreeItem(ActorTreeItem->Actor->GetFolderPath(), ActorTreeItem->Actor->GetWorld()), true);
+				}
+			}
+		}
+		else if (const FComponentTreeItem* ComponentTreeItem = Item->CastTo<FComponentTreeItem>())
+		{
+			if (AActor* ParentActor = ComponentTreeItem->Component->GetOwner())
+			{
+				return Mode->CreateItemFor<FActorTreeItem>(ParentActor, true);
+			}
+		}
+		else if (const FActorFolderTreeItem* FolderTreeItem = Item->CastTo<FActorFolderTreeItem>())
+		{
+			check(Mode->ShouldShowFolders());
+			
+			const FName ParentPath = GetParentPath(FolderTreeItem->Path);
+			if (ParentPath.IsNone())
+			{
+				UWorld* OwningWorld = FolderTreeItem->World.Get();
+				check(OwningWorld);
+				return Mode->CreateItemFor<FWorldTreeItem>(OwningWorld, true);
+			}
+			else
+			{
+				return Mode->CreateItemFor<FActorFolderTreeItem>(FActorFolderTreeItem(ParentPath, FolderTreeItem->World), true);
+			}
+		}
+		return nullptr;
+	}
+
+	void FActorHierarchy::FullRefreshEvent()
+	{
+		FHierarchyChangedData EventData;
+		EventData.Type = FHierarchyChangedData::FullRefresh;
+		
+		HierarchyChangedEvent.Broadcast(EventData);
+	}
+
+	void FActorHierarchy::OnLevelActorAdded(AActor* InActor)
+	{
+		if (InActor && RepresentingWorld.Get() == InActor->GetWorld() && IsActorDisplayableCallback(InActor))
+		{
+			FHierarchyChangedData EventData;
+			EventData.Type = FHierarchyChangedData::Added;
+			EventData.Item = Mode->CreateItemFor<FActorTreeItem>(InActor);
+			HierarchyChangedEvent.Broadcast(EventData);
+		}
+	}
+
+	void FActorHierarchy::OnLevelActorDeleted(AActor* InActor)
+	{
+		if (RepresentingWorld.Get() == InActor->GetWorld())
+		{
+			FHierarchyChangedData EventData;
+			EventData.Type = FHierarchyChangedData::Removed;
+			EventData.ItemID = InActor;
+			HierarchyChangedEvent.Broadcast(EventData);
+		}
+	}
+
+	void FActorHierarchy::OnLevelActorAttached(AActor* InActor, const AActor* InParent)
+	{
+		if (RepresentingWorld.Get() == InActor->GetWorld())
+		{
+			FHierarchyChangedData EventData;
+			EventData.Type = FHierarchyChangedData::Moved;
+			EventData.ItemID = InActor;
+			HierarchyChangedEvent.Broadcast(EventData);
+		}
+	}
+
+	void FActorHierarchy::OnLevelActorDetached(AActor* InActor, const AActor* InParent)
+	{
+		if (RepresentingWorld.Get() == InActor->GetWorld())
+		{
+			FHierarchyChangedData EventData;
+			EventData.Type = FHierarchyChangedData::Moved;
+			EventData.ItemID = InActor;
+			HierarchyChangedEvent.Broadcast(EventData);
+		}
+	}
+
+	void FActorHierarchy::OnComponentsUpdated()
+	{
+		FullRefreshEvent();
+	}
+
+	void FActorHierarchy::OnLevelActorListChanged()
+	{
+		FullRefreshEvent();
+	}
+
+	void FActorHierarchy::OnLevelAdded(ULevel* InLevel, UWorld* InWorld)
+	{
+		if (RepresentingWorld.Get() == InWorld)
+		{
+			FullRefreshEvent();
+		}
+	}
+
+	void FActorHierarchy::OnLevelRemoved(ULevel* InLevel, UWorld* InWorld)
+	{
+		if (RepresentingWorld.Get() == InWorld)
+		{
+			FullRefreshEvent();
+		}
+	}
+
+	/** Called when a folder is to be created */
+	void FActorHierarchy::OnBroadcastFolderCreate(UWorld& InWorld, FName NewPath)
+	{
+		if (Mode->ShouldShowFolders() && RepresentingWorld.Get() == &InWorld)
+		{
+			FHierarchyChangedData EventData;
+			EventData.Type = FHierarchyChangedData::Added;
+			EventData.Item = Mode->CreateItemFor<FActorFolderTreeItem>(FActorFolderTreeItem(NewPath, &InWorld));
+			EventData.ItemActions = ENewItemAction::Select | ENewItemAction::Rename;
+			HierarchyChangedEvent.Broadcast(EventData);
+		}
+	}
+
+	/** Called when a folder is to be moved */
+	void FActorHierarchy::OnBroadcastFolderMove(UWorld& InWorld, FName OldPath, FName NewPath)
+	{
+		if (Mode->ShouldShowFolders() && RepresentingWorld.Get() == &InWorld)
+		{
+			FHierarchyChangedData EventData;
+			EventData.Type = FHierarchyChangedData::FolderMoved;
+			EventData.ItemID = OldPath;
+			EventData.NewPath = NewPath;
+			HierarchyChangedEvent.Broadcast(EventData);
+		}
+	}
+
+	/** Called when a folder is to be deleted */
+	void FActorHierarchy::OnBroadcastFolderDelete(UWorld& InWorld, FName Path)
+	{
+		if (Mode->ShouldShowFolders() && RepresentingWorld.Get() == &InWorld)
+		{
+			FHierarchyChangedData EventData;
+			EventData.Type = FHierarchyChangedData::Removed;
+			EventData.ItemID = Path;
+			HierarchyChangedEvent.Broadcast(EventData);
+		}
+	}
+
+	void FActorHierarchy::OnLevelActorFolderChanged(const AActor* InActor, FName OldPath)
+	{
+		if (Mode->ShouldShowFolders() && RepresentingWorld.Get() == InActor->GetWorld())
+		{
+			FHierarchyChangedData EventData;
+			EventData.Type = FHierarchyChangedData::Moved;
+			EventData.ItemID = FTreeItemID(InActor);
+			HierarchyChangedEvent.Broadcast(EventData);
+		}
+	}
+}
diff --git a/Engine/Source/Editor/SceneOutliner/Private/ActorMode.cpp b/Engine/Source/Editor/SceneOutliner/Private/ActorMode.cpp
new file mode 100644
--- /dev/null
+++ b/Engine/Source/Editor/SceneOutliner/Private/ActorMode.cpp
@@ -0,0 +1,330 @@
+// Copyright Epic Games, Inc. All Rights Reserved.
+
+#include "ActorMode.h"
+#include "ActorHierarchy.h"
+#include "Engine/Selection.h"
+#include "Editor/UnrealEdEngine.h"
+#include "SceneOutlinerDelegates.h"
+#include "ActorEditorUtils.h"
+#include "LevelUtils.h"
+#include "GameFramework/WorldSettings.h"
+#include "DragAndDrop/ActorDragDropOp.h"
+#include "Framework/MultiBox/MultiBoxBuilder.h"
+#include "ActorTreeItem.h"
+#include "FolderTreeItem.h"
+#include "ComponentTreeItem.h"
+#include "WorldTreeItem.h"
+
+#define LOCTEXT_NAMESPACE "SceneOutliner_ActorMode"
+
+namespace SceneOutliner
+{
+	using FActorFilter = TOutlinerPredicateFilter<FActorTreeItem>;
+
+	bool FWeakActorSelector::operator()(const TWeakPtr<ITreeItem>& Item, TWeakObjectPtr<AActor>& DataOut) const
+	{
+		if (FActorTreeItem* ActorItem = Item.Pin()->CastTo<FActorTreeItem>())
+		{
+			if (ActorItem->IsValid())
+			{
+				DataOut = ActorItem->Actor;
+				return true;
+			}
+		}
+		return false;
+	}
+
+	bool FActorSelector::operator()(const TWeakPtr<ITreeItem>& Item, AActor*& ActorPtrOut) const
+	{
+		if (FActorTreeItem* ActorItem = Item.Pin()->CastTo<FActorTreeItem>())
+		{
+			if (ActorItem->IsValid())
+			{
+				AActor* Actor = ActorItem->Actor.Get();
+				if (Actor)
+				{
+					ActorPtrOut = Actor;
+					return true;
+				}
+			}
+		}
+		// If a component is selected, we meant for the owning actor to be selected
+		else if (FComponentTreeItem* ComponentItem = Item.Pin()->CastTo<FComponentTreeItem>())
+		{
+			if (ComponentItem->IsValid())
+			{
+				AActor* Actor = ComponentItem->Component->GetOwner();
+				if (Actor)
+				{
+					ActorPtrOut = Actor;
+					return true;
+				}
+			}
+		}
+
+		return false;
+	}
+
+	FActorMode::FActorMode(SSceneOutliner* InSceneOutliner, bool bInHideComponents, TWeakObjectPtr<UWorld> InSpecifiedWorldToDisplay)
+		: ISceneOutlinerMode(InSceneOutliner)
+		, bHideComponents(bInHideComponents)
+		, SpecifiedWorldToDisplay(InSpecifiedWorldToDisplay)
+	{
+	}
+
+	FActorMode::~FActorMode()
+	{
+		if (Hierarchy)
+		{
+			Hierarchy->OnHierarchyChanged().RemoveAll(this);
+		}
+	}
+
+	TUniquePtr<ISceneOutlinerHierarchy> FActorMode::CreateHierarchy()
+	{
+		TUniquePtr<FActorHierarchy> ActorHierarchy = FActorHierarchy::Create(this, RepresentingWorld, [this](const AActor* InActor) {return IsActorDisplayable(InActor); });
+		ActorHierarchy->SetShowingComponents(!bHideComponents);
+
+		return ActorHierarchy;
+	}
+
+	void FActorMode::Rebuild()
+	{
+		ChooseRepresentingWorld();
+
+		if (Hierarchy)
+		{
+			Hierarchy->OnHierarchyChanged().Clear();
+		}
+
+		Hierarchy = CreateHierarchy();
+	}
+
+	void FActorMode::ChooseRepresentingWorld()
+	{
+		// Select a world to represent
+
+		RepresentingWorld = nullptr;
+
+		// If a specified world was provided, represent it
+		if (SpecifiedWorldToDisplay.IsValid())
+		{
+			RepresentingWorld = SpecifiedWorldToDisplay.Get();
+		}
+
+		// check if the user-chosen world is valid and in the editor contexts
+
+		if (!RepresentingWorld.IsValid() && UserChosenWorld.IsValid())
+		{
+			for (const FWorldContext& Context : GEngine->GetWorldContexts())
+			{
+				if (UserChosenWorld.Get() == Context.World())
+				{
+					RepresentingWorld = UserChosenWorld.Get();
+					break;
+				}
+			}
+		}
+
+		// If the user did not manually select a world, try to pick the most suitable world context
+		if (!RepresentingWorld.IsValid())
+		{
+			// ideally we want a PIE world that is standalone or the first client
+			for (const FWorldContext& Context : GEngine->GetWorldContexts())
+			{
+				UWorld* World = Context.World();
+				if (World && Context.WorldType == EWorldType::PIE)
+				{
+					if (World->GetNetMode() == NM_Standalone)
+					{
+						RepresentingWorld = World;
+						break;
+					}
+					else if (World->GetNetMode() == NM_Client && Context.PIEInstance == 2)	// Slightly dangerous: assumes server is always PIEInstance = 1;
+					{
+						RepresentingWorld = World;
+						break;
+					}
+				}
+			}
+		}
+
+		if (RepresentingWorld == nullptr)
+		{
+			// still not world so fallback to old logic where we just prefer PIE over Editor
+
+			for (const FWorldContext& Context : GEngine->GetWorldContexts())
+			{
+				if (Context.WorldType == EWorldType::PIE)
+				{
+					RepresentingWorld = Context.World();
+					break;
+				}
+				else if (Context.WorldType == EWorldType::Editor)
+				{
+					RepresentingWorld = Context.World();
+				}
+			}
+		}
+	}
+
+	void FActorMode::BuildWorldPickerMenu(FMenuBuilder& MenuBuilder)
+	{
+		MenuBuilder.BeginSection("Worlds", LOCTEXT("WorldsHeading", "Worlds"));
+		{
+			MenuBuilder.AddMenuEntry(
+				LOCTEXT("AutoWorld", "Auto"),
+				LOCTEXT("AutoWorldToolTip", "Automatically pick the world to display based on context."),
+				FSlateIcon(),
+				FUIAction(
+					FExecuteAction::CreateRaw(this, &FActorMode::OnSelectWorld, TWeakObjectPtr<UWorld>()),
+					FCanExecuteAction(),
+					FIsActionChecked::CreateRaw(this, &FActorMode::IsWorldChecked, TWeakObjectPtr<UWorld>())
+				),
+				NAME_None,
+				EUserInterfaceActionType::RadioButton
+			);
+
+			for (const FWorldContext& Context : GEngine->GetWorldContexts())
+			{
+				UWorld* World = Context.World();
+				if (World && (World->WorldType == EWorldType::PIE || Context.WorldType == EWorldType::Editor))
+				{
+					MenuBuilder.AddMenuEntry(
+						GetWorldDescription(World),
+						LOCTEXT("ChooseWorldToolTip", "Display actors for this world."),
+						FSlateIcon(),
+						FUIAction(
+							FExecuteAction::CreateRaw(this, &FActorMode::OnSelectWorld, MakeWeakObjectPtr(World)),
+							FCanExecuteAction(),
+							FIsActionChecked::CreateRaw(this, &FActorMode::IsWorldChecked, MakeWeakObjectPtr(World))
+						),
+						NAME_None,
+						EUserInterfaceActionType::RadioButton
+					);
+				}
+			}
+		}
+		MenuBuilder.EndSection();
+	}
+
+	void FActorMode::OnSelectWorld(TWeakObjectPtr<UWorld> World)
+	{
+		UserChosenWorld = World;
+		SceneOutliner->FullRefresh();
+	}
+
+	bool FActorMode::IsWorldChecked(TWeakObjectPtr<UWorld> World) const
+	{
+		return (UserChosenWorld == World) || (World.IsExplicitlyNull() && !UserChosenWorld.IsValid());
+	}
+
+	void FActorMode::SynchronizeActorSelection()
+	{
+		USelection* SelectedActors = GEditor->GetSelectedActors();
+
+		// Deselect actors in the tree that are no longer selected in the world
+		const FItemSelection Selection(SceneOutliner->GetSelection());
+		auto DeselectActors = [this](FActorTreeItem& Item)
+		{
+			if (!Item.Actor.IsValid() || !Item.Actor.Get()->IsSelected())
+			{
+				SceneOutliner->RemoveFromSelection(Item.AsShared());
+			}
+		};
+		Selection.ForEachItem<FActorTreeItem>(DeselectActors);
+
+		// Show actor selection but only if sub objects are not selected
+		if (!Selection.Has<FComponentTreeItem>())
+		{
+			// See if the tree view selector is pointing at a selected item
+			bool bSelectorInSelectionSet = false;
+
+			TArray<FTreeItemPtr> ActorItems;
+			for (FSelectionIterator SelectionIt(*SelectedActors); SelectionIt; ++SelectionIt)
+			{
+				AActor* Actor = CastChecked< AActor >(*SelectionIt);
+				if (FTreeItemPtr ActorItem = SceneOutliner->GetTreeItem(Actor))
+				{
+					if (!bSelectorInSelectionSet && SceneOutliner->HasSelectorFocus(ActorItem))
+					{
+						bSelectorInSelectionSet = true;
+					}
+
+					ActorItems.Add(ActorItem);
+				}
+			}
+
+			// If NOT bSelectorInSelectionSet then we want to just move the selector to the first selected item.
+			ESelectInfo::Type SelectInfo = bSelectorInSelectionSet ? ESelectInfo::Direct : ESelectInfo::OnMouseClick;
+			SceneOutliner->AddToSelection(ActorItems, SelectInfo);
+		}
+		FSceneOutlinerDelegates::Get().SelectionChanged.Broadcast();
+	}
+
+	bool FActorMode::IsActorDisplayable(const AActor* Actor) const
+	{
+		static const FName SequencerActorTag(TEXT("SequencerActor"));
+
+		return Actor && 
+			!SceneOutliner->GetSharedData().bOnlyShowFolders && 												// Don't show actors if we're only showing folders
+			Actor->IsEditable() &&															// Only show actors that are allowed to be selected and drawn in editor
+			Actor->IsListedInSceneOutliner() &&
+			((Actor->GetWorld()->IsPlayInEditor() || !Actor->HasAnyFlags(RF_Transient)) ||
+				(SceneOutliner->GetSharedData().bShowTransient && Actor->HasAnyFlags(RF_Transient)) ||			// Don't show transient actors in non-play worlds
+				(Actor->ActorHasTag(SequencerActorTag))) &&
+			!Actor->IsTemplate() &&															// Should never happen, but we never want CDOs displayed
+			!FActorEditorUtils::IsABuilderBrush(Actor) &&									// Don't show the builder brush
+			!Actor->IsA(AWorldSettings::StaticClass()) &&									// Don't show the WorldSettings actor, even though it is technically editable
+			!Actor->IsPendingKill() &&														// We don't want to show actors that are about to go away
+			FLevelUtils::IsLevelVisible(Actor->GetLevel());								// Only show Actors whose level is visible
+	}
+
+	void FActorMode::OnFilterTextChanged(const FText& InFilterText)
+	{
+		// Scroll last item (if it passes the filter) into view - this means if we are multi-selecting, we show newest selection that passes the filter
+		if (const AActor* LastSelectedActor = GEditor->GetSelectedActors()->GetBottom<AActor>())
+		{
+			// This part is different than that of OnLevelSelectionChanged(nullptr) because IsItemVisible(TreeItem) & ScrollItemIntoView(TreeItem) are applied to
+			// the current visual state, not to the one after applying the filter. Thus, the scroll would go to the place where the object was located
+			// before applying the FilterText
+
+			// If the object is already in the list, but it does not passes the filter, then we do not want to re-add it, because it will be removed by the filter
+			const FTreeItemPtr TreeItem = SceneOutliner->GetTreeItem(LastSelectedActor);
+			if (!TreeItem.IsValid() || !SceneOutliner->PassesTextFilter(TreeItem))
+			{
+				return;
+			}
+
+			// If the object is not in the list, and it does not passes the filter, then we should not re-add it, because it would be removed by the filter again. Unfortunately,
+			// there is no code to check if a future element (i.e., one that is currently not in the TreeItemMap list) will pass the filter. Therefore, we kind of overkill it
+			// by re-adding that element (even though it will be removed). However, AddItemToTree(FTreeItemRef Item) and similar functions already check the element before
+			// adding it. So this solution is fine.
+			// This solution might affect the performance of the World Outliner when a key is pressed, but it will still work properly when the remove/del keys are pressed. Not
+			// updating the filter when !TreeItem.IsValid() would result in the focus not being updated when the remove/del keys are pressed.
+
+			// In any other case (i.e., if the object passes the current filter), re-add it
+			SceneOutliner->ScrollItemIntoView(TreeItem);
+		}
+	}
+
+	int32 FActorMode::GetTypeSortPriority(const ITreeItem& Item) const
+	{
+		if (Item.IsA<FWorldTreeItem>())
+		{
+			return EItemSortOrder::World;
+		}
+		else if (Item.IsA<FFolderTreeItem>())
+		{
+			return EItemSortOrder::Folder;
+		}
+		else if (Item.IsA<FActorTreeItem>() || Item.IsA<FComponentTreeItem>())
+		{
+			return EItemSortOrder::Actor;
+		}
+		// Warning: using actor mode with an unsupported item type!
+		check(false);
+		return -1;
+	}
+}
+#undef LOCTEXT_NAMESPACE
\ No newline at end of file
diff --git a/Engine/Source/Editor/SceneOutliner/Private/ActorPickingMode.cpp b/Engine/Source/Editor/SceneOutliner/Private/ActorPickingMode.cpp
new file mode 100644
--- /dev/null
+++ b/Engine/Source/Editor/SceneOutliner/Private/ActorPickingMode.cpp
@@ -0,0 +1,48 @@
+// Copyright Epic Games, Inc. All Rights Reserved.
+
+#include "ActorPickingMode.h"
+#include "SSceneOutliner.h"
+#include "ActorTreeItem.h"
+
+namespace SceneOutliner
+{
+	FActorPickingMode::FActorPickingMode(SSceneOutliner* InSceneOutliner, bool bInHideComponents, FOnSceneOutlinerItemPicked OnItemPickedDelegate, TWeakObjectPtr<UWorld> InSpecifiedWorldToDisplay)
+		: FActorMode(InSceneOutliner, bInHideComponents, InSpecifiedWorldToDisplay)
+		, OnItemPicked(OnItemPickedDelegate)
+	{
+	}
+
+	void FActorPickingMode::OnItemSelectionChanged(FTreeItemPtr Item, ESelectInfo::Type SelectionType, const FItemSelection& Selection)
+	{
+		// In actor picking mode, we fire off the notification to whoever is listening.
+		// This may often cause the widget itself to be enqueued for destruction
+		auto SelectedItems = SceneOutliner->GetSelectedItems();
+		if (SelectedItems.Num() > 0)
+		{
+			auto FirstItem = SelectedItems[0];
+			if (FirstItem->CanInteract())
+			{
+				OnItemPicked.ExecuteIfBound(FirstItem.ToSharedRef());
+			}
+		}
+	}
+
+	void FActorPickingMode::OnFilterTextCommited(FItemSelection& Selection, ETextCommit::Type CommitType)
+	{
+		// In actor picking mode, we check to see if we have a selected actor, and if so, fire
+		// off the notification to whoever is listening.  This may often cause the widget itself
+		// to be enqueued for destruction
+		TArray<FActorTreeItem*> Actors;
+		Selection.Get(Actors);
+		if (Actors.Num() == 1 && Actors[0])
+		{
+			// Signal that an actor was selected. We assume it is valid as it won't have been added to ActorsToSelect if not.
+			SceneOutliner->SetItemSelection(Actors[0]->AsShared(), true, ESelectInfo::OnKeyPress);
+		}
+	}
+
+	void FActorPickingMode::CreateViewContent(FMenuBuilder& MenuBuilder)
+	{
+		FActorMode::BuildWorldPickerMenu(MenuBuilder);
+	}
+}
\ No newline at end of file
