SceneOutliner refactor

diff --git a/Engine/Source/Editor/SceneOutliner/Private/ActorTreeItem.cpp b/Engine/Source/Editor/SceneOutliner/Private/ActorTreeItem.cpp
--- a/Engine/Source/Editor/SceneOutliner/Private/ActorTreeItem.cpp
+++ b/Engine/Source/Editor/SceneOutliner/Private/ActorTreeItem.cpp
@@ -8,12 +8,13 @@
 #include "Editor.h"
 #include "ScopedTransaction.h"
 #include "SceneOutlinerPublicTypes.h"
-#include "DragAndDrop/ActorDragDropGraphEdOp.h"
 #include "SceneOutlinerDragDrop.h"
 #include "SceneOutlinerStandaloneTypes.h"
-
-
-
+#include "Widgets/Text/SInlineEditableTextBlock.h"
+#include "ActorEditorUtils.h"
+#include "ClassIconFinder.h"
+#include "ISceneOutliner.h"
+#include "ISceneOutlinerMode.h"
 #include "Logging/MessageLog.h"
 #include "SSocketChooser.h"
 
@@ -21,260 +22,286 @@
 
 namespace SceneOutliner
 {
+const FTreeItemType FActorTreeItem::Type(&ITreeItem::Type);
 
-FDragValidationInfo FActorDropTarget::ValidateDrop(FDragDropPayload& DraggedObjects, UWorld& World) const
+struct SActorTreeLabel : FCommonLabelData, public SCompoundWidget
 {
-	if (DraggedObjects.Folders)
+	SLATE_BEGIN_ARGS(SActorTreeLabel) {}
+	SLATE_END_ARGS()
+
+	void Construct(const FArguments& InArgs, FActorTreeItem& ActorItem, ISceneOutliner& SceneOutliner, const STableRow<FTreeItemPtr>& InRow)
 	{
-		return FDragValidationInfo(FActorDragDropGraphEdOp::ToolTip_IncompatibleGeneric, LOCTEXT("FoldersOnActorError", "Cannot attach folders to actors"));
+		WeakSceneOutliner = StaticCastSharedRef<ISceneOutliner>(SceneOutliner.AsShared());
+
+		TreeItemPtr = StaticCastSharedRef<FActorTreeItem>(ActorItem.AsShared());
+		ActorPtr = ActorItem.Actor;
+
+		HighlightText = SceneOutliner.GetFilterHighlightText();
+
+		TSharedPtr<SInlineEditableTextBlock> InlineTextBlock;
+
+		auto MainContent = SNew(SHorizontalBox)
+
+			// Main actor label
+			+ SHorizontalBox::Slot()
+			.VAlign(VAlign_Center)
+			[
+				SAssignNew(InlineTextBlock, SInlineEditableTextBlock)
+				.Text(this, &SActorTreeLabel::GetDisplayText)
+				.ToolTipText(this, &SActorTreeLabel::GetTooltipText)
+				.HighlightText(HighlightText)
+				.ColorAndOpacity(this, &SActorTreeLabel::GetForegroundColor)
+				.OnTextCommitted(this, &SActorTreeLabel::OnLabelCommitted)
+				.OnVerifyTextChanged(this, &SActorTreeLabel::OnVerifyItemLabelChanged)
+				.IsSelected(FIsSelected::CreateSP(&InRow, &STableRow<FTreeItemPtr>::IsSelectedExclusively))
+				.IsReadOnly_Lambda([Item = ActorItem.AsShared(), this]()
+				{
+					return !CanExecuteRenameRequest(Item.Get());
+				})
+			]
+
+		+ SHorizontalBox::Slot()
+			.VAlign(VAlign_Center)
+			.AutoWidth()
+			.Padding(0.0f, 0.f, 3.0f, 0.0f)
+			[
+				SNew(STextBlock)
+				.Text(this, &SActorTreeLabel::GetTypeText)
+			.Visibility(this, &SActorTreeLabel::GetTypeTextVisibility)
+			.HighlightText(HighlightText)
+			];
+
+		if (WeakSceneOutliner.Pin()->GetMode()->IsInteractive())
+		{
+			ActorItem.RenameRequestEvent.BindSP(InlineTextBlock.Get(), &SInlineEditableTextBlock::EnterEditingMode);
+		}
+
+		ChildSlot
+			[
+				SNew(SHorizontalBox)
+
+				+ SHorizontalBox::Slot()
+				.AutoWidth()
+				.VAlign(VAlign_Center)
+				.Padding(FDefaultTreeItemMetrics::IconPadding())
+				[
+					SNew(SBox)
+					.WidthOverride(FDefaultTreeItemMetrics::IconSize())
+					.HeightOverride(FDefaultTreeItemMetrics::IconSize())
+					[
+						SNew(SImage)
+						.Image(this, &SActorTreeLabel::GetIcon)
+						.ToolTipText(this, &SActorTreeLabel::GetIconTooltip)
+					]
+				]
+
+				+ SHorizontalBox::Slot()
+				.FillWidth(1.0f)
+				.VAlign(VAlign_Center)
+				.Padding(0.0f, 0.0f)
+				[
+					MainContent
+				]
+			];
 	}
 
-	const AActor* DropTarget = Actor.Get();
+private:
+	TWeakPtr<FActorTreeItem> TreeItemPtr;
+	TWeakObjectPtr<AActor> ActorPtr;
+	TAttribute<FText> HighlightText;
 
-	if (!DropTarget || !DraggedObjects.Actors)
+	FText GetDisplayText() const
 	{
-		return FDragValidationInfo(FActorDragDropGraphEdOp::ToolTip_IncompatibleGeneric, FText());
-	}	
-
-	FText AttachErrorMsg;
-	bool bCanAttach = true;
-	bool bDraggedOntoAttachmentParent = true;
+		const AActor* Actor = ActorPtr.Get();
+		return Actor ? FText::FromString(Actor->GetActorLabel()) : LOCTEXT("ActorLabelForMissingActor", "(Deleted Actor)");
+	}
 
-	const auto& DragActors = DraggedObjects.Actors.GetValue();
-	for (const auto& DragActorPtr : DragActors)
+	FText GetTooltipText() const
 	{
-		AActor* DragActor = DragActorPtr.Get();
-		if (DragActor)
+		if (const AActor* Actor = ActorPtr.Get())
 		{
-			if (bCanAttach)
-			{
-				if (DragActor->IsChildActor())
-				{
-					AttachErrorMsg = FText::Format(LOCTEXT("Error_AttachChildActor", "Cannot move {0} as it is a child actor."), FText::FromString(DragActor->GetActorLabel()));
-					bCanAttach = bDraggedOntoAttachmentParent = false;
-					break;
-				}
-				if (!GEditor->CanParentActors(DropTarget, DragActor, &AttachErrorMsg))
-				{
-					bCanAttach = false;
-				}
-			}
-
-			if (DragActor->GetAttachParentActor() != DropTarget)
-			{
-				bDraggedOntoAttachmentParent = false;
-			}
+			FFormatNamedArguments Args;
+			Args.Add(TEXT("ID_Name"), LOCTEXT("CustomColumnMode_InternalName", "ID Name"));
+			Args.Add(TEXT("Name"), FText::FromString(Actor->GetName()));
+			return FText::Format(LOCTEXT("ActorNameTooltip", "{ID_Name}: {Name}"), Args);
 		}
+
+		return FText();
 	}
 
-	const FText ActorLabel = FText::FromString(DropTarget->GetActorLabel());
-	if (bDraggedOntoAttachmentParent)
+	FText GetTypeText() const
 	{
-		if (DragActors.Num() == 1)
+		if (const AActor* Actor = ActorPtr.Get())
 		{
-			return FDragValidationInfo(FActorDragDropGraphEdOp::ToolTip_CompatibleDetach, ActorLabel);
-		}
-		else
-		{
-			return FDragValidationInfo(FActorDragDropGraphEdOp::ToolTip_CompatibleMultipleDetach, ActorLabel);
+			return FText::FromName(Actor->GetClass()->GetFName());
 		}
+
+		return FText();
 	}
-	else if (bCanAttach)
+
+	EVisibility GetTypeTextVisibility() const
 	{
-		if (DragActors.Num() == 1)
-		{
-			return FDragValidationInfo(FActorDragDropGraphEdOp::ToolTip_CompatibleAttach, ActorLabel);
-		}
-		else
-		{
-			return FDragValidationInfo(FActorDragDropGraphEdOp::ToolTip_CompatibleMultipleAttach, ActorLabel);
-		}
+		return HighlightText.Get().IsEmpty() ? EVisibility::Collapsed : EVisibility::Visible;
 	}
-	else
+
+	const FSlateBrush* GetIcon() const
 	{
-		if (DragActors.Num() == 1)
+		if (const AActor* Actor = ActorPtr.Get())
 		{
-			return FDragValidationInfo(FActorDragDropGraphEdOp::ToolTip_IncompatibleGeneric, AttachErrorMsg);
+			if (WeakSceneOutliner.IsValid())
+			{
+				FName IconName = Actor->GetCustomIconName();
+				if (IconName == NAME_None)
+				{
+					IconName = Actor->GetClass()->GetFName();
+				}
+
+				const FSlateBrush* CachedBrush = WeakSceneOutliner.Pin()->GetCachedIconForClass(IconName);
+				if (CachedBrush != nullptr)
+				{
+					return CachedBrush;
+				}
+				else
+				{
+
+					const FSlateBrush* FoundSlateBrush = FClassIconFinder::FindIconForActor(Actor);
+					WeakSceneOutliner.Pin()->CacheIconForClass(IconName, FoundSlateBrush);
+					return FoundSlateBrush;
+				}
+			}
+			else
+			{
+				return nullptr;
+			}
 		}
 		else
 		{
-			const FText ReasonText = FText::Format(LOCTEXT("DropOntoText", "{0}. {1}"), ActorLabel, AttachErrorMsg);
-			return FDragValidationInfo(FActorDragDropGraphEdOp::ToolTip_IncompatibleMultipleAttach, ReasonText);
+			return nullptr;
 		}
 	}
-}
-
-void FActorDropTarget::OnDrop(FDragDropPayload& DraggedObjects, UWorld& World, const FDragValidationInfo& ValidationInfo, TSharedRef<SWidget> DroppedOnWidget)
-{
-	AActor* DropActor = Actor.Get();
-	if (!DropActor)
-	{
-		return;
-	}
-
-	FActorArray DraggedActors = DraggedObjects.Actors.GetValue();
-
-	FMessageLog EditorErrors("EditorErrors");
-	EditorErrors.NewPage(LOCTEXT("ActorAttachmentsPageLabel", "Actor attachment"));
 
-	if (ValidationInfo.TooltipType == FActorDragDropGraphEdOp::ToolTip_CompatibleMultipleDetach || ValidationInfo.TooltipType == FActorDragDropGraphEdOp::ToolTip_CompatibleDetach)
+	const FSlateBrush* GetIconOverlay() const
 	{
-		const FScopedTransaction Transaction( LOCTEXT("UndoAction_DetachActors", "Detach actors") );
+		static const FName SequencerActorTag(TEXT("SequencerActor"));
 
-		for (const auto& WeakActor : DraggedActors)
+		if (const AActor* Actor = ActorPtr.Get())
 		{
-			if (auto* DragActor = WeakActor.Get())
+			if (Actor->ActorHasTag(SequencerActorTag))
 			{
-				DetachActorFromParent(DragActor);
+				return FEditorStyle::GetBrush("Sequencer.SpawnableIconOverlay");
 			}
 		}
+		return nullptr;
 	}
-	else if (ValidationInfo.TooltipType == FActorDragDropGraphEdOp::ToolTip_CompatibleMultipleAttach || ValidationInfo.TooltipType == FActorDragDropGraphEdOp::ToolTip_CompatibleAttach)
-	{
-		// Show socket chooser if we have sockets to select
 
-		//@TODO: Should create a menu for each component that contains sockets, or have some form of disambiguation within the menu (like a fully qualified path)
-		// Instead, we currently only display the sockets on the root component
-		USceneComponent* Component = DropActor->GetRootComponent();
-		if ((Component != NULL) && (Component->HasAnySockets()))
-		{
-			// Create the popup
-			FSlateApplication::Get().PushMenu(
-				DroppedOnWidget,
-				FWidgetPath(),
-				SNew(SSocketChooserPopup)
-				.SceneComponent( Component )
-				.OnSocketChosen_Static(&FActorDropTarget::PerformAttachment, Actor, MoveTemp(DraggedActors) ),
-				FSlateApplication::Get().GetCursorPos(),
-				FPopupTransitionEffect( FPopupTransitionEffect::TypeInPopup )
-				);
-		}
-		else
+	FText GetIconTooltip() const
+	{
+		auto TreeItem = TreeItemPtr.Pin();
+		if (!TreeItem.IsValid())
 		{
-			PerformAttachment(NAME_None, Actor, MoveTemp(DraggedActors));
+			return FText();
 		}
-	}
-
-	// Report errors
-	EditorErrors.Notify(NSLOCTEXT("ActorAttachmentError", "AttachmentsFailed", "Attachments Failed!"));
-}
-
-void FActorDropTarget::PerformAttachment(FName SocketName, TWeakObjectPtr<AActor> Parent, const FActorArray NewAttachments)
-{
-	AActor* ParentActor = Parent.Get();
-	if (ParentActor)
-	{
-		// modify parent and child
-		const FScopedTransaction Transaction( LOCTEXT("UndoAction_PerformAttachment", "Attach actors") );
 
-		// Attach each child
-		bool bAttached = false;
-		for (auto& Child : NewAttachments)
+		FText ToolTipText;
+		if (AActor* Actor = ActorPtr.Get())
 		{
-			AActor* ChildActor = Child.Get();
-			if (GEditor->CanParentActors(ParentActor, ChildActor))
+			ToolTipText = FText::FromString(Actor->GetClass()->GetName());
+			if (WeakSceneOutliner.Pin()->GetMode()->IsInteractive())
 			{
-				GEditor->ParentActors(ParentActor, ChildActor, SocketName);
-
-				ChildActor->SetFolderPath_Recursively(ParentActor->GetFolderPath());
+				USceneComponent* RootComponent = Actor->GetRootComponent();
+				if (RootComponent)
+				{
+					FFormatNamedArguments Args;
+					Args.Add(TEXT("ActorClassName"), ToolTipText);
+
+					if (RootComponent->Mobility == EComponentMobility::Static)
+					{
+						ToolTipText = FText::Format(LOCTEXT("ComponentMobility_Static", "{ActorClassName} with static mobility"), Args);
+					}
+					else if (RootComponent->Mobility == EComponentMobility::Stationary)
+					{
+						ToolTipText = FText::Format(LOCTEXT("ComponentMobility_Stationary", "{ActorClassName} with stationary mobility"), Args);
+					}
+					else if (RootComponent->Mobility == EComponentMobility::Movable)
+					{
+						ToolTipText = FText::Format(LOCTEXT("ComponentMobility_Movable", "{ActorClassName} with movable mobility"), Args);
+					}
+				}
 			}
 		}
-	}
-}
 
-void FActorDropTarget::DetachActorFromParent(AActor* ChildActor)
-{
-	USceneComponent* RootComp = ChildActor->GetRootComponent();
-	if (RootComp && RootComp->GetAttachParent())
-	{
-		AActor* OldParent = RootComp->GetAttachParent()->GetOwner();
-		OldParent->Modify();
-		RootComp->DetachFromComponent(FDetachmentTransformRules::KeepWorldTransform);
-		
-		ChildActor->SetFolderPath_Recursively(OldParent->GetFolderPath());
+		return ToolTipText;
 	}
-}
-
-FActorTreeItem::FActorTreeItem(AActor* InActor)
-	: Actor(InActor)
-	, ID(InActor)
-{
-	bExistsInCurrentWorldAndPIE = GEditor->ObjectsThatExistInEditorWorld.Get(InActor);
-}
 
-FTreeItemPtr FActorTreeItem::FindParent(const FTreeItemMap& ExistingItems) const
-{
-	AActor* ActorPtr = Actor.Get();
-	if (!ActorPtr)
+	FSlateColor GetForegroundColor() const
 	{
-		return nullptr;
-	}
+		auto TreeItem = TreeItemPtr.Pin();
 
-	// Parents should have already been added to the tree
-	AActor* ParentActor = ActorPtr->GetAttachParentActor();
-	if (ParentActor)
-	{
-		return ExistingItems.FindRef(ParentActor);
-	}
-	else
-	{
-		const bool bShouldShowFolders = SharedData->Mode == ESceneOutlinerMode::ActorBrowsing || SharedData->bOnlyShowFolders;
+		if (auto BaseColor = FCommonLabelData::GetForegroundColor(*TreeItem))
+		{
+			return BaseColor.GetValue();
+		}
 
-		const FName ActorFolder = ActorPtr->GetFolderPath();
-		if (bShouldShowFolders && !ActorFolder.IsNone())
+		AActor* Actor = ActorPtr.Get();
+		if (!Actor)
 		{
-			return ExistingItems.FindRef(ActorFolder);
+			// Deleted actor!
+			return FLinearColor(0.2f, 0.2f, 0.25f);
 		}
-	}
 
-	if (UWorld* World = ActorPtr->GetWorld())
-	{
-		return ExistingItems.FindRef(World);
-	}
+		UWorld* OwningWorld = TreeItem->Actor->GetWorld();
+		if (!OwningWorld)
+		{
+			// Invalid world!
+			return FLinearColor(0.2f, 0.2f, 0.25f);
+		}
 
-	return nullptr;
-}
+		const bool bRepresentingPIEWorld = OwningWorld->IsPlayInEditor();
+		if (bRepresentingPIEWorld && !TreeItem->bExistsInCurrentWorldAndPIE)
+		{
+			// Highlight actors that are exclusive to PlayWorld
+			return FLinearColor(0.9f, 0.8f, 0.4f);
+		}
 
-FTreeItemPtr FActorTreeItem::CreateParent() const
-{
-	AActor* ActorPtr = Actor.Get();
-	if (!ActorPtr)
-	{
-		return nullptr;
+		// also darken items that are non selectable in the active mode(s)
+		const bool bInSelected = true;
+		const bool bSelectEvenIfHidden = true;		// @todo outliner: Is this actually OK?
+		if (!GEditor->CanSelectActor(Actor, bInSelected, bSelectEvenIfHidden))
+		{
+			return FCommonLabelData::DarkColor;
+		}
+
+		return FSlateColor::UseForeground();
 	}
 
-	AActor* ParentActor = ActorPtr->GetAttachParentActor();
-	if (ParentActor && ensureMsgf(ParentActor != ActorPtr, TEXT("Encountered an Actor attached to itself (%s)"), *ParentActor->GetName()))
+	bool OnVerifyItemLabelChanged(const FText& InLabel, FText& OutErrorMessage)
 	{
-		return MakeShareable(new FActorTreeItem(ParentActor));
+		return FActorEditorUtils::ValidateActorName(InLabel, OutErrorMessage);
 	}
-	else if(!ParentActor)
-	{
-		const bool bShouldShowFolders = SharedData->Mode == ESceneOutlinerMode::ActorBrowsing || SharedData->bOnlyShowFolders;
 
-		const FName ActorFolder = ActorPtr->GetFolderPath();
-		if (bShouldShowFolders && !ActorFolder.IsNone())
+	void OnLabelCommitted(const FText& InLabel, ETextCommit::Type InCommitInfo)
+	{
+		auto* Actor = ActorPtr.Get();
+		if (Actor && Actor->IsActorLabelEditable() && !InLabel.ToString().Equals(Actor->GetActorLabel(), ESearchCase::CaseSensitive))
 		{
-			return MakeShareable(new FFolderTreeItem(ActorFolder));
-		}
+			const FScopedTransaction Transaction(LOCTEXT("SceneOutlinerRenameActorTransaction", "Rename Actor"));
+			FActorLabelUtilities::RenameExistingActor(Actor, InLabel.ToString());
 
-		if (UWorld* World = ActorPtr->GetWorld())
-		{
-			return MakeShareable(new FWorldTreeItem(World));
+			auto Outliner = WeakSceneOutliner.Pin();
+			if (Outliner.IsValid())
+			{
+				Outliner->SetKeyboardFocus();
+			}
 		}
 	}
+};
 
-	return nullptr;
-}
-
-void FActorTreeItem::Visit(const ITreeItemVisitor& Visitor) const
-{
-	Visitor.Visit(*this);
-}
-
-void FActorTreeItem::Visit(const IMutableTreeItemVisitor& Visitor)
+FActorTreeItem::FActorTreeItem(AActor* InActor)
+	: ITreeItem(Type)
+	, Actor(InActor)
+	, ID(InActor)
 {
-	Visitor.Visit(*this);
+	bExistsInCurrentWorldAndPIE = GEditor->ObjectsThatExistInEditorWorld.Get(InActor);
 }
 
 FTreeItemID FActorTreeItem::GetID() const
@@ -288,11 +315,6 @@ FString FActorTreeItem::GetDisplayString() const
 	return ActorPtr ? ActorPtr->GetActorLabel() : LOCTEXT("ActorLabelForMissingActor", "(Deleted Actor)").ToString();
 }
 
-int32 FActorTreeItem::GetTypeSortPriority() const
-{
-	return ETreeItemSortOrder::Actor;
-}
-
 bool FActorTreeItem::CanInteract() const
 {
 	AActor* ActorPtr = Actor.Get();
@@ -311,29 +333,23 @@ bool FActorTreeItem::CanInteract() const
 	return true;
 }
 
-void FActorTreeItem::PopulateDragDropPayload(FDragDropPayload& Payload) const
+TSharedRef<SWidget> FActorTreeItem::GenerateLabelWidget(ISceneOutliner& Outliner, const STableRow<FTreeItemPtr>& InRow)
 {
-	AActor* ActorPtr = Actor.Get();
-	if (ActorPtr)
-	{
-		if (!Payload.Actors)
-		{
-			Payload.Actors = FActorArray();
-		}
-		Payload.Actors->Add(Actor);
-	}
+	return SNew(SActorTreeLabel, *this, Outliner, InRow);
 }
 
-FDragValidationInfo FActorTreeItem::ValidateDrop(FDragDropPayload& DraggedObjects, UWorld& World) const
+void FActorTreeItem::OnVisibilityChanged(const bool bNewVisibility)
 {
-	FActorDropTarget Target(Actor.Get());
-	return Target.ValidateDrop(DraggedObjects, World);
+	// Save the actor to the transaction buffer to support undo/redo, but do
+	// not call Modify, as we do not want to dirty the actor's package and
+	// we're only editing temporary, transient values
+	SaveToTransactionBuffer(Actor.Get(), false);
+	Actor->SetIsTemporarilyHiddenInEditor(!bNewVisibility);
 }
 
-void FActorTreeItem::OnDrop(FDragDropPayload& DraggedObjects, UWorld& World, const FDragValidationInfo& ValidationInfo, TSharedRef<SWidget> DroppedOnWidget)
+bool FActorTreeItem::GetVisibility() const
 {
-	FActorDropTarget Target(Actor.Get());
-	return Target.OnDrop(DraggedObjects, World, ValidationInfo, DroppedOnWidget);
+	return Actor.IsValid() && !Actor->IsTemporarilyHiddenInEditor(true);
 }
 
 }	// namespace SceneOutliner
diff --git a/Engine/Source/Editor/SceneOutliner/Private/ComponentTreeItem.cpp b/Engine/Source/Editor/SceneOutliner/Private/ComponentTreeItem.cpp
--- a/Engine/Source/Editor/SceneOutliner/Private/ComponentTreeItem.cpp
+++ b/Engine/Source/Editor/SceneOutliner/Private/ComponentTreeItem.cpp
@@ -9,115 +9,190 @@
 #include "Kismet2/ComponentEditorUtils.h"
 #include "SceneOutlinerDragDrop.h"
 #include "SceneOutlinerStandaloneTypes.h"
-
+#include "Widgets/Text/SInlineEditableTextBlock.h"
+#include "ISceneOutliner.h"
+#include "ISceneOutlinerMode.h"
 
 #define LOCTEXT_NAMESPACE "SceneOutliner_ComponentTreeItem"
 
+
 namespace SceneOutliner
 {
-
-FDragValidationInfo FComponentDropTarget::ValidateDrop(FDragDropPayload& DraggedObjects, UWorld& World) const
+	
+struct SComponentTreeLabel : FCommonLabelData, public SCompoundWidget
 {
-	// we don't allow drag and drop for now
-	FText AttachErrorMsg;
-	bool bCanAttach = false;
-	return FDragValidationInfo(FActorDragDropGraphEdOp::ToolTip_IncompatibleGeneric, AttachErrorMsg);
-}
+	SLATE_BEGIN_ARGS(SComponentTreeLabel) {}
+	SLATE_END_ARGS()
 
-void FComponentDropTarget::OnDrop(FDragDropPayload& DraggedObjects, UWorld& World, const FDragValidationInfo& ValidationInfo, TSharedRef<SWidget> DroppedOnWidget)
-{
+	void Construct(const FArguments& InArgs, FComponentTreeItem& ComponentItem, ISceneOutliner& SceneOutliner, const STableRow<FTreeItemPtr>& InRow)
+	{
+		TreeItemPtr = StaticCastSharedRef<FComponentTreeItem>(ComponentItem.AsShared());
+		WeakSceneOutliner = StaticCastSharedRef<ISceneOutliner>(SceneOutliner.AsShared());
+
+		ComponentPtr = ComponentItem.Component;
+		
+		HighlightText = SceneOutliner.GetFilterHighlightText();
+
+		TSharedPtr<SInlineEditableTextBlock> InlineTextBlock;
+
+		auto MainContent = SNew(SHorizontalBox)
+
+			// Main actor label
+			+ SHorizontalBox::Slot()
+			[
+				SAssignNew(InlineTextBlock, SInlineEditableTextBlock)
+				.Text(this, &SComponentTreeLabel::GetDisplayText)
+				.ToolTipText(this, &SComponentTreeLabel::GetTooltipText)
+				//.HighlightText(HighlightText)
+				.ColorAndOpacity(this, &SComponentTreeLabel::GetForegroundColor)
+				//.OnTextCommitted(this, &SComponentTreeLabel::OnLabelCommitted)
+				//.OnVerifyTextChanged(this, &SComponentTreeLabel::OnVerifyItemLabelChanged)
+				.IsReadOnly_Lambda([Item = ComponentItem.AsShared(), this]()
+				{
+					return !CanExecuteRenameRequest(Item.Get());
+				})
+				.IsSelected(FIsSelected::CreateSP(&InRow, &STableRow<FTreeItemPtr>::IsSelectedExclusively))
+			];
+
+		if (WeakSceneOutliner.Pin()->GetMode()->IsInteractive())
+		{
+			ComponentItem.RenameRequestEvent.BindSP(InlineTextBlock.Get(), &SInlineEditableTextBlock::EnterEditingMode);
+		}
 
-}
+		ChildSlot
+			[
+				SNew(SHorizontalBox)
+
+				+ SHorizontalBox::Slot()
+				.AutoWidth()
+				.VAlign(VAlign_Center)
+				.Padding(FDefaultTreeItemMetrics::IconPadding())
+				[
+					SNew(SBox)
+					.WidthOverride(FDefaultTreeItemMetrics::IconSize())
+					.HeightOverride(FDefaultTreeItemMetrics::IconSize())
+					[
+						SNew(SImage)
+						.Image(this, &SComponentTreeLabel::GetIcon)
+						.ToolTipText(this, &SComponentTreeLabel::GetIconTooltip)
+					]
+				]
+
+				+ SHorizontalBox::Slot()
+				.FillWidth(1.0f)
+				.VAlign(VAlign_Center)
+				.Padding(0.0f, 2.0f)
+				[
+					MainContent
+				]
+			];
 
-FComponentTreeItem::FComponentTreeItem(UActorComponent* InComponent)
-	: Component(InComponent)
-	, ID(InComponent)
-{
-	AActor* OwningActor = InComponent->GetOwner();
-	bExistsInCurrentWorldAndPIE = GEditor->ObjectsThatExistInEditorWorld.Get(OwningActor);
+	}
 
-	const FName VariableName = FComponentEditorUtils::FindVariableNameGivenComponentInstance(InComponent);
-	const bool bIsArrayVariable = !VariableName.IsNone() && InComponent->GetOwner() != nullptr && FindFProperty<FArrayProperty>(InComponent->GetOwner()->GetClass(), VariableName);
+private:
+	TWeakPtr<FComponentTreeItem> TreeItemPtr;
+	TWeakObjectPtr<UActorComponent> ComponentPtr;
+	TAttribute<FText> HighlightText;
 
-	if (!VariableName.IsNone() && !bIsArrayVariable)
+
+	FText GetDisplayText() const
 	{
-		CachedDisplayString = VariableName.ToString();
+		auto Item = TreeItemPtr.Pin();
+		return Item.IsValid() ? FText::FromString(Item->GetDisplayString()) : FText();
 	}
-}
 
-FTreeItemPtr FComponentTreeItem::FindParent(const FTreeItemMap& ExistingItems) const
-{
-	UActorComponent* ComponentPtr = Component.Get();
-	if (!ComponentPtr)
+	FText GetTypeText() const
 	{
-		return nullptr;
+		if (const UActorComponent* Component = ComponentPtr.Get())
+		{
+			return FText::FromName(Component->GetClass()->GetFName());
+		}
+
+		return FText();
 	}
-	
-	// Parent actor should have already been added to the tree
-	AActor* ParentActor = ComponentPtr->GetOwner();
-	if (ParentActor)
+
+	EVisibility GetTypeTextVisibility() const
 	{
-		return ExistingItems.FindRef(ParentActor);
+		return HighlightText.Get().IsEmpty() ? EVisibility::Collapsed : EVisibility::Visible;
 	}
-	else
-	{
-		const bool bShouldShowFolders = SharedData->Mode == ESceneOutlinerMode::ActorBrowsing || SharedData->bOnlyShowFolders;
 
-		const FName ComponentFolder = *ComponentPtr->GetDetailedInfo();
-		if (bShouldShowFolders && !ComponentFolder.IsNone())
+	FText GetTooltipText() const
+	{
+		if (const UActorComponent* Component = ComponentPtr.Get())
 		{
-			return ExistingItems.FindRef(ComponentFolder);
+			FFormatNamedArguments Args;
+			Args.Add(TEXT("ID_Name"), LOCTEXT("CustomColumnMode_InternalName", "ID Name"));
+			Args.Add(TEXT("Name"), FText::FromString(Component->GetName()));
+			return FText::Format(LOCTEXT("ComponentNameTooltip", "{ID_Name}: {Name}"), Args);
 		}
-	}
 
-	if (UWorld* World = ComponentPtr->GetWorld())
-	{
-		return ExistingItems.FindRef(World);
+		return FText();
 	}
 
-	return nullptr;
-}
-
-FTreeItemPtr FComponentTreeItem::CreateParent() const
-{
-	UActorComponent* ComponentPtr = Component.Get();
-	if (!ComponentPtr)
+	const FSlateBrush* GetIcon() const
 	{
-		return nullptr;
+		if (const UActorComponent* Component = ComponentPtr.Get())
+		{
+			if (WeakSceneOutliner.IsValid())
+			{
+				const FSlateBrush* CachedBrush = WeakSceneOutliner.Pin()->GetCachedIconForClass(Component->GetClass()->GetFName());
+				if (CachedBrush != nullptr)
+				{
+					return CachedBrush;
+				}
+				else
+				{
+					const FSlateBrush* FoundSlateBrush = FSlateIconFinder::FindIconBrushForClass(Component->GetClass());
+					WeakSceneOutliner.Pin()->CacheIconForClass(Component->GetClass()->GetFName(), const_cast<FSlateBrush*>(FoundSlateBrush));
+					return FoundSlateBrush;
+				}
+			}
+			else
+			{
+				return nullptr;
+			}
+		}
+		else
+		{
+			return nullptr;
+		}
 	}
 
-	AActor* ParentActor = ComponentPtr->GetOwner();
-	if (ParentActor)
+	FText GetIconTooltip() const
 	{
-		return MakeShareable(new FActorTreeItem(ParentActor));
+		FText ToolTipText;
+		return ToolTipText;
 	}
-	else if (!ParentActor)
-	{
-		const bool bShouldShowFolders = SharedData->Mode == ESceneOutlinerMode::ActorBrowsing || SharedData->bOnlyShowFolders;
 
-		const FName ComponentFolder = *ComponentPtr->GetDetailedInfo();
-		if (bShouldShowFolders && !ComponentFolder.IsNone())
+	FSlateColor GetForegroundColor() const
+	{
+		if (auto BaseColor = FCommonLabelData::GetForegroundColor(*TreeItemPtr.Pin()))
 		{
-			return MakeShareable(new FFolderTreeItem(ComponentFolder));
+			return BaseColor.GetValue();
 		}
 
-		if (UWorld* World = ComponentPtr->GetWorld())
-		{
-			return MakeShareable(new FWorldTreeItem(World));
-		}
+		return FSlateColor::UseForeground();
 	}
+};
 
-	return nullptr;
-}
 
-void FComponentTreeItem::Visit(const ITreeItemVisitor& Visitor) const
-{
-	Visitor.Visit(*this);
-}
+const FTreeItemType FComponentTreeItem::Type(&ITreeItem::Type);
 
-void FComponentTreeItem::Visit(const IMutableTreeItemVisitor& Visitor)
+FComponentTreeItem::FComponentTreeItem(UActorComponent* InComponent)
+	: ITreeItem(Type)
+	, Component(InComponent)
+	, ID(InComponent)
 {
-	Visitor.Visit(*this);
+	AActor* OwningActor = InComponent->GetOwner();
+	bExistsInCurrentWorldAndPIE = GEditor->ObjectsThatExistInEditorWorld.Get(OwningActor);
+
+	const FName VariableName = FComponentEditorUtils::FindVariableNameGivenComponentInstance(InComponent);
+	const bool bIsArrayVariable = !VariableName.IsNone() && InComponent->GetOwner() != nullptr && FindFProperty<FArrayProperty>(InComponent->GetOwner()->GetClass(), VariableName);
+
+	if (!VariableName.IsNone() && !bIsArrayVariable)
+	{
+		CachedDisplayString = VariableName.ToString();
+	}
 }
 
 FTreeItemID FComponentTreeItem::GetID() const
@@ -142,11 +217,6 @@ FString FComponentTreeItem::GetDisplayString() const
 	return LOCTEXT("ComponentLabelForMissingComponent", "(Deleted Component)").ToString();
 }
 
-int32 FComponentTreeItem::GetTypeSortPriority() const
-{
-	return ETreeItemSortOrder::Actor;
-}
-
 bool FComponentTreeItem::CanInteract() const
 {
 	UActorComponent* ComponentPtr = Component.Get();
@@ -167,24 +237,10 @@ bool FComponentTreeItem::CanInteract() const
 	return true;
 }
 
-void FComponentTreeItem::PopulateDragDropPayload(FDragDropPayload& Payload) const
+TSharedRef<SWidget> FComponentTreeItem::GenerateLabelWidget(ISceneOutliner& Outliner, const STableRow<FTreeItemPtr>& InRow)
 {
-	// no drag and drop
+	return SNew(SComponentTreeLabel, *this, Outliner, InRow);
 }
-
-FDragValidationInfo FComponentTreeItem::ValidateDrop(FDragDropPayload& DraggedObjects, UWorld& World) const
-{
-	FComponentDropTarget Target(Component.Get());
-	return Target.ValidateDrop(DraggedObjects, World);
-}
-
-void FComponentTreeItem::OnDrop(FDragDropPayload& DraggedObjects, UWorld& World, const FDragValidationInfo& ValidationInfo, TSharedRef<SWidget> DroppedOnWidget)
-{
-	FComponentDropTarget Target(Component.Get());
-	return Target.OnDrop(DraggedObjects, World, ValidationInfo, DroppedOnWidget);
-}
-
-
 } // namespace SceneOutliner
 
 #undef LOCTEXT_NAMESPACE
diff --git a/Engine/Source/Editor/SceneOutliner/Private/FolderTreeItem.cpp b/Engine/Source/Editor/SceneOutliner/Private/FolderTreeItem.cpp
--- a/Engine/Source/Editor/SceneOutliner/Private/FolderTreeItem.cpp
+++ b/Engine/Source/Editor/SceneOutliner/Private/FolderTreeItem.cpp
@@ -6,216 +6,44 @@
 #include "ToolMenus.h"
 #include "EditorStyleSet.h"
 #include "ScopedTransaction.h"
-#include "DragAndDrop/ActorDragDropGraphEdOp.h"
 #include "SceneOutlinerDragDrop.h"
 #include "SSceneOutliner.h"
 
 #include "ActorEditorUtils.h"
-
 #include "EditorActorFolders.h"
 
 #define LOCTEXT_NAMESPACE "SceneOutliner_FolderTreeItem"
 
 namespace SceneOutliner
 {
+const FTreeItemType FFolderTreeItem::Type(&ITreeItem::Type);
 
-FDragValidationInfo FFolderDropTarget::ValidateDrop(FDragDropPayload& DraggedObjects, UWorld& World) const
+bool FFolderPathSelector::operator()(TWeakPtr<ITreeItem> Item, FName& DataOut) const
 {
-	if (DraggedObjects.Folders)
-	{
-		// Iterate over all the folders that have been dragged
-		for (FName DraggedFolder : DraggedObjects.Folders.GetValue())
-		{
-			FName Leaf = GetFolderLeafName(DraggedFolder);
-			FName Parent = GetParentPath(DraggedFolder);
-
-			if (Parent == DestinationPath)
-			{
-				FFormatNamedArguments Args;
-				Args.Add(TEXT("SourceName"), FText::FromName(Leaf));
-
-				FText Text;
-				if (DestinationPath.IsNone())
-				{
-					Text = FText::Format(LOCTEXT("FolderAlreadyAssignedRoot", "{SourceName} is already assigned to root"), Args);
-				}
-				else
-				{
-					Args.Add(TEXT("DestPath"), FText::FromName(DestinationPath));
-					Text = FText::Format(LOCTEXT("FolderAlreadyAssigned", "{SourceName} is already assigned to {DestPath}"), Args);
-				}
-				
-				return FDragValidationInfo(FActorDragDropGraphEdOp::ToolTip_IncompatibleGeneric, Text);
-			}
-
-			const FString DragFolderPath = DraggedFolder.ToString();
-			const FString LeafName = Leaf.ToString();
-			const FString DstFolderPath = DestinationPath.IsNone() ? FString() : DestinationPath.ToString();
-			const FString NewPath = DstFolderPath / LeafName;
-
-			if (FActorFolders::Get().GetFolderProperties(World, FName(*NewPath)))
-			{
-				// The folder already exists
-				FFormatNamedArguments Args;
-				Args.Add(TEXT("DragName"), FText::FromString(LeafName));
-				return FDragValidationInfo(FActorDragDropGraphEdOp::ToolTip_IncompatibleGeneric,
-					FText::Format(LOCTEXT("FolderAlreadyExistsRoot", "A folder called \"{DragName}\" already exists at this level"), Args));
-			}
-			else if (DragFolderPath == DstFolderPath || DstFolderPath.StartsWith(DragFolderPath + "/"))
-			{
-				// Cannot drag as a child of itself
-				FFormatNamedArguments Args;
-				Args.Add(TEXT("FolderPath"), FText::FromName(DraggedFolder));
-				return FDragValidationInfo(
-					FActorDragDropGraphEdOp::ToolTip_IncompatibleGeneric,
-					FText::Format(LOCTEXT("ChildOfItself", "Cannot move \"{FolderPath}\" to be a child of itself"), Args));
-			}
-		}
-	}
-
-	if (DraggedObjects.Actors)
+	if (FFolderTreeItem* FolderItem = Item.Pin()->CastTo<FFolderTreeItem>())
 	{
-		// Iterate over all the folders that have been dragged
-		for (auto WeakActor : DraggedObjects.Actors.GetValue())
+		if (FolderItem->IsValid())
 		{
-			AActor* Actor = WeakActor.Get();
-			if (Actor->IsChildActor())
-			{
-				return FDragValidationInfo(FActorDragDropGraphEdOp::ToolTip_IncompatibleGeneric, FText::Format(LOCTEXT("Error_AttachChildActor", "Cannot move {0} as it is a child actor."), FText::FromString(Actor->GetActorLabel())));
-			}
-			else if (Actor->GetFolderPath() == DestinationPath && !Actor->GetAttachParentActor())
-			{
-				FFormatNamedArguments Args;
-				Args.Add(TEXT("SourceName"), FText::FromString(Actor->GetActorLabel()));
-
-				FText Text;
-				if (DestinationPath.IsNone())
-				{
-					Text = FText::Format(LOCTEXT("FolderAlreadyAssignedRoot", "{SourceName} is already assigned to root"), Args);
-				}
-				else
-				{
-					Args.Add(TEXT("DestPath"), FText::FromName(DestinationPath));
-					Text = FText::Format(LOCTEXT("FolderAlreadyAssigned", "{SourceName} is already assigned to {DestPath}"), Args);
-				}
-				
-				return FDragValidationInfo(FActorDragDropGraphEdOp::ToolTip_IncompatibleGeneric, Text);
-			}
+			DataOut = FolderItem->Path;
+			return true;
 		}
 	}
-
-	// Everything else is a valid operation
-	if (DestinationPath.IsNone())
-	{
-		return FDragValidationInfo(FActorDragDropGraphEdOp::ToolTip_CompatibleGeneric, LOCTEXT("MoveToRoot", "Move to root"));
-	}
-	else
-	{
-		FFormatNamedArguments Args;
-		Args.Add(TEXT("DestPath"), FText::FromName(DestinationPath));
-		return FDragValidationInfo(FActorDragDropGraphEdOp::ToolTip_CompatibleGeneric, FText::Format(LOCTEXT("MoveInto", "Move into \"{DestPath}\""), Args));
-	}
+	return false;
 }
 
-void FFolderDropTarget::OnDrop(FDragDropPayload& DraggedObjects, UWorld& World, const FDragValidationInfo& ValidationInfo, TSharedRef<SWidget> DroppedOnWidget)
-{
-	const FScopedTransaction Transaction( LOCTEXT("MoveOutlinerItems", "Move World Outliner Items") );
-
-	if (DraggedObjects.Folders)
-	{
-		for (FName Folder : DraggedObjects.Folders.GetValue())
-		{
-			MoveFolderTo(Folder, DestinationPath, World);
-		}
-	}
-
-	// Set the folder path on all the dragged actors, and detach any that need to be moved
-	if (DraggedObjects.Actors)
-	{
-		TSet<AActor*> ParentActors;
-		TSet<AActor*> ChildActors;
-
-		for (auto WeakActor : DraggedObjects.Actors.GetValue())
-		{
-			AActor* Actor = WeakActor.Get();
-			if (Actor)
-			{
-				// First mark this object as a parent, then set its children's path
-				ParentActors.Add(Actor);
-				Actor->SetFolderPath(DestinationPath);
-
-				FActorEditorUtils::TraverseActorTree_ParentFirst(Actor, [&](AActor* InActor){
-					ChildActors.Add(InActor);
-					InActor->SetFolderPath(DestinationPath);
-					return true;
-				}, false);
-			}
-		}
-
-		// Detach parent actors
-		for (AActor* Parent : ParentActors)
-		{
-			auto* RootComp = Parent->GetRootComponent();
-
-			// We don't detach if it's a child of another that's been dragged
-			if (RootComp && RootComp->GetAttachParent() && !ChildActors.Contains(Parent))
-			{
-				if (AActor* OldParentActor = RootComp->GetAttachParent()->GetOwner())
-				{
-					OldParentActor->Modify();
-				}
-				RootComp->DetachFromComponent(FDetachmentTransformRules::KeepWorldTransform);
-			}
-		}
-	}
-}
 
 FFolderTreeItem::FFolderTreeItem(FName InPath)
-	: Path(InPath)
+	: ITreeItem(Type)
+	, Path(InPath)
 	, LeafName(GetFolderLeafName(InPath))
 {
 }
 
-FTreeItemPtr FFolderTreeItem::FindParent(const FTreeItemMap& ExistingItems) const
-{
-	const FName ParentPath = GetParentPath(Path);
-	if (!ParentPath.IsNone())
-	{
-		return ExistingItems.FindRef(ParentPath);
-	}
-
-	if (SharedData->RepresentingWorld)
-	{
-		return ExistingItems.FindRef(SharedData->RepresentingWorld);
-	}
-
-	return nullptr;
-}
-
-FTreeItemPtr FFolderTreeItem::CreateParent() const
-{
-	const FName ParentPath = GetParentPath(Path);
-	if (!ParentPath.IsNone())
-	{
-		return MakeShareable(new FFolderTreeItem(ParentPath));
-	}
-
-	if (SharedData->RepresentingWorld)
-	{
-		return MakeShareable(new FWorldTreeItem(SharedData->RepresentingWorld));
-	}
-
-	return nullptr;
-}
-
-void FFolderTreeItem::Visit(const ITreeItemVisitor& Visitor) const
-{
-	Visitor.Visit(*this);
-}
-
-void FFolderTreeItem::Visit(const IMutableTreeItemVisitor& Visitor)
+FFolderTreeItem::FFolderTreeItem(FName InPath, FTreeItemType InType)
+	: ITreeItem(InType)
+	, Path(InPath)
+	, LeafName(GetFolderLeafName(InPath))
 {
-	Visitor.Visit(*this);
 }
 
 FTreeItemID FFolderTreeItem::GetID() const
@@ -228,98 +56,22 @@ FString FFolderTreeItem::GetDisplayString() const
 	return LeafName.ToString();
 }
 
-int32 FFolderTreeItem::GetTypeSortPriority() const
-{
-	return ETreeItemSortOrder::Folder;
-}
-
-
 bool FFolderTreeItem::CanInteract() const
 {
 	return Flags.bInteractive;
 }
 
-void FFolderTreeItem::Delete(FName InNewParentPath)
-{
-	if (!SharedData->RepresentingWorld)
-	{
-		return;
-	}
-
-	struct FResetActorFolders : IMutableTreeItemVisitor
-	{
-		explicit FResetActorFolders(FName InParentPath) : ParentPath(InParentPath) {}
-
-		virtual void Visit(FActorTreeItem& ActorItem) const override
-		{
-			if (AActor* Actor = ActorItem.Actor.Get())
-			{
-				Actor->SetFolderPath_Recursively(ParentPath);
-			}
-		}
-		virtual void Visit(FFolderTreeItem& FolderItem) const override
-		{
-			UWorld* World = FolderItem.SharedData->RepresentingWorld;
-			check(World != nullptr);
-
-			MoveFolderTo(FolderItem.Path, ParentPath, *World);
-		}
-
-		FName ParentPath;
-	};
-
-	const FScopedTransaction Transaction( LOCTEXT("DeleteFolderTransaction", "Delete Folder") );
-
-	FResetActorFolders ResetFolders(InNewParentPath);
-	for (auto& Child : GetChildren())
-	{
-		Child.Pin()->Visit(ResetFolders);
-	}
-
-	FActorFolders::Get().DeleteFolder(*SharedData->RepresentingWorld, Path);
-}
-
-void FFolderTreeItem::CreateSubFolder(TWeakPtr<SSceneOutliner> WeakOutliner)
-{
-	auto Outliner = WeakOutliner.Pin();
-
-	if (Outliner.IsValid() && SharedData->RepresentingWorld)
-	{
-		const FScopedTransaction Transaction(LOCTEXT("UndoAction_CreateFolder", "Create Folder"));
-
-		const FName NewFolderName = FActorFolders::Get().GetDefaultFolderName(*SharedData->RepresentingWorld, Path);
-		FActorFolders::Get().CreateFolder(*SharedData->RepresentingWorld, NewFolderName);
-
-		// At this point the new folder will be in our newly added list, so select it and open a rename when it gets refreshed
-		Outliner->OnItemAdded(NewFolderName, ENewItemAction::Select | ENewItemAction::Rename);
-	}
-}
-
 void FFolderTreeItem::DuplicateHierarchy(TWeakPtr<SSceneOutliner> WeakOutliner)
 {
 	TSharedPtr<SSceneOutliner> Outliner = WeakOutliner.Pin();
 
-	if (Outliner.IsValid() && SharedData->RepresentingWorld)
+	if (Outliner.IsValid())
 	{
 		const FScopedTransaction Transaction(LOCTEXT("UndoAction_DuplicateHierarchy", "Duplicate Folder Hierarchy"));
 		Outliner->DuplicateFoldersHierarchy();
 	}
 }
 
-void FFolderTreeItem::OnExpansionChanged()
-{
-	if (!SharedData->RepresentingWorld)
-	{
-		return;
-	}
-
-	// Update the central store of folder properties with this folder's new expansion state
-	if (FActorFolderProps* Props = FActorFolders::Get().GetFolderProperties(*SharedData->RepresentingWorld, Path))
-	{
-		Props->bIsExpanded = Flags.bIsExpanded;
-	}
-}
-
 void FFolderTreeItem::GenerateContextMenu(UToolMenu* Menu, SSceneOutliner& Outliner)
 {
 	auto SharedOutliner = StaticCastSharedRef<SSceneOutliner>(Outliner.AsShared());
@@ -331,27 +83,6 @@ void FFolderTreeItem::GenerateContextMenu(UToolMenu* Menu, SSceneOutliner& Outli
 	Section.AddMenuEntry("DuplicateFolderHierarchy", LOCTEXT("DuplicateFolderHierarchy", "Duplicate Hierarchy"), FText(), FSlateIcon(), FUIAction(FExecuteAction::CreateSP(&Outliner, &SSceneOutliner::DuplicateFoldersHierarchy)));
 }
 
-void FFolderTreeItem::PopulateDragDropPayload(FDragDropPayload& Payload) const
-{
-	if (!Payload.Folders)
-	{
-		Payload.Folders = FFolderPaths();
-	}
-	Payload.Folders->Add(Path);
-}
-
-FDragValidationInfo FFolderTreeItem::ValidateDrop(FDragDropPayload& DraggedObjects, UWorld& World) const
-{
-	FFolderDropTarget Target(Path);
-	return Target.ValidateDrop(DraggedObjects, World);
-}
-
-void FFolderTreeItem::OnDrop(FDragDropPayload& DraggedObjects, UWorld& World, const FDragValidationInfo& ValidationInfo, TSharedRef<SWidget> DroppedOnWidget)
-{
-	FFolderDropTarget Target(Path);
-	return Target.OnDrop(DraggedObjects, World, ValidationInfo, DroppedOnWidget);
-}
-
 }	// namespace SceneOutliner
 
 #undef LOCTEXT_NAMESPACE
diff --git a/Engine/Source/Editor/SceneOutliner/Private/SOutlinerTreeView.cpp b/Engine/Source/Editor/SceneOutliner/Private/SOutlinerTreeView.cpp
--- a/Engine/Source/Editor/SceneOutliner/Private/SOutlinerTreeView.cpp
+++ b/Engine/Source/Editor/SceneOutliner/Private/SOutlinerTreeView.cpp
@@ -6,10 +6,11 @@
 #include "ISceneOutlinerColumn.h"
 #include "SceneOutlinerPublicTypes.h"
 #include "DragAndDrop/DecoratedDragDropOp.h"
-#include "DragAndDrop/ActorDragDropGraphEdOp.h"
 #include "SceneOutlinerDragDrop.h"
 #include "SSceneOutliner.h"
 
+#include "FolderTreeItem.h"
+
 #define LOCTEXT_NAMESPACE "SSceneOutliner"
 
 namespace SceneOutliner
@@ -26,10 +27,6 @@ namespace SceneOutliner
 				auto* OutlinerOp = static_cast<FSceneOutlinerDragDropOp*>(Operation);
 				OutlinerOp->SetTooltip(ValidationInfo.ValidationText, Icon);
 			}
-			else if (Operation->IsOfType<FActorDragDropGraphEdOp>())
-			{
-				static_cast<FActorDragDropGraphEdOp*>(Operation)->SetToolTip(ValidationInfo.TooltipType, ValidationInfo.ValidationText);
-			}
 			else if (Operation->IsOfType<FDecoratedDragDropOp>())
 			{
 				auto* DecoratedOp = static_cast<FDecoratedDragDropOp*>(Operation);
@@ -59,7 +56,8 @@ namespace SceneOutliner
 		auto TablePtr = Table.Pin();
 		if (TablePtr.IsValid() && MouseEvent.IsMouseButtonDown( EKeys::LeftMouseButton ))
 		{
-			auto Operation = SceneOutliner::CreateDragDropOperation(TablePtr->GetSelectedItems());
+			auto Operation = TablePtr->GetOutlinerPtr().Pin()->CreateDragDropOperation(TablePtr->GetSelectedItems());
+
 			if (Operation.IsValid())
 			{
 				return FReply::Handled().BeginDragDrop(Operation.ToSharedRef());
@@ -69,33 +67,27 @@ namespace SceneOutliner
 		return FReply::Unhandled();
 	}
 
-	FReply HandleDrop(TSharedPtr<ISceneOutliner> SceneOutlinerPtr, const FDragDropEvent& DragDropEvent, IDropTarget& DropTarget, FDragValidationInfo& ValidationInfo, bool bApplyDrop = false)
+	FReply HandleDrop(TSharedPtr<SSceneOutliner> SceneOutlinerPtr, const FDragDropEvent& DragDropEvent, ITreeItem& DropTarget, FDragValidationInfo& ValidationInfo, bool bApplyDrop = false)
 	{
 		if (!SceneOutlinerPtr.IsValid())
 		{
 			return FReply::Unhandled();
 		}
 
-		const FSharedOutlinerData& SharedData = SceneOutlinerPtr->GetSharedData();
-		if (SharedData.Mode != ESceneOutlinerMode::ActorBrowsing)
-		{
-			return FReply::Unhandled();
-		}
-
-		// Don't handle this if we're not showing a hierarchy, not in browsing mode, or the drop operation is not applicable
-		if (!SharedData.bShowParentTree || !SharedData.RepresentingWorld)
+		// Don't handle this if the scene outliner is not in a mode which supports drag and drop
+		if (!SceneOutlinerPtr->CanSupportDragAndDrop())
 		{
 			return FReply::Unhandled();
 		}
 
 		FDragDropPayload DraggedObjects;
 		// Validate now to make sure we don't doing anything we shouldn't
-		if (!DraggedObjects.ParseDrag(*DragDropEvent.GetOperation()))
+		if (!SceneOutlinerPtr->ParseDragDrop(DraggedObjects, *DragDropEvent.GetOperation()))
 		{
 			return FReply::Unhandled();
 		}
 
-		ValidationInfo = DropTarget.ValidateDrop(DraggedObjects, *SharedData.RepresentingWorld);
+		ValidationInfo = SceneOutlinerPtr->ValidateDrop(StaticCast<ITreeItem&>(DropTarget), DraggedObjects);
 
 		if (!ValidationInfo.IsValid())
 		{
@@ -105,13 +97,13 @@ namespace SceneOutliner
 
 		if (bApplyDrop)
 		{
-			DropTarget.OnDrop(DraggedObjects, *SharedData.RepresentingWorld, ValidationInfo, SceneOutlinerPtr.ToSharedRef());
+			SceneOutlinerPtr->OnDropPayload(DropTarget, DraggedObjects, ValidationInfo);
 		}
 
 		return FReply::Handled();
 	}
 
-	FReply HandleDropFromWeak(TWeakPtr<ISceneOutliner> SceneOutlinerWeak, const FDragDropEvent& DragDropEvent, IDropTarget& DropTarget, FDragValidationInfo& ValidationInfo, bool bApplyDrop = false)
+	FReply HandleDropFromWeak(TWeakPtr<SSceneOutliner> SceneOutlinerWeak, const FDragDropEvent& DragDropEvent, ITreeItem& DropTarget, FDragValidationInfo& ValidationInfo, bool bApplyDrop = false)
 	{
 		return HandleDrop(SceneOutlinerWeak.Pin(), DragDropEvent, DropTarget, ValidationInfo, bApplyDrop);
 	}
@@ -135,7 +127,7 @@ namespace SceneOutliner
 	{
 		FDragValidationInfo ValidationInfo = FDragValidationInfo::Invalid();
 
-		FFolderDropTarget DropTarget(NAME_None);
+		FFolderTreeItem DropTarget(NAME_None);
 
 		auto Reply = HandleDropFromWeak(SceneOutlinerWeak, DragDropEvent, DropTarget, ValidationInfo);
 		UpdateOperationDecorator(DragDropEvent, ValidationInfo);
@@ -154,7 +146,7 @@ namespace SceneOutliner
 	FReply SOutlinerTreeView::OnDrop(const FGeometry& MyGeometry, const FDragDropEvent& DragDropEvent)
 	{
 		FDragValidationInfo ValidationInfo = FDragValidationInfo::Invalid();
-		FFolderDropTarget DropTarget(NAME_None);
+		FFolderTreeItem DropTarget(NAME_None);
 
 		return HandleDropFromWeak(SceneOutlinerWeak, DragDropEvent, DropTarget, ValidationInfo, true);
 	}
@@ -165,11 +157,6 @@ namespace SceneOutliner
 		auto SceneOutlinerPtr = SceneOutlinerWeak.Pin();
 		if (ItemPtr.IsValid() && SceneOutlinerPtr.IsValid())
 		{
-			if (SceneOutlinerPtr->GetOnDropOnItem())
-			{
-				return SceneOutlinerPtr->GetOnDropOnItem()(DragDropEvent, *ItemPtr.Get());
-			}
-
 			FDragValidationInfo ValidationInfo = FDragValidationInfo::Invalid();
 			return HandleDrop(SceneOutlinerPtr, DragDropEvent, *ItemPtr, ValidationInfo, true);
 		}
@@ -183,17 +170,10 @@ namespace SceneOutliner
 		auto SceneOutlinerPtr = SceneOutlinerWeak.Pin();
 		if (ItemPtr.IsValid() && SceneOutlinerPtr.IsValid())
 		{
-			if (SceneOutlinerPtr->GetOnDragEnterItem())
-			{
-				SceneOutlinerPtr->GetOnDragEnterItem()(DragDropEvent, *ItemPtr.Get());
-			}
-			else
-			{
-				FDragValidationInfo ValidationInfo = FDragValidationInfo::Invalid();
+			FDragValidationInfo ValidationInfo = FDragValidationInfo::Invalid();
 
-				HandleDrop(StaticCastSharedPtr<ISceneOutliner>(SceneOutlinerPtr), DragDropEvent, *ItemPtr, ValidationInfo, false);
-				UpdateOperationDecorator(DragDropEvent, ValidationInfo);
-			}
+			HandleDrop(SceneOutlinerPtr, DragDropEvent, *ItemPtr, ValidationInfo, false);
+			UpdateOperationDecorator(DragDropEvent, ValidationInfo);
 		}
 	}
 
@@ -201,14 +181,8 @@ namespace SceneOutliner
 	{
 		auto ItemPtr = Item.Pin();
 		auto SceneOutlinerPtr = SceneOutlinerWeak.Pin();
-		if (ItemPtr.IsValid() && SceneOutlinerPtr.IsValid() && SceneOutlinerPtr->GetOnDragLeaveItem())
-		{
-			SceneOutlinerPtr->GetOnDragLeaveItem()(DragDropEvent, *ItemPtr.Get());
-		}
-		else
-		{
-			ResetOperationDecorator(DragDropEvent);
-		}
+
+		ResetOperationDecorator(DragDropEvent);
 	}
 
 	FReply SSceneOutlinerTreeRow::OnDragOver(const FGeometry& MyGeometry, const FDragDropEvent& DragDropEvent)
@@ -216,18 +190,11 @@ namespace SceneOutliner
 		auto SceneOutlinerPtr = SceneOutlinerWeak.Pin();
 		if (SSceneOutliner* SceneOutliner = SceneOutlinerPtr.Get())
 		{
-			if (SceneOutliner->GetSharedData().Mode == ESceneOutlinerMode::Custom)
+			if (const auto* ItemPtr = Item.Pin().Get())
 			{
-				if (const auto* ItemPtr = Item.Pin().Get())
-				{
-					if (SceneOutliner->GetOnDragOverItem())
-					{
-						return SceneOutliner->GetOnDragOverItem()(DragDropEvent, *ItemPtr);
-					}
-				}
-
-				return FReply::Unhandled();
+				return SceneOutliner->OnDragOverItem(DragDropEvent, *ItemPtr);
 			}
+			return FReply::Unhandled();
 		}
 
 		return FReply::Handled();
@@ -322,29 +289,8 @@ namespace SceneOutliner
 		auto Args = FSuperRowType::FArguments()
 			.Style(&FEditorStyle::Get().GetWidgetStyle<FTableRowStyle>("SceneOutliner.TableViewRow"));
 
-		// We only support standard drag and drop when in actor browsing mode
-		if (SceneOutliner->GetSharedData().Mode == ESceneOutlinerMode::ActorBrowsing)
-		{
-			Args.OnDragDetected_Static(SceneOutliner::OnDragDetected, TWeakPtr<SOutlinerTreeView>(OutlinerTreeView));
-		}
-		// If we are in custom mode allow support of custom on detect drag
-		else if (SceneOutliner->GetSharedData().Mode == ESceneOutlinerMode::Custom)
-		{
-			Args.OnDragDetected_Lambda([SceneOutlinerWeak = SceneOutlinerWeak, ItemPtr = Item]( const FGeometry& MyGeometry, const FPointerEvent& MouseEvent)
-				{
-					if (const ITreeItem* ItemRawPtr = ItemPtr.Pin().Get())
-					{
-						if (const SSceneOutliner* SceneOutliner = SceneOutlinerWeak.Pin().Get())
-						{
-							if (SceneOutliner->GetOnItemDragDetected())
-							{
-								return SceneOutliner->GetOnItemDragDetected()(*ItemRawPtr);
-							}
-						}
-					}
-					return FReply::Unhandled();
-				});
-		}
+
+		Args.OnDragDetected_Static(SceneOutliner::OnDragDetected, TWeakPtr<SOutlinerTreeView>(OutlinerTreeView));
 
 		SMultiColumnTableRow<FTreeItemPtr>::Construct(Args, OutlinerTreeView);
 	}
