SceneOutliner refactor

diff --git a/Engine/Source/Editor/SceneOutliner/Private/SSceneOutliner.cpp b/Engine/Source/Editor/SceneOutliner/Private/SSceneOutliner.cpp
--- a/Engine/Source/Editor/SceneOutliner/Private/SSceneOutliner.cpp
+++ b/Engine/Source/Editor/SceneOutliner/Private/SSceneOutliner.cpp
@@ -4,7 +4,6 @@
 
 #include "EdMode.h"
 #include "Editor.h"
-#include "Editor/GroupActor.h"
 #include "Editor/UnrealEdEngine.h"
 #include "EditorModeManager.h"
 #include "EditorStyleSet.h"
@@ -14,7 +13,6 @@
 #include "Framework/Application/SlateApplication.h"
 #include "Framework/Commands/UIAction.h"
 #include "Framework/MultiBox/MultiBoxBuilder.h"
-#include "GameFramework/WorldSettings.h"
 #include "HAL/PlatformApplicationMisc.h"
 #include "ISceneOutlinerColumn.h"
 #include "Kismet2/ComponentEditorUtils.h"
@@ -23,7 +21,6 @@
 #include "SceneOutlinerDelegates.h"
 #include "SceneOutlinerFilters.h"
 #include "SceneOutlinerModule.h"
-#include "SceneOutlinerSettings.h"
 #include "ScopedTransaction.h"
 #include "Textures/SlateIcon.h"
 #include "ToolMenus.h"
@@ -35,18 +32,12 @@
 #include "Widgets/Input/SSearchBox.h"
 #include "Widgets/Layout/SSeparator.h"
 #include "Widgets/SOverlay.h"
-
-#include "SceneOutlinerDragDrop.h"
 #include "SceneOutlinerMenuContext.h"
 
-#include "ActorEditorUtils.h"
-#include "LevelUtils.h"
-
-#include "EditorActorFolders.h"
 #include "Framework/Notifications/NotificationManager.h"
 #include "Widgets/Notifications/SNotificationList.h"
-#include "Features/IModularFeatures.h"
-#include "ISceneOutlinerTraversal.h"
+#include "ISceneOutlinerMode.h"
+#include "FolderTreeItem.h"
 #include "UObject/PackageReload.h"
 
 DEFINE_LOG_CATEGORY_STATIC(LogSceneOutliner, Log, All);
@@ -58,227 +49,27 @@ DEFINE_LOG_CATEGORY_STATIC(LogSceneOutliner, Log, All);
 
 namespace SceneOutliner
 {
-	DECLARE_MULTICAST_DELEGATE(FOnSharedSettingsChanged);
-	FOnSharedSettingsChanged OnSharedSettingChangedDelegate;
-
-	FText GetWorldDescription(UWorld* World)
-	{
-		FText Description;
-		if(World)
-		{
-			FText PostFix;
-			const FWorldContext* WorldContext = nullptr;
-			for (const FWorldContext& Context : GEngine->GetWorldContexts())
-			{
-				if(Context.World() == World)
-				{
-					WorldContext = &Context;
-					break;
-				}
-			}
-
-			if (World->WorldType == EWorldType::PIE)
-			{
-				switch(World->GetNetMode())
-				{
-					case NM_Client:
-						if (WorldContext)
-						{
-							PostFix = FText::Format(LOCTEXT("ClientPostfixFormat", "(Client {0})"), FText::AsNumber(WorldContext->PIEInstance - 1));
-						}
-						else
-						{
-							PostFix = LOCTEXT("ClientPostfix", "(Client)");
-						}
-						break;
-					case NM_DedicatedServer:
-					case NM_ListenServer:
-						PostFix = LOCTEXT("ServerPostfix", "(Server)");
-						break;
-					case NM_Standalone:
-						PostFix = LOCTEXT("PlayInEditorPostfix", "(Play In Editor)");
-						break;
-				}
-			}
-			else if(World->WorldType == EWorldType::Editor)
-			{
-				PostFix = LOCTEXT("EditorPostfix", "(Editor)");
-			}
-
-			Description = FText::Format(LOCTEXT("WorldFormat", "{0} {1}"), FText::FromString(World->GetFName().GetPlainNameString()), PostFix);	
-		}
-
-		return Description;
-	}
-
-	TSharedPtr< FOutlinerFilter > CreateSelectedActorFilter()
-	{
-		auto* Filter = new FOutlinerPredicateFilter(FActorFilterPredicate::CreateStatic([](const AActor* InActor){	return InActor->IsSelected(); }), EDefaultFilterBehaviour::Fail);
-
-		// If anything fails this filter, make it non interactive. We don't want to allow selection of implicitly included parents which might nuke the actor selection.
-		Filter->FailedItemState = EFailedFilterState::NonInteractive;
-		return MakeShareable( Filter );
-	}
-
-	TSharedPtr< FOutlinerFilter > CreateHideTemporaryActorsFilter()
-	{
-		return MakeShareable( new FOutlinerPredicateFilter( FActorFilterPredicate::CreateStatic( []( const AActor* InActor ){			
-			return ((InActor->GetWorld() && InActor->GetWorld()->WorldType != EWorldType::PIE) || GEditor->ObjectsThatExistInEditorWorld.Get(InActor)) && !InActor->HasAnyFlags(EObjectFlags::RF_Transient);
-		} ), EDefaultFilterBehaviour::Pass ) );
-	}
-
-	TSharedPtr< FOutlinerFilter > CreateIsInCurrentLevelFilter()
-	{
-		struct FOnlyCurrentLevelFilter : FOutlinerFilter
-		{
-			FOnlyCurrentLevelFilter() : FOutlinerFilter(EDefaultFilterBehaviour::Fail, EFailedFilterState::Interactive) {}
-
-			virtual bool PassesFilter(const AActor* InActor) const override
-			{
-				if (InActor->GetWorld())
-				{
-					return InActor->GetLevel() == InActor->GetWorld()->GetCurrentLevel();
-				}
-				
-				return false;
-			}
-		};
-
-		return MakeShared<FOnlyCurrentLevelFilter>();
-	}
-
-	TSharedPtr< FOutlinerFilter > CreateShowActorComponentsFilter()
-	{
-		TSharedRef< FOutlinerPredicateFilter > Filter = MakeShared<FOutlinerPredicateFilter>(FActorFilterPredicate::CreateStatic([](const AActor* InActor) {	return InActor!= nullptr; }), EDefaultFilterBehaviour::Fail);
-		Filter->ComponentPred = FComponentFilterPredicate::CreateStatic([](const UActorComponent* InComponent) {return Cast<UPrimitiveComponent>(InComponent) != nullptr; });
-
-		// If anything fails this filter, make it non interactive. We don't want to allow selection of implicitly included parents which might nuke the actor selection.
-		Filter->FailedItemState = EFailedFilterState::NonInteractive;
-		return Filter;
-	}
-
-	struct FItemSelection : IMutableTreeItemVisitor
-	{
-		mutable TArray<FActorTreeItem*> Actors;
-		mutable TArray<FWorldTreeItem*> Worlds;
-		mutable TArray<FFolderTreeItem*> Folders;
-		mutable TArray<FComponentTreeItem*> Components;
-		mutable TArray<FSubComponentTreeItem*> SubComponents;
-
-		FItemSelection()
-		{}
-
-		FItemSelection(SOutlinerTreeView& Tree)
-		{
-			for (auto& Item : Tree.GetSelectedItems())
-			{
-				Item->Visit(*this);
-			}
-		}
-
-		TArray<TWeakObjectPtr<AActor>> GetWeakActors() const
-		{
-			TArray<TWeakObjectPtr<AActor>> ActorArray;
-			for (const auto* ActorItem : Actors)
-			{
-				if (ActorItem->Actor.IsValid())
-				{
-					ActorArray.Add(ActorItem->Actor);
-				}
-			}
-			return ActorArray;
-		}
-
-		TArray<AActor*> GetActorPtrs() const
-		{
-			TArray<AActor*> ActorArray;
-			for (const auto* ActorItem : Actors)
-			{
-				if (AActor* Actor = ActorItem->Actor.Get())
-				{
-					ActorArray.Add(Actor);
-				}
-			}
-
-			// if we select a component then we are actually wanting the owning actor to be selected
-			for (const auto* ComponentItem : Components)
-			{
-				if (UActorComponent* ActorComponent = ComponentItem->Component.Get())
-				{
-					AActor* Actor = ActorComponent->GetOwner();
-					ActorArray.Add(Actor);
-				}
-			}
-
-			// if we select a sub item from within a component then we are actually wanting the owning actor to be selected
-			for (const auto* SubItem : SubComponents)
-			{
-				if (UActorComponent* ActorComponent = SubItem->ParentComponent.Get())
-				{
-					AActor* Actor = ActorComponent->GetOwner();
-					ActorArray.Add(Actor);
-				}
-			}
-
-			return ActorArray;
-		}
-
-	private:
-		virtual void Visit(FActorTreeItem& ActorItem) const override
-		{
-			Actors.Add(&ActorItem);
-		}
-		virtual void Visit(FWorldTreeItem& WorldItem) const override
-		{
-			Worlds.Add(&WorldItem);
-		}
-		virtual void Visit(FFolderTreeItem& FolderItem) const override
-		{
-			Folders.Add(&FolderItem);
-		}
-		virtual void Visit(FComponentTreeItem& ComponentItem) const override
-		{
-			Components.Add(&ComponentItem);
-		}
-		virtual void Visit(FSubComponentTreeItem& SubComponentItem) const override
-		{
-			SubComponents.Add(&SubComponentItem);
-		}
-
-	};
-
-	void SSceneOutliner::Construct( const FArguments& InArgs, const FInitializationOptions& InInitOptions )
+	void SSceneOutliner::Construct(const FArguments& InArgs, const FInitializationOptions& InInitOptions)
 	{
 		// Copy over the shared data from the initialization options
-		static_cast<FSharedDataBase&>(*SharedData) = static_cast<const FSharedDataBase&>(InInitOptions);
+		static_cast<FSharedOutlinerData&>(*SharedData) = static_cast<const FSharedOutlinerData&>(InInitOptions);
 
-		SelectionMode = (SharedData->Mode == ESceneOutlinerMode::ActorBrowsing) ? ESelectionMode::Multi : ESelectionMode::Single;
-
-		OnItemPicked = InArgs._OnItemPickedDelegate;
+		// We use the filter collection provided, otherwise we create our own
+		Filters = InInitOptions.Filters.IsValid() ? InInitOptions.Filters : MakeShareable(new FOutlinerFilters);
 
-		PRAGMA_DISABLE_DEPRECATION_WARNINGS
-		if (InInitOptions.OnSelectionChanged.IsBound())
-		{
-			FSceneOutlinerDelegates::Get().SelectionChanged.Add(InInitOptions.OnSelectionChanged);
-		}
-		PRAGMA_ENABLE_DEPRECATION_WARNINGS
+		check(InInitOptions.ModeFactory.IsBound());
+		Mode = InInitOptions.ModeFactory.Execute(this);
+		check(Mode);
 
 		bFullRefresh = true;
 		bNeedsRefresh = true;
 		bNeedsColumRefresh = true;
 		bIsReentrant = false;
 		bSortDirty = true;
-		bActorSelectionDirty = SharedData->Mode == ESceneOutlinerMode::ActorBrowsing;
-		FilteredActorCount = 0;
+		bSelectionDirty = true;
 		SortOutlinerTimer = 0.0f;
 		bPendingFocusNextFrame = InInitOptions.bFocusSearchBoxWhenOpened;
 
-		// Use the variable for the User Chosen World to enforce the Specified World To Display
-		if ( InInitOptions.SpecifiedWorldToDisplay )
-		{
-			SharedData->UserChosenWorld = InInitOptions.SpecifiedWorldToDisplay;
-		}
-
 		SortByColumn = FBuiltInColumnTypes::Label();
 		SortMode = EColumnSortMode::Ascending;
 
@@ -293,30 +84,13 @@ namespace SceneOutliner
 
 		TSharedRef<SVerticalBox> VerticalBox = SNew(SVerticalBox);
 
-		// We use the filter collection provided, otherwise we create our own
-		Filters = InInitOptions.Filters.IsValid() ? InInitOptions.Filters : MakeShareable(new FOutlinerFilters);
-	
-		// Add additional filters
-		if (SharedData->Mode == ESceneOutlinerMode::ActorBrowsing)
+		for (auto& ModeFilterInfo : Mode->GetFilterInfos())
 		{
-			FSceneOutlinerModule& SceneOutlinerModule = FModuleManager::LoadModuleChecked< FSceneOutlinerModule >("SceneOutliner");
-
-			for (auto& OutlinerFilterInfo : SceneOutlinerModule.OutlinerFilterInfoMap)
-			{
-				OutlinerFilterInfo.Value.InitFilter(Filters);
-			}
+			ModeFilterInfo.Value.InitFilter(Filters);
 		}
 
 		SearchBoxFilter->OnChanged().AddSP( this, &SSceneOutliner::FullRefresh );
-		Filters->OnChanged().AddSP( this, &SSceneOutliner::FullRefresh );
-
-		//Apply filters based on global preferences
-		if (SharedData->Mode == ESceneOutlinerMode::ActorBrowsing)
-		{
-			ApplyShowOnlySelectedFilter( IsShowingOnlySelected() );
-			ApplyHideTemporaryActorsFilter( IsHidingTemporaryActors() );
-			ApplyShowOnlyCurrentLevelFilter( IsShowingOnlyCurrentLevel() );
-		}
+		Filters->OnChanged().AddSP(this, &SSceneOutliner::FullRefresh);
 
 		HeaderRowWidget =
 			SNew( SHeaderRow )
@@ -343,7 +117,8 @@ namespace SceneOutliner
 			.OnTextCommitted( this, &SSceneOutliner::OnFilterTextCommitted )
 		];
 
-		if (SharedData->Mode == ESceneOutlinerMode::ActorBrowsing && InInitOptions.bShowCreateNewFolder)
+		
+		if (Mode->SupportsCreateNewFolder() && InInitOptions.bShowCreateNewFolder)
 		{
 			Toolbar->AddSlot()
 				.VAlign(VAlign_Center)
@@ -352,7 +127,7 @@ namespace SceneOutliner
 				[
 					SNew(SButton)
 					.ButtonStyle(FAppStyle::Get(), "SimpleButton")
-					.ToolTipText(LOCTEXT("CreateFolderToolTip", "Create a new folder containing the current actor selection"))
+					.ToolTipText(LOCTEXT("CreateFolderToolTip", "Create a new folder containing the current selection"))
 					.OnClicked(this, &SSceneOutliner::OnCreateFolderClicked)
 					[
 						SNew(SImage)
@@ -363,7 +138,7 @@ namespace SceneOutliner
 		}
 
 
-		if (SharedData->Mode == ESceneOutlinerMode::ActorBrowsing)
+		if (Mode->ShowViewButton())
 		{
 			// View mode combo button
 			Toolbar->AddSlot()
@@ -372,7 +147,7 @@ namespace SceneOutliner
 			[
 				SAssignNew( ViewOptionsComboButton, SComboButton )
 				.ComboButtonStyle( FAppStyle::Get(), "SimpleComboButton" ) // Use the tool bar item style for this button
-				.OnGetMenuContent( this, &SSceneOutliner::GetViewButtonContent, false, !InInitOptions.SpecifiedWorldToDisplay)
+				.OnGetMenuContent( this, &SSceneOutliner::GetViewButtonContent, Mode->ShowFilterOptions())
 				.HasDownArrow(false)
 				.ButtonContent()
 				[
@@ -413,8 +188,7 @@ namespace SceneOutliner
 			[
 				SAssignNew( OutlinerTreeView, SOutlinerTreeView, StaticCastSharedRef<SSceneOutliner>(AsShared()) )
 
-				// multi-select if we're in browsing mode, 
-				// single-select if we're in picking mode,
+				// Determined by the mode
 				.SelectionMode( this, &SSceneOutliner::GetSelectionMode )
 
 				// Point the tree to our array of root-level items.  Whenever this changes, we'll call RequestTreeRefresh()
@@ -452,9 +226,10 @@ namespace SceneOutliner
 			]
 		];
 
-		if (SharedData->Mode == ESceneOutlinerMode::ActorBrowsing)
+
+		// Bottom panel status bar, if enabled by the mode
+		if (Mode->ShowStatusBar())
 		{
-			// Bottom panel
 			VerticalBox->AddSlot()
 			.AutoHeight()
 			[
@@ -467,120 +242,23 @@ namespace SceneOutliner
 				.Padding(14, 9)
 				[
 					SNew( STextBlock )
-					.Text( this, &SSceneOutliner::GetFilterStatusText )
-					.ColorAndOpacity( this, &SSceneOutliner::GetFilterStatusTextColor )
-				]
-
-		];
-		} //end if (SharedData->Mode == ESceneOutlinerMode::ActorBrowsing)
-		else if ( !InInitOptions.SpecifiedWorldToDisplay )
-		{
-			// Bottom panel
-			VerticalBox->AddSlot()
-			.AutoHeight()
-			[
-				SNew(SHorizontalBox)
-
-				// World picker combo button
-				+SHorizontalBox::Slot()
-				.FillWidth(1.0f)
-				.HAlign(HAlign_Right)
-				[
-					SAssignNew( ViewOptionsComboButton, SComboButton )
-					.ContentPadding(0)
-					.ForegroundColor( this, &SSceneOutliner::GetViewButtonForegroundColor )
-					.ButtonStyle( FEditorStyle::Get(), "ToggleButton" ) // Use the tool bar item style for this button
-					.OnGetMenuContent( this, &SSceneOutliner::GetViewButtonContent, true, !InInitOptions.SpecifiedWorldToDisplay)
-					.ButtonContent()
-					[
-						SNew(SHorizontalBox)
-
-						+SHorizontalBox::Slot()
-						.AutoWidth()
-						.VAlign(VAlign_Center)
-						[
-							SNew(SImage).Image( FEditorStyle::GetBrush("SceneOutliner.World") )
-						]
-
-						+SHorizontalBox::Slot()
-						.AutoWidth()
-						.Padding(2, 0, 0, 0)
-						.VAlign(VAlign_Center)
-						[
-							SNew(STextBlock).Text( LOCTEXT("ChooseWorldMenu", "Choose World") )
-						]
-					]
+					.Text(this, &SSceneOutliner::GetFilterStatusText)
+					.ColorAndOpacity(this, &SSceneOutliner::GetFilterStatusTextColor)
 				]
 			];
 		}
 
-
 		// Don't allow tool-tips over the header
 		HeaderRowWidget->EnableToolTipForceField( true );
 
-
 		// Populate our data set
 		Populate();
 
-		// We only synchronize selection when in actor browsing mode
-		if (SharedData->Mode == ESceneOutlinerMode::ActorBrowsing)
-		{
-			// Populate and register to find out when the level's selection changes
-			OnLevelSelectionChanged( NULL );
-			USelection::SelectionChangedEvent.AddRaw(this, &SSceneOutliner::OnLevelSelectionChanged);
-			USelection::SelectObjectEvent.AddRaw(this, &SSceneOutliner::OnLevelSelectionChanged);
-
-			// Capture selection changes of bones from mesh selection in fracture tools
-			FSceneOutlinerDelegates::Get().OnComponentSelectionChanged.AddRaw(this, &SSceneOutliner::OnComponentSelectionChanged);
-			FSceneOutlinerDelegates::Get().OnComponentsUpdated.AddRaw(this, &SSceneOutliner::OnComponentsUpdated);
-		}
-
-		// Register to find out when actors are added or removed
-		// @todo outliner: Might not catch some cases (see: CALLBACK_ActorPropertiesChange, CALLBACK_LayerChange, CALLBACK_LevelDirtied, CALLBACK_OnActorMoved, CALLBACK_UpdateLevelsForAllActors)
-		FEditorDelegates::MapChange.AddSP( this, &SSceneOutliner::OnMapChange );
-		FEditorDelegates::NewCurrentLevel.AddSP( this, &SSceneOutliner::OnNewCurrentLevel );
-		GEngine->OnLevelActorListChanged().AddSP( this, &SSceneOutliner::OnLevelActorListChanged );
-		FWorldDelegates::LevelAddedToWorld.AddSP( this, &SSceneOutliner::OnLevelAdded );
-		FWorldDelegates::LevelRemovedFromWorld.AddSP( this, &SSceneOutliner::OnLevelRemoved );
-
-		GEngine->OnLevelActorAdded().AddSP( this, &SSceneOutliner::OnLevelActorsAdded );
-		GEngine->OnLevelActorDetached().AddSP( this, &SSceneOutliner::OnLevelActorsDetached );
-		GEngine->OnLevelActorFolderChanged().AddSP( this, &SSceneOutliner::OnLevelActorFolderChanged );
-
-		GEngine->OnLevelActorDeleted().AddSP( this, &SSceneOutliner::OnLevelActorsRemoved );
-		GEngine->OnLevelActorAttached().AddSP( this, &SSceneOutliner::OnLevelActorsAttached );
-
-		GEngine->OnLevelActorRequestRename().AddSP( this, &SSceneOutliner::OnLevelActorsRequestRename );
-
-		// Register to update when an undo/redo operation has been called to update our list of actors
+		// Register to update when an undo/redo operation has been called to update our list of items
 		GEditor->RegisterForUndo( this );
 
 		// Register to be notified when properties are edited
-		FCoreDelegates::OnActorLabelChanged.AddRaw(this, &SSceneOutliner::OnActorLabelChanged);
 		FCoreUObjectDelegates::OnPackageReloaded.AddRaw(this, &SSceneOutliner::OnAssetReloaded);
-
-		auto& Folders = FActorFolders::Get();
-		Folders.OnFolderCreate.AddSP(this, &SSceneOutliner::OnBroadcastFolderCreate);
-		Folders.OnFolderMove.AddSP(this, &SSceneOutliner::OnBroadcastFolderMove);
-		Folders.OnFolderDelete.AddSP(this, &SSceneOutliner::OnBroadcastFolderDelete);
-
-		if ( SharedData->Mode == ESceneOutlinerMode::ActorBrowsing )
-		{
-			// Only the actor browsing mode seems to need those
-			FEditorDelegates::OnEditCutActorsBegin.AddSP(this, &SSceneOutliner::OnEditCutActorsBegin);
-			FEditorDelegates::OnEditCutActorsEnd.AddSP(this, &SSceneOutliner::OnEditCutActorsEnd);
-			FEditorDelegates::OnEditCopyActorsBegin.AddSP(this, &SSceneOutliner::OnEditCopyActorsBegin);
-			FEditorDelegates::OnEditCopyActorsEnd.AddSP(this, &SSceneOutliner::OnEditCopyActorsEnd);
-			FEditorDelegates::OnEditPasteActorsBegin.AddSP(this, &SSceneOutliner::OnEditPasteActorsBegin);
-			FEditorDelegates::OnEditPasteActorsEnd.AddSP(this, &SSceneOutliner::OnEditPasteActorsEnd);
-			FEditorDelegates::OnDuplicateActorsBegin.AddSP(this, &SSceneOutliner::OnDuplicateActorsBegin);
-			FEditorDelegates::OnDuplicateActorsEnd.AddSP(this, &SSceneOutliner::OnDuplicateActorsEnd);
-			FEditorDelegates::OnDeleteActorsBegin.AddSP(this, &SSceneOutliner::OnDeleteActorsBegin);
-			FEditorDelegates::OnDeleteActorsEnd.AddSP(this, &SSceneOutliner::OnDeleteActorsEnd);
-		}
-
-		SetUseSharedSceneOutlinerSettings(SharedData->Mode == ESceneOutlinerMode::Custom);
-		OnSharedSettingChangedDelegate.AddSP(this, &SSceneOutliner::OnSharedSettingChanged);
 	}
 
 	void SSceneOutliner::SetupColumns(SHeaderRow& HeaderRow)
@@ -651,51 +329,18 @@ namespace SceneOutliner
 
 	SSceneOutliner::~SSceneOutliner()
 	{
-		FSceneOutlinerDelegates::Get().OnComponentSelectionChanged.RemoveAll(this);
-		FSceneOutlinerDelegates::Get().OnComponentsUpdated.RemoveAll(this);
-
-		// We only synchronize selection when in actor browsing mode
-		if( SharedData->Mode == ESceneOutlinerMode::ActorBrowsing )
-		{
-			USelection::SelectionChangedEvent.RemoveAll(this);
-			USelection::SelectObjectEvent.RemoveAll(this);
-		}
-		FEditorDelegates::MapChange.RemoveAll( this );
-		FEditorDelegates::NewCurrentLevel.RemoveAll( this );
+		Mode->GetHierarchy()->OnHierarchyChanged().RemoveAll(this);
+		delete Mode;
 
 		if(GEngine)
 		{
-			GEngine->OnLevelActorListChanged().RemoveAll(this);
 			GEditor->UnregisterForUndo(this);
 		}
 
 		SearchBoxFilter->OnChanged().RemoveAll( this );
-		Filters->OnChanged().RemoveAll( this );
-		
-		FWorldDelegates::LevelAddedToWorld.RemoveAll( this );
-		FWorldDelegates::LevelRemovedFromWorld.RemoveAll( this );
+		Filters->OnChanged().RemoveAll(this);
 
-		FCoreDelegates::OnActorLabelChanged.RemoveAll(this);
 		FCoreUObjectDelegates::OnPackageReloaded.RemoveAll(this);
-
-		if (FActorFolders::IsAvailable())
-		{
-			auto& Folders = FActorFolders::Get();
-			Folders.OnFolderCreate.RemoveAll(this);
-			Folders.OnFolderMove.RemoveAll(this);
-			Folders.OnFolderDelete.RemoveAll(this);
-		}
-
-		FEditorDelegates::OnEditCutActorsBegin.RemoveAll(this);
-		FEditorDelegates::OnEditCutActorsEnd.RemoveAll(this);
-		FEditorDelegates::OnEditCopyActorsBegin.RemoveAll(this);
-		FEditorDelegates::OnEditCopyActorsEnd.RemoveAll(this);
-		FEditorDelegates::OnEditPasteActorsBegin.RemoveAll(this);
-		FEditorDelegates::OnEditPasteActorsEnd.RemoveAll(this);
-		FEditorDelegates::OnDuplicateActorsBegin.RemoveAll(this);
-		FEditorDelegates::OnDuplicateActorsEnd.RemoveAll(this);
-		FEditorDelegates::OnDeleteActorsBegin.RemoveAll(this);
-		FEditorDelegates::OnDeleteActorsEnd.RemoveAll(this);
 	}
 
 	void SSceneOutliner::OnItemAdded(const FTreeItemID& ItemID, uint8 ActionMask)
@@ -710,627 +355,190 @@ namespace SceneOutliner
 
 		return ViewOptionsComboButton->IsHovered() ? FEditorStyle::GetSlateColor(InvertedForegroundName) : FEditorStyle::GetSlateColor(DefaultForegroundName);
 	}
-
-	TSharedRef<SWidget> SSceneOutliner::GetViewButtonContent(bool bWorldPickerOnly, bool bShouldDisplayChooseWorld)
+	 
+	TSharedRef<SWidget> SSceneOutliner::GetViewButtonContent(bool bShowFilters)
 	{
-		FMenuBuilder MenuBuilder(!bWorldPickerOnly, NULL);
+		// Menu should stay open on selection if filters are not being shown
+		FMenuBuilder MenuBuilder(bShowFilters, NULL);
 
-		if(bWorldPickerOnly)
-		{
-			BuildWorldPickerContent(MenuBuilder);
-		}
-		else
+		if (bShowFilters)
 		{
 			MenuBuilder.BeginSection("AssetThumbnails", LOCTEXT("ShowHeading", "Show"));
 			{
-				MenuBuilder.AddMenuEntry(
-					LOCTEXT("ToggleShowOnlySelected", "Only Selected"),
-					LOCTEXT("ToggleShowOnlySelectedToolTip", "When enabled, only displays actors that are currently selected."),
-					FSlateIcon(),
-					FUIAction(
-					FExecuteAction::CreateSP( this, &SSceneOutliner::ToggleShowOnlySelected ),
-					FCanExecuteAction(),
-					FIsActionChecked::CreateSP( this, &SSceneOutliner::IsShowingOnlySelected )
-					),
-					NAME_None,
-					EUserInterfaceActionType::ToggleButton
-				);
-
-				MenuBuilder.AddMenuEntry(
-					LOCTEXT("ToggleHideTemporaryActors", "Hide Temporary Actors"),
-					LOCTEXT("ToggleHideTemporaryActorsToolTip", "When enabled, hides temporary/run-time Actors."),
-					FSlateIcon(),
-					FUIAction(
-					FExecuteAction::CreateSP( this, &SSceneOutliner::ToggleHideTemporaryActors ),
-					FCanExecuteAction(),
-					FIsActionChecked::CreateSP( this, &SSceneOutliner::IsHidingTemporaryActors )
-					),
-					NAME_None,
-					EUserInterfaceActionType::ToggleButton
-				);
-
-				MenuBuilder.AddMenuEntry(
-					LOCTEXT("ToggleShowOnlyCurrentLevel", "Only in Current Level"),
-					LOCTEXT("ToggleShowOnlyCurrentLevelToolTip", "When enabled, only shows Actors that are in the Current Level."),
-					FSlateIcon(),
-					FUIAction(
-					FExecuteAction::CreateSP( this, &SSceneOutliner::ToggleShowOnlyCurrentLevel ),
-					FCanExecuteAction(),
-					FIsActionChecked::CreateSP( this, &SSceneOutliner::IsShowingOnlyCurrentLevel )
-					),
-					NAME_None,
-					EUserInterfaceActionType::ToggleButton
-				);
-
-				// Temporarily disable this feature until it can be redesigned.
-				//MenuBuilder.AddMenuEntry(
-				//	LOCTEXT("ToggleHideFoldersContainingHiddenActors", "Hide Folders with Only Hidden Actors"),
-				//	LOCTEXT("ToggleHideFoldersContainingHiddenActorsToolTip", "When enabled, only shows Folders containing non-hidden Actors."),
-				//	FSlateIcon(),
-				//	FUIAction(
-				//		FExecuteAction::CreateSP(this, &SSceneOutliner::ToggleHideFoldersContainingOnlyHiddenActors),
-				//		FCanExecuteAction(),
-				//		FIsActionChecked::CreateSP(this, &SSceneOutliner::IsHidingFoldersContainingOnlyHiddenActors)
-				//	),
-				//	NAME_None,
-				//	EUserInterfaceActionType::ToggleButton
-				//);
-
-				MenuBuilder.AddMenuEntry(
-					LOCTEXT("ToggleShowActorComponents", "Show Actor Components"),
-					LOCTEXT("ToggleShowActorComponentsToolTip", "When enabled, shows components belonging to actors."),
-					FSlateIcon(),
-					FUIAction(
-						FExecuteAction::CreateSP(this, &SSceneOutliner::ToggleShowActorComponents),
-						FCanExecuteAction(),
-						FIsActionChecked::CreateSP(this, &SSceneOutliner::IsShowingActorComponents)
-					),
-					NAME_None,
-					EUserInterfaceActionType::ToggleButton
-				);
-
-				// Add additional filters
-				FSceneOutlinerModule& SceneOutlinerModule = FModuleManager::LoadModuleChecked< FSceneOutlinerModule >("SceneOutliner");
-			
-				for (auto& OutlinerFilterInfo : SceneOutlinerModule.OutlinerFilterInfoMap)
+				// Add mode filters
+				for (auto& ModeFilterInfo : Mode->GetFilterInfos())
 				{
-					OutlinerFilterInfo.Value.AddMenu(MenuBuilder);
+					ModeFilterInfo.Value.AddMenu(MenuBuilder);
 				}
 			}
 			MenuBuilder.EndSection();
-
-			if( bShouldDisplayChooseWorld )
-			{
-				MenuBuilder.BeginSection("AssetThumbnails", LOCTEXT("ShowWorldHeading", "World"));
-				{
-					MenuBuilder.AddSubMenu(
-						LOCTEXT("ChooseWorldSubMenu", "Choose World"),
-						LOCTEXT("ChooseWorldSubMenuToolTip", "Choose the world to display in the outliner."),
-						FNewMenuDelegate::CreateSP(this, &SSceneOutliner::BuildWorldPickerContent)
-					);
-				}
-				MenuBuilder.EndSection();
-			}
 		}
+		Mode->CreateViewContent(MenuBuilder);
 
 		return MenuBuilder.MakeWidget();
 	}
 
-	void SSceneOutliner::BuildWorldPickerContent(FMenuBuilder& MenuBuilder)
+	ESelectionMode::Type SSceneOutliner::GetSelectionMode() const
 	{
-		MenuBuilder.BeginSection("Worlds", LOCTEXT("WorldsHeading", "Worlds"));
-		{
-			MenuBuilder.AddMenuEntry(
-				LOCTEXT("AutoWorld", "Auto"),
-				LOCTEXT("AutoWorldToolTip", "Automatically pick the world to display based on context."),
-				FSlateIcon(),
-				FUIAction(
-				FExecuteAction::CreateSP( this, &SSceneOutliner::OnSelectWorld, TWeakObjectPtr<UWorld>() ),
-				FCanExecuteAction(),
-				FIsActionChecked::CreateSP( this, &SSceneOutliner::IsWorldChecked, TWeakObjectPtr<UWorld>() )
-				),
-				NAME_None,
-				EUserInterfaceActionType::RadioButton
-			);
-
-			for (const FWorldContext& Context : GEngine->GetWorldContexts())
-			{
-				UWorld* World = Context.World();
-				if (World && (World->WorldType == EWorldType::PIE || Context.WorldType == EWorldType::Editor))
-				{
-					MenuBuilder.AddMenuEntry(
-						GetWorldDescription(World),
-						LOCTEXT("ChooseWorldToolTip", "Display actors for this world."),
-						FSlateIcon(),
-						FUIAction(
-						FExecuteAction::CreateSP( this, &SSceneOutliner::OnSelectWorld, MakeWeakObjectPtr(World) ),
-						FCanExecuteAction(),
-						FIsActionChecked::CreateSP( this, &SSceneOutliner::IsWorldChecked, MakeWeakObjectPtr(World) )
-						),
-						NAME_None,
-						EUserInterfaceActionType::RadioButton
-					);
-				}
-			}
-		}
-		MenuBuilder.EndSection();
+		return Mode->GetSelectionMode();
 	}
 
-	/** FILTERS */
+	void SSceneOutliner::Refresh()
+	{
+		bNeedsRefresh = true;
+	}
 
-	void SSceneOutliner::OnSharedSettingChanged()
+	void SSceneOutliner::FullRefresh()
 	{
-		// Only update if we use the shared settings
-		if (!SceneOutlinerSettings)
-		{
-			ApplyHideTemporaryActorsFilter(IsHidingTemporaryActors());
-			ApplyShowActorComponentsFilter(IsShowingActorComponents());
-			ApplyShowOnlyCurrentLevelFilter(IsShowingOnlyCurrentLevel());
-			ApplyShowOnlySelectedFilter(IsShowingOnlySelected());
-		}
+		UE_LOG(LogSceneOutliner, Verbose, TEXT("Full Refresh"));
+		bFullRefresh = true;
+		Refresh();
+	}
+
+	void SSceneOutliner::RefreshSelection()
+	{
+		bSelectionDirty = true;
 	}
 
-	// Show Only Selected
-	void SSceneOutliner::ToggleShowOnlySelected()
+	void SSceneOutliner::Populate()
 	{
-		const bool bEnableFlag = !IsShowingOnlySelected();
+		// Block events while we clear out the list
+ 		TGuardValue<bool> ReentrantGuard(bIsReentrant, true);
+
+		// Get a collection of items and folders which were formerly collapsed
+		const FParentsExpansionState ExpansionStateInfo = GetParentsExpansionState();
 
-		if (SceneOutlinerSettings)
+		bool bMadeAnySignificantChanges = false;
+		if (bFullRefresh)
 		{
-			SceneOutlinerSettings->bShowOnlySelectedActors = bEnableFlag;
+			// Clear the selection here - RepopulateEntireTree will reconstruct it.
+			OutlinerTreeView->ClearSelection();
+
+			RepopulateEntireTree();
+
+			bMadeAnySignificantChanges = true;
+			bFullRefresh = false;
 		}
-		else
+
+		// Only deal with 500 at a time
+		const int32 End = FMath::Min(PendingOperations.Num(), 500);
+		for (int32 Index = 0; Index < End; ++Index)
 		{
-			USceneOutlinerSettings* Settings = GetMutableDefault<USceneOutlinerSettings>();
-			Settings->bShowOnlySelectedActors = bEnableFlag;
-			Settings->PostEditChange();
-			OnSharedSettingChangedDelegate.Broadcast();
+			auto& PendingOp = PendingOperations[Index];
+			switch (PendingOp.Type)
+			{
+			case FPendingTreeOperation::Added:
+				bMadeAnySignificantChanges = AddItemToTree(PendingOp.Item) || bMadeAnySignificantChanges;
+				break;
+
+			case FPendingTreeOperation::Moved:
+				bMadeAnySignificantChanges = true;
+				OnItemMoved(PendingOp.Item);
+				break;
+
+			case FPendingTreeOperation::Removed:
+				bMadeAnySignificantChanges = true;
+				RemoveItemFromTree(PendingOp.Item);
+				break;
+
+			default:
+				check(false);
+				break;
+			}
 		}
 
-		ApplyShowOnlySelectedFilter(bEnableFlag);
-	}
+		PendingOperations.RemoveAt(0, End);
+		SetParentsExpansionState(ExpansionStateInfo);
 
-	void SSceneOutliner::ApplyShowOnlySelectedFilter(bool bShowOnlySelected)
-	{
-		if ( !SelectedActorFilter.IsValid() )
+
+		for (FName Folder : PendingFoldersSelect)
 		{
-			SelectedActorFilter = CreateSelectedActorFilter();
+			if (FTreeItemPtr* Item = TreeItemMap.Find(Folder))
+			{
+				OutlinerTreeView->SetItemSelection(*Item, true);
+			}
 		}
+		PendingFoldersSelect.Empty();
 
-		if ( bShowOnlySelected )
-		{			
-			Filters->Add( SelectedActorFilter );
+		// Check if we need to sort because we are finished with the populating operations
+		bool bFinalSort = false;
+		if (PendingOperations.Num() == 0)
+		{
+			// We're fully refreshed now.
+			NewItemActions.Empty();
+			bNeedsRefresh = false;
+			if (bDisableIntermediateSorting)
+			{
+				bDisableIntermediateSorting = false;
+				bFinalSort = true;
+			}
 		}
-		else
+
+		// If we are allowing intermediate sorts and met the conditions, or this is the final sort after all ops are complete
+		if ((bMadeAnySignificantChanges && !bDisableIntermediateSorting) || bFinalSort)
 		{
-			Filters->Remove( SelectedActorFilter );
+			RequestSort();
 		}
 	}
 
-	bool SSceneOutliner::IsShowingOnlySelected() const
+	bool SSceneOutliner::ShouldShowFolders() const
 	{
-		if (SceneOutlinerSettings)
-		{
-			return SceneOutlinerSettings->bShowOnlySelectedActors;
-		}
-		return GetDefault<USceneOutlinerSettings>()->bShowOnlySelectedActors;
+		return Mode->ShouldShowFolders();
 	}
 
-	// Hide Temporary Actors
-	void SSceneOutliner::ToggleHideTemporaryActors()
+	void SSceneOutliner::EmptyTreeItems()
 	{
-		const bool bEnableFlag = !IsHidingTemporaryActors();
+		PendingOperations.Empty();
+		TreeItemMap.Reset();
+		PendingTreeItemMap.Empty();
 
-		if (SceneOutlinerSettings)
-		{
-			SceneOutlinerSettings->bHideTemporaryActors = bEnableFlag;
-		}
-		else
-		{
-			USceneOutlinerSettings* Settings = GetMutableDefault<USceneOutlinerSettings>();
-			Settings->bHideTemporaryActors = bEnableFlag;
-			Settings->PostEditChange();
-			OnSharedSettingChangedDelegate.Broadcast();
-		}
+		FolderCount = 0;
 
-		ApplyHideTemporaryActorsFilter(bEnableFlag);
+		RootTreeItems.Empty();
 	}
 
-	void SSceneOutliner::ApplyHideTemporaryActorsFilter(bool bHideTemporaryActors)
+	void SSceneOutliner::AddPendingItem(FTreeItemPtr Item)
 	{
-		if ( !HideTemporaryActorsFilter.IsValid() )
-		{
-			HideTemporaryActorsFilter = CreateHideTemporaryActorsFilter();
-		}
-
-		if ( bHideTemporaryActors )
-		{			
-			Filters->Add( HideTemporaryActorsFilter );
-		}
-		else
-		{
-			Filters->Remove( HideTemporaryActorsFilter );
-		}
+		PendingTreeItemMap.Add(Item->GetID(), Item);
+		PendingOperations.Emplace(FPendingTreeOperation::Added, Item.ToSharedRef());
 	}
-	bool SSceneOutliner::IsHidingTemporaryActors() const
+
+	void SSceneOutliner::AddPendingItemAndChildren(FTreeItemPtr Item)
 	{
-		if (SceneOutlinerSettings)
+		if (Item)
 		{
-			return SceneOutlinerSettings->bHideTemporaryActors;
+			AddPendingItem(Item);
+
+			TArray<FTreeItemPtr> Children;
+			Mode->GetHierarchy()->CreateChildren(Item, Children);
+			for (auto& Child : Children)
+			{
+				AddPendingItem(Child);
+			}
+
+			Refresh();
 		}
-		return GetDefault<USceneOutlinerSettings>()->bHideTemporaryActors;
-	}
-
-	// Show Only Actors In Current Level
-	void SSceneOutliner::ToggleShowOnlyCurrentLevel()
-	{
-		const bool bEnableFlag = !IsShowingOnlyCurrentLevel();
-
-		if (SceneOutlinerSettings)
-		{
-			SceneOutlinerSettings->bShowOnlyActorsInCurrentLevel = bEnableFlag;
-		}
-		else
-		{
-			USceneOutlinerSettings* Settings = GetMutableDefault<USceneOutlinerSettings>();
-			Settings->bShowOnlyActorsInCurrentLevel = bEnableFlag;
-			Settings->PostEditChange();
-			OnSharedSettingChangedDelegate.Broadcast();
-		}
-
-		ApplyShowOnlyCurrentLevelFilter(bEnableFlag);
-	}
-	void SSceneOutliner::ApplyShowOnlyCurrentLevelFilter(bool bShowOnlyActorsInCurrentLevel)
-	{
-		if ( !ShowOnlyActorsInCurrentLevelFilter.IsValid() )
-		{
-			ShowOnlyActorsInCurrentLevelFilter = CreateIsInCurrentLevelFilter();
-		}
-
-		if ( bShowOnlyActorsInCurrentLevel )
-		{			
-			Filters->Add( ShowOnlyActorsInCurrentLevelFilter );
-		}
-		else
-		{
-			Filters->Remove( ShowOnlyActorsInCurrentLevelFilter );
-		}
-	}
-
-	void SSceneOutliner::ToggleHideFoldersContainingOnlyHiddenActors()
-	{
-		const bool bEnableFlag = !IsHidingFoldersContainingOnlyHiddenActors();
-
-		USceneOutlinerSettings* Settings = GetMutableDefault<USceneOutlinerSettings>();
-		Settings->bHideFoldersContainingHiddenActors = bEnableFlag;
-		Settings->PostEditChange();
-
-		FullRefresh();
-	}
-
-	bool SSceneOutliner::IsShowingActorComponents() const
-	{
-		if (SceneOutlinerSettings)
-		{
-			return SceneOutlinerSettings->bShowActorComponents;
-		}
-		return SharedData->Mode == ESceneOutlinerMode::ComponentPicker || (GetDefault<USceneOutlinerSettings>()->bShowActorComponents);
-	}
-
-	void SSceneOutliner::ToggleShowActorComponents()
-	{
-		if ( SharedData->Mode != ESceneOutlinerMode::ComponentPicker )
-		{
-			const bool bEnableFlag = !IsShowingActorComponents();
-
-			if (SceneOutlinerSettings)
-			{
-				SceneOutlinerSettings->bShowActorComponents = bEnableFlag;
-			}
-			else
-			{
-				USceneOutlinerSettings* Settings = GetMutableDefault<USceneOutlinerSettings>();
-				Settings->bShowActorComponents = bEnableFlag;
-				Settings->PostEditChange();
-				OnSharedSettingChangedDelegate.Broadcast();
-			}
-
-			ApplyShowActorComponentsFilter(bEnableFlag);
-		}
-	}
-
-	void SSceneOutliner::ApplyShowActorComponentsFilter(bool bShowActorComponents)
-	{
-		if (!ShowActorComponentsFilter.IsValid())
-		{
-			ShowActorComponentsFilter = CreateShowActorComponentsFilter();
-		}
-
-		if (bShowActorComponents)
-		{
-			Filters->Add(ShowActorComponentsFilter);
-		}
-		else
-		{
-			Filters->Remove(ShowActorComponentsFilter);
-		}
-	}
-
-	bool SSceneOutliner::IsShowingOnlyCurrentLevel() const
-	{
-		if (SceneOutlinerSettings)
-		{
-			return SceneOutlinerSettings->bShowOnlyActorsInCurrentLevel;
-		}
-		return GetDefault<USceneOutlinerSettings>()->bShowOnlyActorsInCurrentLevel;
-	}
-
-	bool SSceneOutliner::IsHidingFoldersContainingOnlyHiddenActors() const
-	{
-		// Temporarily disable this feature until it can be redesigned.
-		return false;
-		//return GetDefault<USceneOutlinerSettings>()->bHideFoldersContainingHiddenActors;
-	}
-
-	/** END FILTERS */
-
-
-	ESelectionMode::Type SSceneOutliner::GetSelectionMode() const
-	{
-		return SelectionMode;
-	}
-
-
-	void SSceneOutliner::Refresh()
-	{
-		if (IsHidingFoldersContainingOnlyHiddenActors())
-		{
-			bFullRefresh = true;
-		}
-		bNeedsRefresh = true;
-	}
-
-	void SSceneOutliner::FullRefresh()
-	{
-		bFullRefresh = true;
-		Refresh();
-	}
-
-	void SSceneOutliner::OnLevelActorListChanged()
-	{
-		bDisableIntermediateSorting = true;
-		FullRefresh();
-	}
-
-	void SSceneOutliner::Populate()
-	{
-		// Block events while we clear out the list.  We don't want actors in the level to become deselected
-		// while we doing this
- 		TGuardValue<bool> ReentrantGuard(bIsReentrant, true);
-
-		SharedData->RepresentingWorld = nullptr;
-
-		// check if the user-chosen world is valid and in the editor contexts
-		if(UWorld* UserChosenWorld = SharedData->UserChosenWorld.Get())
-		{
-			for (const FWorldContext& Context : GEngine->GetWorldContexts())
-			{
-				if(UserChosenWorld == Context.World())
-				{
-					SharedData->RepresentingWorld = UserChosenWorld;
-					break;
-				}
-			}
-		}
-		
-		if(SharedData->RepresentingWorld == nullptr)
-		{
-			// try to pick the most suitable world context
-
-			// ideally we want a PIE world that is standalone or the first client
-			for (const FWorldContext& Context : GEngine->GetWorldContexts())
-			{
-				UWorld* World = Context.World();
-				if (World && Context.WorldType == EWorldType::PIE)
-				{
-					if(World->GetNetMode() == NM_Standalone)
-					{
-						SharedData->RepresentingWorld = World;
-						break;
-					}
-					else if(World->GetNetMode() == NM_Client && Context.PIEInstance == 2)	// Slightly dangerous: assumes server is always PIEInstance = 1;
-					{
-						SharedData->RepresentingWorld = World;
-						break;
-					}
-				}
-			}
-		}
-
-		if(SharedData->RepresentingWorld == nullptr)
-		{
-			// still not world so fallback to old logic where we just prefer PIE over Editor
-
-			for (const FWorldContext& Context : GEngine->GetWorldContexts())
-			{
-				if (Context.WorldType == EWorldType::PIE)
-				{
-					SharedData->RepresentingWorld = Context.World();
-					break;
-				}
-				else if (Context.WorldType == EWorldType::Editor)
-				{
-					SharedData->RepresentingWorld = Context.World();
-				}
-			}
-		}
-		
-		if (!CheckWorld())
-		{
-			return;
-		}
-
-		SharedData->bRepresentingPlayWorld = SharedData->RepresentingWorld->WorldType == EWorldType::PIE;
-
-		// Get a collection of items and folders which were formerly collapsed
-		const FParentsExpansionState ExpansionStateInfo = GetParentsExpansionState();
-
-		bool bMadeAnySignificantChanges = false;
-		if(bFullRefresh)
-		{
-			// Clear the selection here - RepopulateEntireTree will reconstruct it.
-			OutlinerTreeView->ClearSelection();
-
-			RepopulateEntireTree();
-
-			bMadeAnySignificantChanges = true;
-			bFullRefresh = false;
-		}
-
-		// Only deal with 500 at a time
-		const int32 End = FMath::Min(PendingOperations.Num(), 500);
-		for (int32 Index = 0; Index < End; ++Index)
-		{
-			auto& PendingOp = PendingOperations[Index];
-			switch (PendingOp.Type)
-			{
-				case FPendingTreeOperation::Added:
-					bMadeAnySignificantChanges = AddItemToTree(PendingOp.Item) || bMadeAnySignificantChanges;
-					break;
-
-				case FPendingTreeOperation::Moved:
-					bMadeAnySignificantChanges = true;
-					OnItemMoved(PendingOp.Item);
-					break;
-
-				case FPendingTreeOperation::Removed:
-					bMadeAnySignificantChanges = true;
-					RemoveItemFromTree(PendingOp.Item);
-					break;
-
-				default:
-					check(false);
-					break;
-			}
-		}
-
-		PendingOperations.RemoveAt(0, End);
-		SetParentsExpansionState(ExpansionStateInfo);
-
-
-		for (FName Folder : PendingFoldersSelect)
-		{
-			if (FTreeItemPtr* Item = TreeItemMap.Find(Folder))
-			{
-				OutlinerTreeView->SetItemSelection(*Item, true);
-			}
-		}
-		PendingFoldersSelect.Empty();
-
-		// Check if we need to sort because we are finished with the populating operations
-		bool bFinalSort = false;
-		if (PendingOperations.Num() == 0)
-		{
-			// We're fully refreshed now.
-			NewItemActions.Empty();
-			bNeedsRefresh = false;
-			if (bDisableIntermediateSorting)
-			{
-				bDisableIntermediateSorting = false;
-				bFinalSort = true;
-			}
-
-			HideFoldersContainingOnlyHiddenActors();
-		}
-
-		// If we are allowing intermediate sorts and met the conditions, or this is the final sort after all ops are complete
-		if ((bMadeAnySignificantChanges && !bDisableIntermediateSorting) || bFinalSort)
-		{
-			RequestSort();
-		}
-	}
-
-	bool SSceneOutliner::ShouldShowFolders() const
-	{
-		return SharedData->Mode == ESceneOutlinerMode::ActorBrowsing || SharedData->bOnlyShowFolders;
-	}
-
-	void SSceneOutliner::EmptyTreeItems()
-	{
-		FilteredActorCount = 0;
-		ApplicableActors.Empty();
-
-		PendingOperations.Empty();
-		TreeItemMap.Reset();
-		PendingTreeItemMap.Empty();
-
-		RootTreeItems.Empty();
 	}
 
 	void SSceneOutliner::RepopulateEntireTree()
 	{
-		// to avoid dependencies Custom Tree Items are accessed via modular features
-		TArray<ISceneOutlinerTraversal*> ConstructTreeItemImp = IModularFeatures::Get().GetModularFeatureImplementations<ISceneOutlinerTraversal>("SceneOutlinerTraversal");
-		ISceneOutlinerTraversal *CustomImplementation = nullptr;
-		if (ConstructTreeItemImp.Num() > 0 && ConstructTreeItemImp[0] != nullptr)
-		{
-			// As an optimization, since we have only one customisation at present, just grab the one CustomImplementation to mitigate a further for loop inside the actor iterator
-			check(ConstructTreeItemImp.Num() < 2);
-			CustomImplementation = ConstructTreeItemImp[0];
-		}
-
 		EmptyTreeItems();
 
-		ConstructItemFor<FWorldTreeItem>(SharedData->RepresentingWorld);
+		// Rebuild the hierarchy
+		Mode->Rebuild();
+		Mode->GetHierarchy()->OnHierarchyChanged().AddSP(this, &SSceneOutliner::OnHierarchyChangedEvent);
 
-		if (!SharedData->bOnlyShowFolders)
-		{
-			// Iterate over every actor in memory. WARNING: This is potentially very expensive!
-			for( FActorIterator ActorIt(SharedData->RepresentingWorld); ActorIt; ++ActorIt )
-			{
-				AActor* Actor = *ActorIt;
-				if (Actor && IsActorDisplayable(Actor))
-				{
-					if (Filters->PassesAllFilters(FActorTreeItem(Actor)))
-					{
-						ApplicableActors.Emplace(Actor);
-					}
-					ConstructItemFor<FActorTreeItem>(Actor);
+		// Create all the items which match the filters, parent-child relationships are handled when each item is actually added to the tree
 
-					if (IsShowingActorComponents())
-					{
-						for (UActorComponent* Component : Actor->GetComponents())
-						{
-							if (Filters->PassesAllFilters(FComponentTreeItem(Component)))
-							{
-								bool IsHandled = false;
-								if (CustomImplementation)
-								{
-									IsHandled = CustomImplementation->ConstructTreeItem(*this, Component);
-								}
-								if (!IsHandled)
-								{
-									// add the actor's components - default implementation
-									ConstructItemFor<FComponentTreeItem>(Component);
-								}
-							}
-						}
-					}
-				}
-			}
-		}
+		TArray<FTreeItemPtr> Items;
+		Mode->GetHierarchy()->CreateItems(Items);
 
-		if (!IsShowingOnlySelected() && ShouldShowFolders())
+		for (FTreeItemPtr& Item : Items)
 		{
-			// Add any folders which might match the current search terms
-			for (const auto& Pair : FActorFolders::Get().GetFolderPropertiesForWorld(*SharedData->RepresentingWorld))
-			{
-				if (!TreeItemMap.Contains(Pair.Key))
-				{
-					ConstructItemFor<FFolderTreeItem>(Pair.Key);
-				}
-			}
+			AddPendingItem(Item);
 		}
+
+		Refresh();
 	}
 
 	void SSceneOutliner::OnChildRemovedFromParent(ITreeItem& Parent)
@@ -1378,6 +586,16 @@ namespace SceneOutliner
 		}
 	}
 
+	FTreeItemPtr SSceneOutliner::GetTreeItem(FTreeItemID ItemID, bool bIncludePending)
+	{
+		FTreeItemPtr Result = TreeItemMap.FindRef(ItemID);
+		if (bIncludePending && !Result.IsValid())
+		{
+			Result = PendingTreeItemMap.FindRef(ItemID);
+		}
+		return Result;
+	}
+
 	void SSceneOutliner::RemoveItemFromTree(FTreeItemRef InItem)
 	{
 		if (TreeItemMap.Contains(InItem->GetID()))
@@ -1394,14 +612,9 @@ namespace SceneOutliner
 				RootTreeItems.Remove(InItem);
 			}
 
-			InItem->Visit(FFunctionalVisitor().Actor([&](const FActorTreeItem& ActorItem){
-				if (!ActorItem.Flags.bIsFilteredOut)
-				{
-					--FilteredActorCount;
-				}
-			}));
-
 			TreeItemMap.Remove(InItem->GetID());
+
+			Mode->OnItemRemoved(InItem);
 		}
 	}
 
@@ -1409,20 +622,26 @@ namespace SceneOutliner
 	{
 		if (SharedData->bShowParentTree)
 		{
-			FTreeItemPtr Parent = Item->FindParent(TreeItemMap);
+			FTreeItemPtr Parent = Mode->GetHierarchy()->FindParent(*Item, TreeItemMap);
 			if (Parent.IsValid())
 			{
 				return Parent;
 			}
 			else
 			{
-				auto NewParent = Item->CreateParent();
-				if (NewParent.IsValid())
+				// Try to find the parent in the pending items
+				Parent = Mode->GetHierarchy()->FindParent(*Item, PendingTreeItemMap);
+				if (!Parent.IsValid())
 				{
-					NewParent->Flags.bIsFilteredOut = !Filters->TestAndSetInteractiveState(*NewParent) || !SearchBoxFilter->PassesFilter(*NewParent);
+					// If there isn't already a parent for this item, try to create one for it
+					Parent = Mode->GetHierarchy()->CreateParentItem(Item);
+				}
+				
+				if (Parent.IsValid())
+				{
+					AddUnfilteredItemToTree(Parent.ToSharedRef());
 
-					AddUnfilteredItemToTree(NewParent.ToSharedRef());
-					return NewParent;
+					return Parent;
 				}
 			}
 		}
@@ -1436,11 +655,8 @@ namespace SceneOutliner
 
 		PendingTreeItemMap.Remove(ItemID);
 
-		FValidateItemBeforeAddingToTree ValidateItemVisitor;
-		Item->Visit(ValidateItemVisitor);
-
-		// If a tree item already exists that represents the same data or if the actor is invalid, bail
-		if (TreeItemMap.Find(ItemID)  || !ValidateItemVisitor.Result())
+		// If a tree item already exists that represents the same data or if the item represents invalid data, bail
+		if (TreeItemMap.Find(ItemID)  || !Item->IsValid())
 		{
 			return false;
 		}
@@ -1461,7 +677,7 @@ namespace SceneOutliner
 					OutlinerTreeView->SetItemSelection(Item, true);
 				}
 
-				if (*ActionMask & ENewItemAction::Rename && CanExecuteRenameRequest(Item))
+				if (*ActionMask & ENewItemAction::Rename && CanExecuteRenameRequest(*Item))
 				{
 					PendingRenameItem = Item;
 				}
@@ -1478,8 +694,6 @@ namespace SceneOutliner
 
 	void SSceneOutliner::AddUnfilteredItemToTree(FTreeItemRef Item)
 	{
-		Item->SharedData = SharedData;
-
 		FTreeItemPtr Parent = EnsureParentForItem(Item);
 
 		const FTreeItemID ItemID = Item->GetID();
@@ -1506,21 +720,12 @@ namespace SceneOutliner
 			RootTreeItems.Add(Item);
 		}
 
-		if (SharedData->Mode == ESceneOutlinerMode::ActorBrowsing)
-		{
-			Item->Visit(FOnItemAddedToTree(*this));
-		}
-		else if (SharedData->Mode == ESceneOutlinerMode::Custom)
+		// keep track of the number of active folders
+		if (Item->IsA<FFolderTreeItem>())
 		{
-			if(TTreeItemGetter<bool>* ShouldSelectNewItem = ShouldSelectNewItemVisitor.Get())
-			{
-				Item->Visit(*ShouldSelectNewItem);
-				if (ShouldSelectNewItem->Result())
-				{
-					OutlinerTreeView->SetItemSelection(Item, true);
-				}
-			}
+			++FolderCount;
 		}
+		Mode->OnItemAdded(Item);
 	}
 
 	SSceneOutliner::FParentsExpansionState SSceneOutliner::GetParentsExpansionState() const
@@ -1556,157 +761,36 @@ namespace SceneOutliner
 		}
 	}
 
-	void SSceneOutliner::HideFoldersContainingOnlyHiddenActors()
+	void SSceneOutliner::PopulateSearchStrings(const ITreeItem& Item, TArray< FString >& OutSearchStrings) const
 	{
-		if (IsHidingFoldersContainingOnlyHiddenActors())
+		for (const auto& Pair : Columns)
 		{
-			for (const FTreeItemPtr & TreeItem : RootTreeItems)
-			{
-				HideFoldersContainingOnlyHiddenActors(TreeItem, true);
-			}
+			Pair.Value->PopulateSearchStrings(Item, OutSearchStrings);
 		}
 	}
 
-	bool SSceneOutliner::HideFoldersContainingOnlyHiddenActors(FTreeItemPtr Parent, bool bIsRoot)
+	void SSceneOutliner::GetSelectedFolders(TArray<FFolderTreeItem*>& OutFolders) const
 	{
-		TArray<FTreeItemPtr> ItemsToRemove;
-
-		bool bActorsHidden = true;
-		bool bFoldersHidden = true;
-
-		const TArray<TWeakPtr<ITreeItem>>& Children = Parent->GetChildren();
-
-		if (Children.Num())
-		{
-
-			for (const TWeakPtr<ITreeItem> & ChildItem : Children)
-			{
-				FTreeItemPtr TreeItem = ChildItem.Pin();
-
-				TWeakObjectPtr<AActor> TrueTreeActor = StaticCastSharedPtr<FActorTreeItem>(TreeItem)->Actor;
-
-				if (!bIsRoot && TrueTreeActor.IsValid())
-				{
-					if (bActorsHidden)
-					{
-						if (!(TrueTreeActor->IsTemporarilyHiddenInEditor()))
-						{
-							bActorsHidden = false;
-						}
-					}
-				}
-				else
-				{
-					bool bCurrentFolderHidden = HideFoldersContainingOnlyHiddenActors(TreeItem);
-
-					if (bCurrentFolderHidden)
-					{
-						ItemsToRemove.Add(TreeItem);
-					}
-
-					bFoldersHidden = bCurrentFolderHidden & bFoldersHidden;
-
-				}
-
-			}
-		}
-		else
-		{
-			return false;
-		}
-
-		for (const FTreeItemPtr & Item : ItemsToRemove)
-		{
-			FTreeItemRef RemoveItem = Item.ToSharedRef();
-			Parent->RemoveChild(RemoveItem);
-		}
-
-		return bActorsHidden && bFoldersHidden;
-	}
-
-	void SSceneOutliner::PopulateSearchStrings(const ITreeItem& Item, TArray< FString >& OutSearchStrings) const
-	{
-		for (const auto& Pair : Columns)
-		{
-			Pair.Value->PopulateSearchStrings(Item, OutSearchStrings);
-		}
-	}
-
-	TArray<FFolderTreeItem*> SSceneOutliner::GetSelectedFolders() const
-	{
-		return FItemSelection(*OutlinerTreeView).Folders;
-	}
+		return FItemSelection(*OutlinerTreeView).Get<FFolderTreeItem>(OutFolders);
+	}
 
 	TArray<FName> SSceneOutliner::GetSelectedFolderNames() const
 	{
-		TArray<FFolderTreeItem*> SelectedFolders = GetSelectedFolders();
-		TArray<FName> SelectedFolderNames;
-		for (FFolderTreeItem* Folder : SelectedFolders)
-		{
-			if (Folder)
-			{
-				SelectedFolderNames.Add(Folder->Path);
-			}
-		}
-		return SelectedFolderNames;
+		return GetSelection().GetData<FName>(FFolderPathSelector());
 	}
 
 	TSharedPtr<SWidget> SSceneOutliner::OnOpenContextMenu()
 	{
-		/** Legacy mode and now also used by the custom mode */
-		if (SharedData->ContextMenuOverride.IsBound())
-		{
-			return SharedData->ContextMenuOverride.Execute();
-		}
-
-		if (SharedData->Mode == ESceneOutlinerMode::ActorBrowsing)
-		{
-			TArray<AActor*> SelectedActors;
-			GEditor->GetSelectedActors()->GetSelectedObjects<AActor>(SelectedActors);
-
-			// Make sure that no components are selected
-			if (GEditor->GetSelectedComponentCount() > 0)
-			{
-				// We want to be able to undo to regain the previous component selection
-				const FScopedTransaction Transaction(NSLOCTEXT("UnrealEd", "ClickingOnActorsContextMenu", "Clicking on Actors (context menu)"));
-				USelection* ComponentSelection = GEditor->GetSelectedComponents();
-				ComponentSelection->Modify(false);
-				ComponentSelection->DeselectAll();
-
-				GUnrealEd->UpdatePivotLocationForSelection();
-				GEditor->RedrawLevelEditingViewports(false);
-			}
-
-			return BuildDefaultContextMenu();
-		}
-
-		return TSharedPtr<SWidget>();
+		return Mode->CreateContextMenu();
 	}
 	
 	bool SSceneOutliner::Delete_CanExecute()
 	{
-		if (SharedData->Mode == ESceneOutlinerMode::ActorPicker)
-		{
-			FItemSelection ItemSelection(*OutlinerTreeView);
-			if (ItemSelection.Folders.Num() > 0 && ItemSelection.Folders.Num() == OutlinerTreeView->GetNumItemsSelected())
-			{
-				return true;
-			}
-		}
-		return false;
+		return Mode->CanDelete();
 	}
-
 	bool SSceneOutliner::Rename_CanExecute()
 	{
-		if (SharedData->Mode == ESceneOutlinerMode::ActorBrowsing)
-		{
-			FItemSelection ItemSelection(*OutlinerTreeView);
-			if (ItemSelection.Folders.Num() == 1 && ItemSelection.Folders.Num() == OutlinerTreeView->GetNumItemsSelected())
-			{
-				return true;
-			}
-		}
-		return false;
+		return Mode->CanRename();
 	}
 
 	void SSceneOutliner::Rename_Execute()
@@ -1714,23 +798,12 @@ namespace SceneOutliner
 		FItemSelection ItemSelection(*OutlinerTreeView);
 		FTreeItemPtr ItemToRename;
 
-		if (SharedData->Mode == ESceneOutlinerMode::ActorBrowsing)
+		if (Mode->CanRename())
 		{
-			// handle folders only here, actors and components are handled in LevelEditorActions::Rename_Execute
-			if (ItemSelection.Folders.Num() == 1 && ItemSelection.Folders.Num() == OutlinerTreeView->GetNumItemsSelected())
-			{
-				ItemToRename = OutlinerTreeView->GetSelectedItems()[0];
-			}
-		}
-		else if (SharedData->Mode == ESceneOutlinerMode::Custom)
-		{
-			if (OutlinerTreeView->GetNumItemsSelected() == 1)
-			{
-				ItemToRename = OutlinerTreeView->GetSelectedItems()[0];
-			}
+			ItemToRename = OutlinerTreeView->GetSelectedItems()[0];
 		}
 
-		if (ItemToRename.IsValid() && CanExecuteRenameRequest(ItemToRename) && ItemToRename->CanInteract())
+		if (ItemToRename.IsValid() && CanExecuteRenameRequest(*ItemToRename) && ItemToRename->CanInteract())
 		{
 			PendingRenameItem = ItemToRename->AsShared();
 			ScrollItemIntoView(ItemToRename);
@@ -1739,127 +812,37 @@ namespace SceneOutliner
 
 	bool SSceneOutliner::Cut_CanExecute()
 	{
-		if (SharedData->Mode == ESceneOutlinerMode::ActorBrowsing)
-		{
-			FItemSelection ItemSelection(*OutlinerTreeView);
-			if (ItemSelection.Folders.Num() > 0 && ItemSelection.Folders.Num() == OutlinerTreeView->GetNumItemsSelected())
-			{
-				return true;
-			}
-		}
-		return false;
+		return Mode->CanCut();
 	}
 
 	bool SSceneOutliner::Copy_CanExecute()
 	{
-		if (SharedData->Mode == ESceneOutlinerMode::ActorBrowsing)
-		{
-			FItemSelection ItemSelection(*OutlinerTreeView);
-			if (ItemSelection.Folders.Num() > 0 && ItemSelection.Folders.Num() == OutlinerTreeView->GetNumItemsSelected())
-			{
-				return true;
-			}
-		}
-		return false;
+		return Mode->CanCopy();
 	}
 
 	bool SSceneOutliner::Paste_CanExecute()
 	{
-		if (SharedData->Mode == ESceneOutlinerMode::ActorBrowsing)
-		{
-			if (CanPasteFoldersOnlyFromClipboard())
-			{
-				return true;
-			}
-		}
-		return false;
-	}
-
-	bool SSceneOutliner::CanPasteFoldersOnlyFromClipboard()
-	{
-		// Intentionally not checking if the level is locked/hidden here, as it's better feedback for the user if they attempt to paste
-		// and get the message explaining why it's failed, than just not having the option available to them.
-		FString PasteString;
-		FPlatformApplicationMisc::ClipboardPaste(PasteString);
-		return PasteString.StartsWith("BEGIN FOLDERLIST");
+		return Mode->CanPaste();
 	}
 
 	bool SSceneOutliner::CanSupportDragAndDrop() const
 	{
-		return SharedData->Mode == ESceneOutlinerMode::ActorBrowsing || SharedData->Mode == ESceneOutlinerMode::Custom;
+		return Mode->CanSupportDragAndDrop();
 	}
 
-	bool SSceneOutliner::CanExecuteRenameRequest(const FTreeItemPtr& ItemPtr) const
+	bool SSceneOutliner::CanExecuteRenameRequest(const ITreeItem& ItemPtr) const
 	{
-		if (TTreeItemGetter<bool>* Visitor = CanRenameItemVisitor.Get())
-		{
-			ItemPtr->Visit(*Visitor);
-			return Visitor->Result();
-		}
-
-		// Legacy default behavior
-		return true;
+		return Mode->CanRenameItem(ItemPtr);
 	}
 
 	int32 SSceneOutliner::AddFilter(const TSharedRef<SceneOutliner::FOutlinerFilter>& Filter)
 	{
-		// Deal with build in filters. If a build in filter is already in the filters, add will return it's index
-		if (Filter == HideTemporaryActorsFilter && !IsHidingTemporaryActors())
-		{
-			ToggleHideTemporaryActors();
-			return Filters->Num() - 1;
-		}
-		else if (Filter == ShowActorComponentsFilter && !IsShowingActorComponents())
-		{
-			ToggleShowActorComponents();
-			return Filters->Num() - 1;
-		}
-		else if (Filter == ShowOnlyActorsInCurrentLevelFilter && !IsShowingOnlyCurrentLevel())
-		{
-			ToggleShowOnlyCurrentLevel();
-			return Filters->Num() - 1;
-		}
-		else if (Filter == SelectedActorFilter && !IsShowingOnlySelected())
-		{
-			ToggleShowOnlySelected();
-			return Filters->Num() - 1;
-		}
-
-		// Custom Filter
 		return Filters->Add(Filter);
 	}
 
 	bool SSceneOutliner::RemoveFilter(const TSharedRef<SceneOutliner::FOutlinerFilter>& Filter)
 	{
-		bool bRemovedAFilter = false;
-
-		// Deal with build in filters. If a build in filter is already in the filters toggle to update the settings
-		if (Filter == HideTemporaryActorsFilter && IsHidingTemporaryActors())
-		{
-			ToggleHideTemporaryActors();
-			bRemovedAFilter = true;
-		}
-		else if (Filter == ShowActorComponentsFilter && IsShowingActorComponents())
-		{
-			ToggleShowActorComponents();
-			bRemovedAFilter = true;
-		}
-		else if (Filter == ShowOnlyActorsInCurrentLevelFilter && IsShowingOnlyCurrentLevel())
-		{
-			ToggleShowOnlyCurrentLevel();
-			bRemovedAFilter = true;
-		}
-		else if (Filter == SelectedActorFilter && IsShowingOnlySelected())
-		{
-			ToggleShowOnlySelected();
-			bRemovedAFilter = true;
-		}
-		else
-		{
-			bRemovedAFilter = Filters->Remove(Filter) > 0;
-		}
-
-		return bRemovedAFilter;
+		return  Filters->Remove(Filter) > 0;
 	}
 
 	TSharedPtr<SceneOutliner::FOutlinerFilter> SSceneOutliner::GetFilterAtIndex(int32 Index)
@@ -1897,210 +880,49 @@ namespace SceneOutliner
 		return ColumnsName;
 	}
 
-	ICustomSceneOutliner& SSceneOutliner::SetSelectionMode(ESelectionMode::Type InSelectionMode)
-	{
-		SelectionMode = InSelectionMode;
-		return *this;
-	}
-
-	ICustomSceneOutliner& SSceneOutliner::SetCanRenameItem(TUniquePtr<TTreeItemGetter<bool>>&& CanRenameItem)
-	{
-		CanRenameItemVisitor = MoveTemp(CanRenameItem);
-		return *this;
-	}
-
-	ICustomSceneOutliner& SSceneOutliner::SetShouldSelectItemWhenAdded(TUniquePtr<TTreeItemGetter<bool>>&& ShouldSelectItemWhenAdded)
-	{
-		ShouldSelectNewItemVisitor = MoveTemp(ShouldSelectItemWhenAdded);
-		return *this;
-	}
-
-	ICustomSceneOutliner& SSceneOutliner::SetOnItemDragDetected(TUniqueFunction<FReply(const SceneOutliner::ITreeItem&)> Callback)
-	{
-		OnItemDragDetected = MoveTemp(Callback);
-		return *this;
-	}
-
-	ICustomSceneOutliner& SSceneOutliner::SetOnDragOverItem(TUniqueFunction<FReply(const FDragDropEvent&, const SceneOutliner::ITreeItem&)> Callback)
-	{
-		OnDragOverItem = MoveTemp(Callback);
-		return *this;
-	}
-
-	ICustomSceneOutliner& SSceneOutliner::SetOnDropOnItem(TUniqueFunction<FReply(const FDragDropEvent&, const SceneOutliner::ITreeItem&)> Callback)
-	{
-		OnDropOnItem = MoveTemp(Callback);
-		return *this;
-	}
-
-	ICustomSceneOutliner& SSceneOutliner::SetOnDragEnterItem(TUniqueFunction<void(const FDragDropEvent&, const SceneOutliner::ITreeItem&)> Callback)
-	{
-		OnDragEnterItem = MoveTemp(Callback);
-		return *this;
-	}
-
-	ICustomSceneOutliner& SSceneOutliner::SetOnDragLeaveItem(TUniqueFunction<void(const FDragDropEvent&, const SceneOutliner::ITreeItem&)> Callback)
-	{
-		OnDragLeaveItem = MoveTemp(Callback);
-		return *this;
-	}
-
-
-	const TUniqueFunction<FReply(const SceneOutliner::ITreeItem&)>& SSceneOutliner::GetOnItemDragDetected() const
-	{
-		return OnItemDragDetected;
-	}
-
-	const TUniqueFunction<FReply(const FDragDropEvent&, const SceneOutliner::ITreeItem&)>& SSceneOutliner::GetOnDragOverItem() const
-	{
-		return OnDragOverItem;
-	}
-
-	const TUniqueFunction<FReply(const FDragDropEvent&, const SceneOutliner::ITreeItem&)>& SSceneOutliner::GetOnDropOnItem() const
-	{
-		return OnDropOnItem;
-	}
-
-	const TUniqueFunction<void(const FDragDropEvent&, const SceneOutliner::ITreeItem&)>& SSceneOutliner::GetOnDragEnterItem() const
-	{
-		return OnDragEnterItem;
-	}
-
-	const TUniqueFunction<void(const FDragDropEvent&, const SceneOutliner::ITreeItem&)>& SSceneOutliner::GetOnDragLeaveItem() const
-	{
-		return OnDragLeaveItem;
-	}
-
-	ICustomSceneOutliner& SSceneOutliner::SetUseSharedSceneOutlinerSettings(bool bUseSharedSettings)
-	{
-		if ( bUseSharedSettings && !SceneOutlinerSettings )
-		{
-			SceneOutlinerSettings = NewObject<USceneOutlinerSettings>();
-			ApplyHideTemporaryActorsFilter(IsHidingTemporaryActors());
-			ApplyShowActorComponentsFilter(IsShowingActorComponents());
-			ApplyShowOnlyCurrentLevelFilter(IsShowingOnlyCurrentLevel());
-			ApplyShowOnlySelectedFilter(IsShowingOnlySelected());
-		}
-		else if ( !bUseSharedSettings && SceneOutlinerSettings )
-		{
-			SceneOutlinerSettings = nullptr;
-			ApplyHideTemporaryActorsFilter(IsHidingTemporaryActors());
-			ApplyShowActorComponentsFilter(IsShowingActorComponents());
-			ApplyShowOnlyCurrentLevelFilter(IsShowingOnlyCurrentLevel());
-			ApplyShowOnlySelectedFilter(IsShowingOnlySelected());
-		}
-		return *this;
-	}
-
-	bool SSceneOutliner::IsUsingSharedSceneOutlinerSettings() const 
-	{
-		return SceneOutlinerSettings != nullptr;
-	}
-
-	ICustomSceneOutliner& SSceneOutliner::SetHideTemporaryActors(bool bHideTemporaryActors)
-	{
-		if ( bHideTemporaryActors != IsHidingTemporaryActors() )
-		{
-			ToggleHideTemporaryActors();
-		}
-		return *this;
-	}
-
-	ICustomSceneOutliner& SSceneOutliner::SetShowOnlyCurrentLevel(bool bShowOnlyCurrentLevel)
-	{
-		if ( bShowOnlyCurrentLevel != IsShowingOnlyCurrentLevel() )
-		{
-			ToggleShowOnlyCurrentLevel();
-		}
-		return *this;
-	}
-
-	ICustomSceneOutliner& SSceneOutliner::SetShownOnlySelected(bool bShowOnlySelected)
-	{
-		if ( bShowOnlySelected != IsShowingOnlySelected() )
-		{
-			ToggleShowOnlySelected();
-		}
-		return *this;
-	}
-
-	ICustomSceneOutliner& SSceneOutliner::SetShowActorComponents(bool bShowActorComponents)
-	{
-		if ( bShowActorComponents != IsShowingActorComponents() )
-		{
-			ToggleShowActorComponents();
-		}
-		return *this;
-	}
-
-	void SSceneOutliner::SetSelection(const SceneOutliner::TTreeItemGetter<bool>& ItemSelector)
+	void SSceneOutliner::SetSelection(const TFunctionRef<bool(SceneOutliner::ITreeItem&)> Selector)
 	{
 		TArray<FTreeItemPtr> ItemsToAdd;
-		for (const auto& TPairIdAndItem : TreeItemMap)
+		for (const auto& Pair : TreeItemMap)
 		{
-			FTreeItemPtr ItemPtr = TPairIdAndItem.Value;
-			if (ITreeItem*  Item = ItemPtr.Get())
+			FTreeItemPtr ItemPtr = Pair.Value;
+			if (ITreeItem* Item = ItemPtr.Get())
 			{
-				Item->Visit(ItemSelector);
-				if (ItemSelector.Result())
+				if (Selector(*Item))
 				{
 					ItemsToAdd.Add(ItemPtr);
 				}
 			}
 		}
+
+		SetItemSelection(ItemsToAdd, true);
+	}
+
+	void SSceneOutliner::SetItemSelection(const TArray<FTreeItemPtr>& InItems, bool bSelected, ESelectInfo::Type SelectInfo)
+	{
 		OutlinerTreeView->ClearSelection();
-		OutlinerTreeView->SetItemSelection(ItemsToAdd, true);
+		OutlinerTreeView->SetItemSelection(InItems, bSelected, SelectInfo);
 	}
 
-	void SSceneOutliner::AddToSelection(const SceneOutliner::TTreeItemGetter<bool>& ItemSelector)
+	void SSceneOutliner::SetItemSelection(const FTreeItemPtr& InItem, bool bSelected, ESelectInfo::Type SelectInfo)
 	{
-		TArray<FTreeItemPtr> ItemsToAdd;
-		for (const auto& TPairIdAndItem : TreeItemMap)
-		{
-			FTreeItemPtr ItemPtr = TPairIdAndItem.Value;
-			if (ITreeItem*  Item = ItemPtr.Get())
-			{
-				Item->Visit(ItemSelector);
-				if (ItemSelector.Result())
-				{
-					ItemsToAdd.Add(ItemPtr);
-				}
-			}
-		}
-		OutlinerTreeView->SetItemSelection(ItemsToAdd, true);
+		OutlinerTreeView->ClearSelection();
+		OutlinerTreeView->SetItemSelection(InItem, bSelected, SelectInfo);
 	}
 
-	void SSceneOutliner::RemoveFromSelection(const SceneOutliner::TTreeItemGetter<bool>& ItemDeselector)
+	void SSceneOutliner::AddToSelection(const TArray<FTreeItemPtr>& InItems, ESelectInfo::Type SelectInfo)
 	{
-		TArray<FTreeItemPtr> ItemsToRemove;
-		for (const FTreeItemPtr& ItemPtr : OutlinerTreeView->GetSelectedItems())
-		{
-			if (ITreeItem* Item = ItemPtr.Get())
-			{
-				Item->Visit(ItemDeselector);
-				if (ItemDeselector.Result())
-				{
-					ItemsToRemove.Add(ItemPtr);
-				}
-			}
-		}
-		OutlinerTreeView->SetItemSelection(ItemsToRemove, false);
+		OutlinerTreeView->SetItemSelection(InItems, true, SelectInfo);
 	}
 
-	void SSceneOutliner::AddObjectToSelection(const UObject* Object)
+	void SSceneOutliner::RemoveFromSelection(const TArray<FTreeItemPtr>& InItems)
 	{
-		if(FTreeItemPtr* ItemPtr = TreeItemMap.Find(Object))
-		{
-			OutlinerTreeView->SetItemSelection(*ItemPtr, true);
-		}
+		OutlinerTreeView->SetItemSelection(InItems, false);
 	}
 
-	void SSceneOutliner::RemoveObjectFromSelection(const UObject* Object)
+	void SSceneOutliner::RemoveFromSelection(const FTreeItemPtr& InItem)
 	{
-		if (FTreeItemPtr* ItemPtr = TreeItemMap.Find(Object))
-		{
-			OutlinerTreeView->SetItemSelection(*ItemPtr, false);
-		}
+		OutlinerTreeView->SetItemSelection(InItem, false);
 	}
 
 	void SSceneOutliner::AddFolderToSelection(const FName& FolderName)
@@ -2121,123 +943,9 @@ namespace SceneOutliner
 
 	void SSceneOutliner::ClearSelection()
 	{
-		OutlinerTreeView->ClearSelection();
-	}
-
-	TSharedPtr<SWidget> SSceneOutliner::BuildDefaultContextMenu()
-	{
-		if (!CheckWorld())
-		{
-			return nullptr;
-		}
-
-		RegisterDefaultContextMenu();
-
-		FItemSelection ItemSelection(*OutlinerTreeView);
-
-		USceneOutlinerMenuContext* ContextObject = NewObject<USceneOutlinerMenuContext>();
-		ContextObject->SceneOutliner = SharedThis(this);
-		ContextObject->bShowParentTree = SharedData->bShowParentTree;
-		ContextObject->NumSelectedItems = OutlinerTreeView->GetNumItemsSelected();
-		ContextObject->NumSelectedFolders = ItemSelection.Folders.Num();
-		ContextObject->NumWorldsSelected = ItemSelection.Worlds.Num();
-		FToolMenuContext Context(ContextObject);
-
-		// Allow other systems to override menu name and provide additional context
-		static const FName DefaultContextMenuName("SceneOutliner.DefaultContextMenu");
-		FName MenuName = DefaultContextMenuName;
-		SharedData->ModifyContextMenu.ExecuteIfBound(MenuName, Context);
-
-		// Build up the menu for a selection
-		UToolMenus* ToolMenus = UToolMenus::Get();
-		UToolMenu* Menu = ToolMenus->GenerateMenu(MenuName, Context);
-
-		for (const FToolMenuSection& Section : Menu->Sections)
-		{
-			if (Section.Blocks.Num() > 0)
-			{
-				return ToolMenus->GenerateWidget(Menu);
-			}
-		}
-
-		return nullptr;
-	}
-
-	void SSceneOutliner::RegisterDefaultContextMenu()
-	{
-		static const FName DefaultContextBaseMenuName("SceneOutliner.DefaultContextMenuBase");
-		static const FName DefaultContextMenuName("SceneOutliner.DefaultContextMenu");
-
-		UToolMenus* ToolMenus = UToolMenus::Get();
-
-		if (!ToolMenus->IsMenuRegistered(DefaultContextBaseMenuName))
+		if (!bIsReentrant)
 		{
-			UToolMenu* Menu = ToolMenus->RegisterMenu(DefaultContextBaseMenuName);
-
-			Menu->AddDynamicSection("DynamicSection1", FNewToolMenuDelegate::CreateLambda([](UToolMenu* InMenu)
-			{
-				USceneOutlinerMenuContext* Context = InMenu->FindContext<USceneOutlinerMenuContext>();
-				if (!Context || !Context->SceneOutliner.IsValid())
-				{
-					return;
-				}
-
-				SSceneOutliner* SceneOutliner = Context->SceneOutliner.Pin().Get();
-				if (Context->bShowParentTree)
-				{
-					if (Context->NumSelectedItems == 0)
-					{
-						InMenu->FindOrAddSection("Section").AddMenuEntry(
-							"CreateFolder",
-							LOCTEXT("CreateFolder", "Create Folder"),
-							FText(),
-							FSlateIcon(FEditorStyle::GetStyleSetName(), "SceneOutliner.NewFolderIcon"),
-							FUIAction(FExecuteAction::CreateSP(SceneOutliner, &SSceneOutliner::CreateFolder)));
-					}
-					else
-					{
-						if (Context->NumSelectedItems == 1)
-						{
-							SceneOutliner->GetTree().GetSelectedItems()[0]->GenerateContextMenu(InMenu, *SceneOutliner);
-						}
-
-						// If we've only got folders selected, show the selection and edit sub menus
-						if (Context->NumSelectedItems > 0 && Context->NumSelectedFolders == Context->NumSelectedItems)
-						{
-							InMenu->FindOrAddSection("Section").AddSubMenu(
-								"SelectSubMenu",
-								LOCTEXT("SelectSubmenu", "Select"),
-								LOCTEXT("SelectSubmenu_Tooltip", "Select the contents of the current selection"),
-								FNewToolMenuDelegate::CreateSP(SceneOutliner, &SSceneOutliner::FillSelectionSubMenu));
-						}
-					}
-				}
-			}));
-
-			Menu->AddDynamicSection("DynamicMainSection", FNewToolMenuDelegate::CreateLambda([](UToolMenu* InMenu)
-			{
-				// We always create a section here, even if there is no parent so that clients can still extend the menu
-				FToolMenuSection& Section = InMenu->AddSection("MainSection");
-
-				if (USceneOutlinerMenuContext* Context = InMenu->FindContext<USceneOutlinerMenuContext>())
-				{
-					// Don't add any of these menu items if we're not showing the parent tree
-					// Can't move worlds or level blueprints
-					if (Context->bShowParentTree && Context->NumSelectedItems > 0 && Context->NumWorldsSelected == 0 && Context->SceneOutliner.IsValid())
-					{
-						Section.AddSubMenu(
-							"MoveActorsTo",
-							LOCTEXT("MoveActorsTo", "Move To"),
-							LOCTEXT("MoveActorsTo_Tooltip", "Move selection to another folder"),
-							FNewToolMenuDelegate::CreateSP(Context->SceneOutliner.Pin().Get(), &SSceneOutliner::FillFoldersSubMenu));
-					}
-				}
-			}));
-		}
-
-		if (!ToolMenus->IsMenuRegistered(DefaultContextMenuName))
-		{
-			ToolMenus->RegisterMenu(DefaultContextMenuName, DefaultContextBaseMenuName);
+			OutlinerTreeView->ClearSelection();
 		}
 	}
 
@@ -2255,69 +963,40 @@ namespace SceneOutliner
 		// We use a pointer here to save copying the whole array for every invocation of the filter delegate
 		TSharedRef<TSet<FName>> ExcludedParents(new TSet<FName>());
 
-		struct FFindInvalidFolders : ITreeItemVisitor
+		for (const auto& Item : OutlinerTreeView->GetSelectedItems())
 		{
-			TSet<FName>& ExcludedParents;
-			const TMap<FTreeItemID, FTreeItemPtr>& TreeItemMap;
-			FFindInvalidFolders(TSet<FName>& InExcludedParents, const TMap<FTreeItemID, FTreeItemPtr>& InTreeItemMap)
-				: ExcludedParents(InExcludedParents), TreeItemMap(InTreeItemMap)
-			{}
-
-			static bool ItemHasSubFolders(const TWeakPtr<ITreeItem>& WeakItem)
-			{
-				bool bHasSubFolder = false;
-				WeakItem.Pin()->Visit(FFunctionalVisitor().Folder([&](const FFolderTreeItem&){
-					bHasSubFolder = true;
-				}));
-				return bHasSubFolder;
-			}
-
-			virtual void Visit(const FActorTreeItem& ActorItem) const override
+			if (FFolderTreeItem* ParentFolderItem = Item->GetParent()->CastTo<FFolderTreeItem>())
 			{
-				if (AActor* Actor = ActorItem.Actor.Get())
+				auto FolderHasOtherSubFolders = [&Item](const TWeakPtr<ITreeItem>& WeakItem)
 				{
-					// We exclude actor parent folders if they don't have any sub folders
-					const FName& Folder = Actor->GetFolderPath();
-					if (!Folder.IsNone() && !ExcludedParents.Contains(Folder))
+					if (WeakItem.Pin()->IsA<FFolderTreeItem>() && WeakItem.Pin() != Item)
 					{
-						FTreeItemPtr FolderItem = TreeItemMap.FindRef(Folder);
-						if (FolderItem.IsValid() && !FolderItem->GetChildren().ContainsByPredicate(&ItemHasSubFolders))
-						{
-							ExcludedParents.Add(Folder);
-						}
+						return true;
 					}
-				}
-			}
+					return false;
+				};
 
-			virtual void Visit(const FFolderTreeItem& Folder) const override
-			{
-				// Cannot move into its parent
-				const FName ParentPath = GetParentPath(Folder.Path);
-				if (!ParentPath.IsNone())
-				{
-					ExcludedParents.Add(ParentPath);
-				}
-				else
+				// Exclude this items direct parent if it is a folder and has no other subfolders we can move to
+				if (!ParentFolderItem->GetChildren().ContainsByPredicate(FolderHasOtherSubFolders))
 				{
-					// Failing that, cannot move into itself, or any child
-					ExcludedParents.Add(Folder.Path);
+					ExcludedParents->Add(ParentFolderItem->Path);
 				}
 			}
-		};
-
-		auto Visitor = FFindInvalidFolders(*ExcludedParents, TreeItemMap);
-		for (const auto& Item : OutlinerTreeView->GetSelectedItems())
-		{
-			Item->Visit(Visitor);
+			
+			if (FFolderTreeItem* FolderItem = Item->CastTo<FFolderTreeItem>())
+			{
+				// Cannot move into itself, or any child
+				ExcludedParents->Add(FolderItem->Path);
+			}
 		}
 
 		return ExcludedParents;
 	}
-
+	
 	void SSceneOutliner::AddMoveToFolderOutliner(UToolMenu* Menu) const
 	{
-		// We don't show this if there aren't any folders in the world
-		if (!FActorFolders::Get().GetFolderPropertiesForWorld(*SharedData->RepresentingWorld).Num())
+		// We don't show this if there aren't any folders in the world and if the mode is showing folders
+		if (!Mode->ShouldShowFolders() || FolderCount == 0)
 		{
 			return;
 		}
@@ -2326,17 +1005,16 @@ namespace SceneOutliner
 		FInitializationOptions MiniSceneOutlinerInitOptions;
 		MiniSceneOutlinerInitOptions.bShowHeaderRow = false;
 		MiniSceneOutlinerInitOptions.bFocusSearchBoxWhenOpened = true;
-		MiniSceneOutlinerInitOptions.bOnlyShowFolders = true;
 		
 		// Don't show any folders that are a child of any of the selected folders
 		auto ExcludedParents = GatherInvalidMoveToDestinations();
 		if (ExcludedParents->Num())
 		{
 			// Add a filter if necessary
-			auto FilterOutChildFolders = [](FName Path, TSharedRef<TSet<FName>> InExcludedParents){
-				for (const auto& Parent : *InExcludedParents)
+			auto FilterOutChildFolders = [](FName Path, TSharedRef<TSet<FName>> ExcludedParents){
+				for (const auto& Parent : *ExcludedParents)
 				{
-					if (Path == Parent || FActorFolders::PathIsChildOf(Path.ToString(), Parent.ToString()))
+					if (Path == Parent || PathIsChildOf(Path, Parent))
 					{
 						return false;
 					}
@@ -2344,25 +1022,52 @@ namespace SceneOutliner
 				return true;
 			};
 
-			MiniSceneOutlinerInitOptions.Filters->AddFilterPredicate(FFolderFilterPredicate::CreateStatic(FilterOutChildFolders, ExcludedParents), EDefaultFilterBehaviour::Pass);
+			MiniSceneOutlinerInitOptions.Filters->AddFilterPredicate<FFolderTreeItem>(FFolderTreeItem::FFilterPredicate::CreateStatic(FilterOutChildFolders, ExcludedParents), EDefaultFilterBehaviour::Pass);
 		}
 
 		{
-			// Filter in/out the world according to whether it is valid to move to/from the root
-			FDragDropPayload DraggedObjects(OutlinerTreeView->GetSelectedItems());
+			struct FFilterRoot : public FOutlinerFilter 
+			{
+				FFilterRoot(TFunctionRef<bool(const ITreeItem&)> InFindParent) 
+					: FOutlinerFilter(EDefaultFilterBehaviour::Pass) 
+					, FindParent(InFindParent)
+				{}
+				virtual bool PassesFilter(const ITreeItem& Item) const override
+				{
+					// if item has no parent, it is a root item and should be filtered out
+					if (!FindParent(Item))
+					{
+						return false;
+					}
+					return DefaultBehaviour == EDefaultFilterBehaviour::Pass;
+				}
+
+				TFunctionRef<bool(const ITreeItem&)> FindParent;
+			};
 
-			const bool bMoveToRootValid = FFolderDropTarget(FName()).ValidateDrop(DraggedObjects, *SharedData->RepresentingWorld).IsValid();
+			auto FindParent = [this](const ITreeItem& Item)
+			{
+				FTreeItemPtr Parent = Mode->GetHierarchy()->FindParent(Item, TreeItemMap);
+				if (!Parent)
+				{
+					Parent = Mode->GetHierarchy()->FindParent(Item, PendingTreeItemMap);
+				}
+				return Parent.IsValid();
+			};
 
-			MiniSceneOutlinerInitOptions.Filters->AddFilterPredicate(FWorldFilterPredicate::CreateStatic([](const UWorld*, bool bInMoveToRootValid){
-				return bInMoveToRootValid;
-			}, bMoveToRootValid), EDefaultFilterBehaviour::Pass);
+			// Filter in/out root items according to whether it is valid to move to/from the root
+			FDragDropPayload DraggedObjects(OutlinerTreeView->GetSelectedItems());
+			const bool bMoveToRootValid = Mode->ValidateDrop(FFolderTreeItem(NAME_None), DraggedObjects).IsValid();
+			if (!bMoveToRootValid)
+			{
+				MiniSceneOutlinerInitOptions.Filters->Add(MakeShared<FFilterRoot>(FindParent));
+			}
 		}
 
-		// Don't show the actor info column
-		MiniSceneOutlinerInitOptions.UseDefaultColumns();
-		MiniSceneOutlinerInitOptions.ColumnMap.Remove(FBuiltInColumnTypes::ActorInfo());
+		//Let the mode decide how folder selection is handled
+
+		MiniSceneOutlinerInitOptions.ModeFactory = Mode->CreateFolderPickerMode();
 
-		// Actor selector to allow the user to choose a folder
 		FSceneOutlinerModule& SceneOutlinerModule = FModuleManager::LoadModuleChecked<FSceneOutlinerModule>("SceneOutliner");
 		TSharedRef< SWidget > MiniSceneOutliner =
 			SNew(SVerticalBox)
@@ -2371,7 +1076,6 @@ namespace SceneOutliner
 			[
 				SNew(SSceneOutliner, MiniSceneOutlinerInitOptions)
 				.IsEnabled(FSlateApplication::Get().GetNormalExecutionAttribute())
-				.OnItemPickedDelegate(FOnSceneOutlinerItemPicked::CreateSP(const_cast<SSceneOutliner*>(this), &SSceneOutliner::MoveSelectionTo))
 			];
 
 		FToolMenuSection& Section = Menu->AddSection(FName(), LOCTEXT("ExistingFolders", "Existing:"));
@@ -2386,191 +1090,41 @@ namespace SceneOutliner
 	{
 		FToolMenuSection& Section = Menu->AddSection("Section");
 		Section.AddMenuEntry(
-			"AddChildrenToSelection",
-			LOCTEXT( "AddChildrenToSelection", "Immediate Children" ),
-			LOCTEXT( "AddChildrenToSelection_ToolTip", "Select all immediate actor children of the selected folders" ),
-			FSlateIcon(),
-			FExecuteAction::CreateSP(const_cast<SSceneOutliner*>(this), &SSceneOutliner::SelectFoldersDescendants, /*bSelectImmediateChildrenOnly=*/ true));
-		Section.AddMenuEntry(
-			"AddDescendantsToSelection",
-			LOCTEXT( "AddDescendantsToSelection", "All Descendants" ),
-			LOCTEXT( "AddDescendantsToSelection_ToolTip", "Select all actor descendants of the selected folders" ),
-			FSlateIcon(),
-			FExecuteAction::CreateSP(const_cast<SSceneOutliner*>(this), &SSceneOutliner::SelectFoldersDescendants, /*bSelectImmediateChildrenOnly=*/ false));
-	}
-
-	void SSceneOutliner::SelectFoldersDescendants(bool bSelectImmediateChildrenOnly)
-	{
-		struct FExpandFoldersRecursive : IMutableTreeItemVisitor
-		{
-			SSceneOutliner& Outliner;
-			bool bSelectImmediateChildrenOnly;
-
-			FExpandFoldersRecursive(SSceneOutliner& InOutliner, bool bInSelectImmediateChildrenOnly) :
-				Outliner(InOutliner), bSelectImmediateChildrenOnly(bInSelectImmediateChildrenOnly) {}
-
-			virtual void Visit(FActorTreeItem& ActorItem) const override
-			{
-				if (!Outliner.OutlinerTreeView->IsItemExpanded(ActorItem.AsShared()))
-				{
-					Outliner.OutlinerTreeView->SetItemExpansion(ActorItem.AsShared(), true);
-				}
-
-				if (!bSelectImmediateChildrenOnly)
-				{
-					for (TWeakPtr<ITreeItem> Child : ActorItem.GetChildren())
-					{
-						Child.Pin()->Visit(*this);
-					}
-				}
-			}
-
-			virtual void Visit(FFolderTreeItem& FolderItem) const override
-			{
-				if (!Outliner.OutlinerTreeView->IsItemExpanded(FolderItem.AsShared()))
-				{
-					Outliner.OutlinerTreeView->SetItemExpansion(FolderItem.AsShared(), true);
-				}
-
-				if (!bSelectImmediateChildrenOnly)
-				{
-					for (TWeakPtr<ITreeItem> Child : FolderItem.GetChildren())
-					{
-						Child.Pin()->Visit(*this);
-					}
-				}
-
-			}
-		};
-
-		struct FSelectActorsRecursive : ITreeItemVisitor
-		{
-			bool bSelectImmediateChildrenOnly;
-			FSelectActorsRecursive(bool bInSelectImmediateChildrenOnly) : bSelectImmediateChildrenOnly(bInSelectImmediateChildrenOnly) {}
-
-			virtual void Visit(const FActorTreeItem& ActorItem) const override
-			{
-				if (AActor* Actor = ActorItem.Actor.Get())
-				{
-					GEditor->SelectActor(Actor, true, /*bNotify=*/false);
-				}
-
-				if (!bSelectImmediateChildrenOnly)
-				{
-					for (TWeakPtr<ITreeItem> Child : ActorItem.GetChildren())
-					{
-						Child.Pin()->Visit(*this);
-					}
-				}
-			}
-
-			virtual void Visit(const FFolderTreeItem& FolderItem) const override
-			{
-				if (!bSelectImmediateChildrenOnly)
-				{
-					for (TWeakPtr<ITreeItem> Child : FolderItem.GetChildren())
-					{
-						Child.Pin()->Visit(*this);
-					}
-				}
-			}
-		};
-
-		struct FSelectFoldersRecursive : IMutableTreeItemVisitor
-		{
-			SSceneOutliner& Outliner;
-			bool bSelectImmediateChildrenOnly;
-			FSelectFoldersRecursive(SSceneOutliner& InOutliner, bool bInSelectImmediateChildrenOnly) : Outliner(InOutliner), bSelectImmediateChildrenOnly(bInSelectImmediateChildrenOnly) {}
-
-			virtual void Visit(FFolderTreeItem& FolderItem) const override
-			{
-				Outliner.OutlinerTreeView->SetItemSelection(FolderItem.AsShared(), true);
-
-				if (!bSelectImmediateChildrenOnly)
-				{
-					for (TWeakPtr<ITreeItem> Child : FolderItem.GetChildren())
-					{
-						Child.Pin()->Visit(*this);
-					}
-				}	
-			}
-		};
+			"AddChildrenToSelection",
+			LOCTEXT( "AddChildrenToSelection", "Immediate Children" ),
+			LOCTEXT( "AddChildrenToSelection_ToolTip", "Select all immediate children of the selected folders" ),
+			FSlateIcon(),
+			FExecuteAction::CreateSP(const_cast<SSceneOutliner*>(this), &SSceneOutliner::SelectFoldersDescendants, /*bSelectImmediateChildrenOnly=*/ true));
+		Section.AddMenuEntry(
+			"AddDescendantsToSelection",
+			LOCTEXT( "AddDescendantsToSelection", "All Descendants" ),
+			LOCTEXT( "AddDescendantsToSelection_ToolTip", "Select all descendants of the selected folders" ),
+			FSlateIcon(),
+			FExecuteAction::CreateSP(const_cast<SSceneOutliner*>(this), &SSceneOutliner::SelectFoldersDescendants, /*bSelectImmediateChildrenOnly=*/ false));
+	}
 
-		TArray<FFolderTreeItem*> SelectedFolders = GetSelectedFolders();
+	void SSceneOutliner::SelectFoldersDescendants(bool bSelectImmediateChildrenOnly)
+	{
+		TArray<FFolderTreeItem*> SelectedFolders;
+		GetSelectedFolders(SelectedFolders);
 		OutlinerTreeView->ClearSelection();
 
-		FExpandFoldersRecursive ExpandFoldersRecursive(*this, bSelectImmediateChildrenOnly);
-		for (FFolderTreeItem* Folder : SelectedFolders)
-		{
-			Folder->Visit(ExpandFoldersRecursive);
-		}
-
 		if (SelectedFolders.Num())
 		{
-			// We'll batch selection changes instead by using BeginBatchSelectOperation()
-			GEditor->GetSelectedActors()->BeginBatchSelectOperation();
-
-			FSelectActorsRecursive SelectActorsRecursive(bSelectImmediateChildrenOnly);
-			for (const FFolderTreeItem* Folder : SelectedFolders)
-			{
-				for (TWeakPtr<ITreeItem> Child : Folder->GetChildren())
-				{
-					Child.Pin()->Visit(SelectActorsRecursive);
-				}
-			}
-
-			GEditor->GetSelectedActors()->EndBatchSelectOperation(/*bNotify*/false);
-			GEditor->NoteSelectionChange();
-		}
-
-		// Don't select folders, only select actors
-		/*  
-		FSelectFoldersRecursive SelectFoldersRecursiveVisitor(*this, bSelectImmediateChildrenOnly);
-		for (const FFolderTreeItem* Folder : SelectedFolders)
-		{
-			for (TWeakPtr<ITreeItem> Child : Folder->GetChildren())
-			{
-				Child.Pin()->Visit(SelectFoldersRecursiveVisitor);
-			}
+			Mode->SelectFoldersDescendants(SelectedFolders, bSelectImmediateChildrenOnly);
 		}
-		*/
 
 		Refresh();
 	}
 
-	void SSceneOutliner::MoveSelectionTo(FTreeItemRef NewParent)
-	{
-		struct FMoveToFolder : ITreeItemVisitor
-		{
-			SSceneOutliner& Outliner;
-			FMoveToFolder(SSceneOutliner& InOutliner) : Outliner(InOutliner) {}
-
-			virtual void Visit(const FFolderTreeItem& Folder) const override
-			{
-				Outliner.MoveSelectionTo(Folder.Path);
-			}
-			virtual void Visit(const FWorldTreeItem&) const override
-			{
-				Outliner.MoveSelectionTo(FName());
-			}
-		};
-
-		NewParent->Visit(FMoveToFolder(*this));
-	}
-
 	void SSceneOutliner::MoveSelectionTo(FName NewParent)
 	{
-		if (!CheckWorld())
-		{
-			return;
-		}
-
 		FSlateApplication::Get().DismissAllMenus();
 		
-		FFolderDropTarget 	DropTarget(NewParent);
+		FFolderTreeItem 	DropTarget(NewParent);
 		FDragDropPayload 	DraggedObjects(OutlinerTreeView->GetSelectedItems());
 
-		FDragValidationInfo Validation = DropTarget.ValidateDrop(DraggedObjects, *SharedData->RepresentingWorld);
+		FDragValidationInfo Validation = Mode->ValidateDrop(DropTarget, DraggedObjects);
 		if (!Validation.IsValid())
 		{
 			FNotificationInfo Info(Validation.ValidationText);
@@ -2582,8 +1136,8 @@ namespace SceneOutliner
 			return;
 		}
 
-		const FScopedTransaction Transaction( LOCTEXT("MoveOutlinerItems", "Move World Outliner Items") );
-		DropTarget.OnDrop(DraggedObjects, *SharedData->RepresentingWorld, Validation, SNullWidget::NullWidget);
+		const FScopedTransaction Transaction(LOCTEXT("MoveOutlinerItems", "Move Outliner Items"));
+		Mode->OnDrop(DropTarget, DraggedObjects, Validation);
 	}
 
 	FReply SSceneOutliner::OnCreateFolderClicked()
@@ -2594,163 +1148,31 @@ namespace SceneOutliner
 
 	void SSceneOutliner::CreateFolder()
 	{
-		if (!CheckWorld())
-		{
-			return;
-		}
-
-		UWorld& World = *SharedData->RepresentingWorld;
-		const FScopedTransaction Transaction(LOCTEXT("UndoAction_CreateFolder", "Create Folder"));
-
-		const FName NewFolderName = FActorFolders::Get().GetDefaultFolderNameForSelection(World);
-		FActorFolders::Get().CreateFolderContainingSelection(World, NewFolderName);
+		const FName& NewFolderName = Mode->CreateNewFolder();
 
+		// Move any selected folders into the new folder
 		auto PreviouslySelectedItems = OutlinerTreeView->GetSelectedItems();
-
-		auto Visit = [&](const FFolderTreeItem& Folder)
-		{
-			MoveFolderTo(Folder.Path, NewFolderName, World);
-		};
-		auto Visitor = FFunctionalVisitor().Folder(Visit);
-
-		// Move any selected folders into the new folder name
 		for (const auto& Item : PreviouslySelectedItems)
 		{
-			Item->Visit(Visitor);
+			if (FFolderTreeItem* FolderItem = Item->CastTo<FFolderTreeItem>())
+			{
+				FolderItem->MoveTo(NewFolderName);
+			}
 		}
 
 		// At this point the new folder will be in our newly added list, so select it and open a rename when it gets refreshed
 		NewItemActions.Add(NewFolderName, ENewItemAction::Select | ENewItemAction::Rename);
-	}
-
-	void SSceneOutliner::OnBroadcastFolderCreate(UWorld& InWorld, FName NewPath)
-	{
-		if (!ShouldShowFolders() || &InWorld != SharedData->RepresentingWorld)
-		{
-			return;
-		}
-
-		if (!TreeItemMap.Contains(NewPath))
-		{
-			ConstructItemFor<FFolderTreeItem>(NewPath);
-		}
-	}
-
-	void SSceneOutliner::OnBroadcastFolderMove(UWorld& InWorld, FName OldPath, FName NewPath)
-	{
-		if (!ShouldShowFolders() || &InWorld != SharedData->RepresentingWorld)
-		{
-			return;
-		}
-
-		FTreeItemPtr Item = TreeItemMap.FindRef(OldPath);
-		if (Item.IsValid())
-		{
-			// Remove it from the map under the old ID (which is derived from the folder path)
-			TreeItemMap.Remove(Item->GetID());
-
-			// Now change the path and put it back in the map with its new ID
-			auto Folder = StaticCastSharedPtr<FFolderTreeItem>(Item);
-			Folder->Path = NewPath;
-			Folder->LeafName = GetFolderLeafName(NewPath);
-
-			TreeItemMap.Add(Item->GetID(), Item);
-
-			// Add an operation to move the item in the hierarchy
-			PendingOperations.Emplace(FPendingTreeOperation::Moved, Item.ToSharedRef());
-			Refresh();
-		}
-	}
-
-	void SSceneOutliner::OnBroadcastFolderDelete(UWorld& InWorld, FName Path)
-	{
-		if (&InWorld != SharedData->RepresentingWorld)
-		{
-			return;
-		}
-
-		auto* Folder = TreeItemMap.Find(Path);
-		if (Folder)
-		{
-			PendingOperations.Emplace(FPendingTreeOperation::Removed, Folder->ToSharedRef());
-			Refresh();
-		}
-	}
-
-	void SSceneOutliner::OnEditCutActorsBegin()
-	{
-		// Only a callback in actor browsing mode
-		CopyFoldersBegin();
-		DeleteFoldersBegin();
-	}
-
-	void SSceneOutliner::OnEditCutActorsEnd()
-	{
-		// Only a callback in actor browsing mode
-		CopyFoldersEnd();
-		DeleteFoldersEnd();
-	}
-
-	void SSceneOutliner::OnEditCopyActorsBegin()
-	{
-		// Only a callback in actor browsing mode
-		CopyFoldersBegin();
-	}
-
-	void SSceneOutliner::OnEditCopyActorsEnd()
-	{
-		// Only a callback in actor browsing mode
-		CopyFoldersEnd();
-	}
-
-	void SSceneOutliner::OnEditPasteActorsBegin()
-	{
-		// Only a callback in actor browsing mode
-		TArray<FName> Folders = GetClipboardPasteFolders();
-		PasteFoldersBegin(Folders);
-	}
-
-	void SSceneOutliner::OnEditPasteActorsEnd()
-	{
-		// Only a callback in actor browsing mode
-		PasteFoldersEnd();
-	}
-
-	void SSceneOutliner::OnDuplicateActorsBegin()
-	{
-		// Only a callback in actor browsing mode
-		TArray<FFolderTreeItem*> SelectedFolders = GetSelectedFolders();
-		PasteFoldersBegin(SelectedFolders);
-	}
-
-	void SSceneOutliner::OnDuplicateActorsEnd()
-	{
-		// Only a callback in actor browsing mode
-		PasteFoldersEnd();
-	}
-
-	void SSceneOutliner::OnDeleteActorsBegin()
-	{
-		// Only a callback in actor browsing mode
-		DeleteFoldersBegin();
-	}
-
-	void SSceneOutliner::OnDeleteActorsEnd()
-	{
-		// Only a callback in actor browsing mode
-		DeleteFoldersEnd();
+		RequestSort();
 	}
 
 	void SSceneOutliner::CopyFoldersBegin()
 	{
-		// Only a callback in actor browsing mode
 		CacheFoldersEdit = GetSelectedFolderNames();
 		FPlatformApplicationMisc::ClipboardPaste(CacheClipboardContents);
 	}
 
 	void SSceneOutliner::CopyFoldersEnd()
 	{
-		// Only a callback in actor browsing mode
 		if (CacheFoldersEdit.Num() > 0)
 		{
 			CopyFoldersToClipboard(CacheFoldersEdit, CacheClipboardContents);
@@ -2762,7 +1184,7 @@ namespace SceneOutliner
 	{
 		if (InFolders.Num() > 0)
 		{
-			// If clipboard paste has changed since we cached it, actors must have been cut 
+			// If clipboard paste has changed since we cached it, items must have been cut 
 			// so folders need to appended to clipboard contents rather than replacing them.
 			FString CurrClipboardContents;
 			FPlatformApplicationMisc::ClipboardPaste(CurrClipboardContents);
@@ -2785,33 +1207,14 @@ namespace SceneOutliner
 		}
 	}
 
-	void SSceneOutliner::PasteFoldersBegin(TArray<FFolderTreeItem*> InFolders)
-	{
-		TArray<FName> FolderNames;
-
-		for (FFolderTreeItem* Folder : InFolders)
-		{
-			if (Folder)
-			{
-				FolderNames.Add(Folder->Path);
-			}
-		}
-
-		PasteFoldersBegin(FolderNames);
-	}
-
 	void SSceneOutliner::PasteFoldersBegin(TArray<FName> InFolders)
 	{
-		struct FCacheExistingChildrenAction : IMutableTreeItemVisitor
+		auto CacheExistingChildrenAction = [this](const FTreeItemPtr& Item)
 		{
-			SSceneOutliner& Outliner;
-
-			FCacheExistingChildrenAction(SSceneOutliner& InOutliner) : Outliner(InOutliner) {}
-
-			virtual void Visit(FFolderTreeItem& FolderItem) const override
+			if (FFolderTreeItem* FolderItem = Item->CastTo<FFolderTreeItem>())
 			{
 				TArray<FTreeItemID> ExistingChildren;
-				for (const TWeakPtr<ITreeItem>& Child : FolderItem.GetChildren())
+				for (const TWeakPtr<ITreeItem>& Child : FolderItem->GetChildren())
 				{
 					if (Child.IsValid())
 					{
@@ -2819,7 +1222,7 @@ namespace SceneOutliner
 					}
 				}
 
-				Outliner.CachePasteFolderExistingChildrenMap.Add(FolderItem.Path, ExistingChildren);
+				CachePasteFolderExistingChildrenMap.Add(FolderItem->Path, ExistingChildren);
 			}
 		};
 
@@ -2838,69 +1241,13 @@ namespace SceneOutliner
 		{
 			if (FTreeItemPtr* TreeItem = TreeItemMap.Find(Folder))
 			{
-				FCacheExistingChildrenAction CacheExistingChildrenAction(*this);
-				(*TreeItem)->Visit(CacheExistingChildrenAction);
+				CacheExistingChildrenAction(*TreeItem);
 			}
 		}
 	}
 
 	void SSceneOutliner::PasteFoldersEnd()
 	{
-		struct FReparentDuplicatedActorsAction : IMutableTreeItemVisitor
-		{
-			SSceneOutliner& Outliner;
-			mutable bool bVisitedFolder;
-			mutable FName* NewFolderPath;
-			TMap<FName, FName>* FolderMap;
-
-			FReparentDuplicatedActorsAction(SSceneOutliner& InOutliner, TMap<FName, FName>* InFolderMap) : Outliner(InOutliner), bVisitedFolder(false), NewFolderPath(nullptr), FolderMap(InFolderMap) {}
-
-			virtual void Visit(FFolderTreeItem& FolderItem) const override
-			{
-				if (!bVisitedFolder)
-				{
-					bVisitedFolder = true;
-
-					if (!FolderItem.Path.IsNone())
-					{
-						NewFolderPath = FolderMap->Find(FolderItem.Path);
-						if (NewFolderPath && *NewFolderPath != FolderItem.Path)
-						{
-							for (TWeakPtr<ITreeItem> Child : FolderItem.GetChildren())
-							{
-								if (Child.IsValid())
-								{
-									Child.Pin()->Visit(*this);
-								}
-							}
-						}
-					}
-				}
-			}
-
-			virtual void Visit(FActorTreeItem& ActorItem) const override
-			{
-				if (NewFolderPath)
-				{
-					if (AActor* Actor = ActorItem.Actor.Get())
-					{
-						FName ParentPath = Actor->GetFolderPath();
-						TArray<FTreeItemID>* ExistingChildren = Outliner.CachePasteFolderExistingChildrenMap.Find(ParentPath);
-
-						if (ExistingChildren && !ExistingChildren->Contains(ActorItem.GetID()))
-						{
-							Actor->SetFolderPath_Recursively(*NewFolderPath);
-						}
-					}
-				}
-			}
-		};
-
-		if (!CheckWorld())
-		{
-			return;
-		}
-			
 		const FScopedTransaction Transaction(NSLOCTEXT("UnrealEd", "PasteItems", "Paste Items"));
 
 		// Create new folder
@@ -2916,8 +1263,7 @@ namespace SceneOutliner
 					ParentPath = *NewParentPath;
 				}
 
-				FName NewFolderPath = FActorFolders::Get().GetFolderName(*SharedData->RepresentingWorld, ParentPath, LeafName);
-				FActorFolders::Get().CreateFolder(*SharedData->RepresentingWorld, NewFolderPath);
+				FName NewFolderPath = Mode->CreateFolder(ParentPath, LeafName);
 				FolderMap.Add(Folder, NewFolderPath);
 			}
 		}
@@ -2925,21 +1271,27 @@ namespace SceneOutliner
 		// Populate our data set
 		Populate();
 
-		// Reparent duplicated actors if the folder has been pasted/duplicated
-		for (FName Folder : CacheFoldersEdit)
+		// Reparent duplicated items if the folder has been pasted/duplicated
+		for (FName OldFolderName : CacheFoldersEdit)
 		{
-			if (const FName* NewFolder = FolderMap.Find(Folder))
+			// Get the new folder that was created from this name
+			if (const FName* NewFolderName = FolderMap.Find(OldFolderName))
 			{
-				if (FTreeItemPtr* FolderItem = TreeItemMap.Find(Folder))
+				if (FTreeItemPtr* OldFolderItem = TreeItemMap.Find(OldFolderName))
 				{
-					FReparentDuplicatedActorsAction ReparentDuplicatedActors(*this, &FolderMap);
-					(*FolderItem)->Visit(ReparentDuplicatedActors);
-				}
+					for (const TWeakPtr<ITreeItem>& Child : (*OldFolderItem)->GetChildren())
+					{
+						// If this child did not exist in the folder before the paste operation, it should be moved to the new folder
+						TArray<FTreeItemID>* ExistingChildren = CachePasteFolderExistingChildrenMap.Find(OldFolderName);
 
-				Folder = *NewFolder;
+						if (ExistingChildren && !ExistingChildren->Contains(Child.Pin()->GetID()))
+						{
+							Mode->ReparentItemToFolder(*NewFolderName, Child.Pin());
+						}
+					}
+				}
+				PendingFoldersSelect.Add(*NewFolderName);
 			}
-
-			PendingFoldersSelect.Add(Folder);
 		}
 
 		CacheFoldersEdit.Reset();
@@ -2949,60 +1301,47 @@ namespace SceneOutliner
 
 	void SSceneOutliner::DuplicateFoldersHierarchy()
 	{
-		struct FSelectFoldersRecursive : IMutableTreeItemVisitor
+		TFunctionRef<void(const FTreeItemPtr&)> RecursiveFolderSelect = [&](const FTreeItemPtr& Item)
 		{
-			SSceneOutliner& Outliner;
-			FSelectFoldersRecursive(SSceneOutliner& InOutliner) : Outliner(InOutliner) {}
-
-			virtual void Visit(FFolderTreeItem& FolderItem) const override
+			if (Item->IsA<FFolderTreeItem>())
 			{
-				// select folders to be duplicated
-				Outliner.OutlinerTreeView->SetItemSelection(FolderItem.AsShared(), true);
-
-				// FItemSelection Selection(*Outliner.OutlinerTreeView);
-
-				for (TWeakPtr<ITreeItem> Child : FolderItem.GetChildren())
-				{
-					Child.Pin()->Visit(*this);
-				}
+				OutlinerTreeView->SetItemSelection(Item, true);
+			}
+			for (const TWeakPtr<ITreeItem> Child : Item->GetChildren())
+			{
+				RecursiveFolderSelect(Child.Pin());
 			}
 		};
 
-		if (!CheckWorld())
-		{
-			return;
-		}
-
 		const FScopedTransaction Transaction(NSLOCTEXT("UnrealEd", "DuplicateFoldersHierarchy", "Duplicate Folders Hierarchy"));
 
-		TArray<FFolderTreeItem*> SelectedFolders = GetSelectedFolders();
+		TArray<FFolderTreeItem*> SelectedFolders;
+		GetSelectedFolders(SelectedFolders);
 
 		if (SelectedFolders.Num() > 0)
 		{
-			// Select actor descendants
+			// Select folder descendants
 			SelectFoldersDescendants();
 
 			// Select all sub-folders
 			for (FFolderTreeItem* Folder : SelectedFolders)
 			{
-				FSelectFoldersRecursive SelectFoldersRecursive(*this);
-				Folder->Visit(SelectFoldersRecursive);
+				RecursiveFolderSelect(Folder->AsShared());
 			}
 
 			// Duplicate selected
-			GUnrealEd->Exec(SharedData->RepresentingWorld, TEXT("DUPLICATE"));
+			Mode->OnDuplicateSelected();
 		}
 	}
 
 	void SSceneOutliner::DeleteFoldersBegin()
 	{
-		// Only a callback in actor browsing mode
-		CacheFoldersDelete = GetSelectedFolders();
+		CacheFoldersDelete.Empty();
+		GetSelectedFolders(CacheFoldersDelete);
 	}
 
 	void SSceneOutliner::DeleteFoldersEnd()
 	{
-		// Only a callback in actor browsing mode
 		struct FMatchName
 		{
 			FMatchName(const FName InPathName)
@@ -3098,18 +1437,27 @@ namespace SceneOutliner
 		return Folders;
 	}
 
-	void SSceneOutliner::ScrollItemIntoView(FTreeItemPtr Item)
+	void SSceneOutliner::ScrollItemIntoView(const FTreeItemPtr& Item)
 	{
 		auto Parent = Item->GetParent();
-		while(Parent.IsValid())
+		while (Parent.IsValid())
 		{
 			OutlinerTreeView->SetItemExpansion(Parent->AsShared(), true);
 			Parent = Parent->GetParent();
 		}
-
 		OutlinerTreeView->RequestScrollIntoView(Item);
 	}
 
+	void SSceneOutliner::SetItemExpansion(const FTreeItemPtr& Item, bool bIsExpanded)
+	{
+		OutlinerTreeView->SetItemExpansion(Item, bIsExpanded);
+	}
+
+	bool SSceneOutliner::IsItemExpanded(const FTreeItemPtr& Item) const
+	{
+		return OutlinerTreeView->IsItemExpanded(Item);
+	}
+
 	TSharedRef< ITableRow > SSceneOutliner::OnGenerateRowForOutlinerTree( FTreeItemPtr Item, const TSharedRef< STableViewBase >& OwnerTable )
 	{
 		return SNew( SSceneOutlinerTreeRow, OutlinerTreeView.ToSharedRef(), SharedThis(this) ).Item( Item );
@@ -3146,23 +1494,6 @@ namespace SceneOutliner
 		}
 	}
 
-	static const FName SequencerActorTag(TEXT("SequencerActor"));
-
-	bool SSceneOutliner::IsActorDisplayable( const AActor* Actor ) const
-	{
-		return	!SharedData->bOnlyShowFolders && 												// Don't show actors if we're only showing folders
-				Actor->IsEditable() &&															// Only show actors that are allowed to be selected and drawn in editor
-				Actor->IsListedInSceneOutliner() &&
-				( (SharedData->bRepresentingPlayWorld || !Actor->HasAnyFlags(RF_Transient)) ||
-				  (SharedData->bShowTransient && Actor->HasAnyFlags(RF_Transient)) ||			// Don't show transient actors in non-play worlds
-				  (Actor->ActorHasTag(SequencerActorTag)) ) &&		
-				!Actor->IsTemplate() &&															// Should never happen, but we never want CDOs displayed
-				!FActorEditorUtils::IsABuilderBrush(Actor) &&									// Don't show the builder brush
-				!Actor->IsA( AWorldSettings::StaticClass() ) &&									// Don't show the WorldSettings actor, even though it is technically editable
-				!Actor->IsPendingKill() &&														// We don't want to show actors that are about to go away
-				FLevelUtils::IsLevelVisible( Actor->GetLevel() );								// Only show Actors whose level is visible
-	}
-
 	void SSceneOutliner::OnOutlinerTreeSelectionChanged( FTreeItemPtr TreeItem, ESelectInfo::Type SelectInfo )
 	{
 		if (SelectInfo == ESelectInfo::Direct)
@@ -3170,225 +1501,25 @@ namespace SceneOutliner
 			return;
 		}
 
-		if (SharedData->Mode == ESceneOutlinerMode::Custom)
-		{
-			OnItemSelectionChanged.Broadcast(TreeItem, SelectInfo);
-			return;
-		}
-
-		if( SharedData->Mode == ESceneOutlinerMode::ActorPicker || SharedData->Mode == ESceneOutlinerMode::ComponentPicker )
-		{
-			// In actor picking mode, we fire off the notification to whoever is listening.
-			// This may often cause the widget itself to be enqueued for destruction
-			if( OutlinerTreeView->GetNumItemsSelected() > 0 )
-			{
-				auto FirstItem = OutlinerTreeView->GetSelectedItems()[0];
-				if (FirstItem->CanInteract())
-				{
-					OnItemPicked.ExecuteIfBound( FirstItem->AsShared() );
-				}
-			}
-		}
-
-		// We only synchronize selection when in actor browsing mode
-		else if( SharedData->Mode == ESceneOutlinerMode::ActorBrowsing )
-		{
-			if( !bIsReentrant )
-			{
-				TGuardValue<bool> ReentrantGuard(bIsReentrant,true);
-
-				// @todo outliner: Can be called from non-interactive selection
-
-				// The tree let us know that selection has changed, but wasn't able to tell us
-				// what changed.  So we'll perform a full difference check and update the editor's
-				// selected actors to match the control's selection set.
-
-				// Make a list of all the actors that should now be selected in the world.
-				FItemSelection Selection(*OutlinerTreeView);
-
-				// notify components of selection change
-				if (Selection.SubComponents.Num() > 0)
-				{
-					FSceneOutlinerDelegates::Get().OnSubComponentSelectionChanged.Broadcast(Selection.SubComponents);
-				}
-
-				auto SelectedActors = TSet<AActor*>(Selection.GetActorPtrs());
-
-				bool bChanged = false;
-				bool bAnyInPIE = false;
-				for (auto* Actor : SelectedActors)
-				{
-					if (!bAnyInPIE && Actor && Actor->GetOutermost()->HasAnyPackageFlags(PKG_PlayInEditor))
-					{
-						bAnyInPIE = true;
-					}
-					if (!GEditor->GetSelectedActors()->IsSelected(Actor))
-					{
-						bChanged = true;
-						break;
-					}
-				}
-
-				for (FSelectionIterator SelectionIt( *GEditor->GetSelectedActors() ); SelectionIt && !bChanged; ++SelectionIt)
-				{
-					AActor* Actor = CastChecked< AActor >( *SelectionIt );
-					if (!bAnyInPIE && Actor->GetOutermost()->HasAnyPackageFlags(PKG_PlayInEditor))
-					{
-						bAnyInPIE = true;
-					}
-					if (!SelectedActors.Contains(Actor))
-					{
-						// Actor has been deselected
-						bChanged = true;
-
-						// If actor was a group actor, remove its members from the ActorsToSelect list
-						AGroupActor* DeselectedGroupActor = Cast<AGroupActor>(Actor);
-						if ( DeselectedGroupActor )
-						{
-							TArray<AActor*> GroupActors;
-							DeselectedGroupActor->GetGroupActors( GroupActors );
-
-							for (auto* GroupActor : GroupActors)
-							{
-								SelectedActors.Remove(GroupActor);
-							}
-
-						}
-					}
-				}
-
-				// If there's a discrepancy, update the selected actors to reflect this list.
-				if ( bChanged )
-				{
-					const FScopedTransaction Transaction( NSLOCTEXT("UnrealEd", "ClickingOnActors", "Clicking on Actors"), !bAnyInPIE );
-					GEditor->GetSelectedActors()->Modify();
-
-					// Clear the selection.
-					GEditor->SelectNone(false, true, true);
-
-					// We'll batch selection changes instead by using BeginBatchSelectOperation()
-					GEditor->GetSelectedActors()->BeginBatchSelectOperation();
-
-					const bool bShouldSelect = true;
-					const bool bNotifyAfterSelect = false;
-					const bool bSelectEvenIfHidden = true;	// @todo outliner: Is this actually OK?
-					for (auto* Actor : SelectedActors)
-					{
-						UE_LOG(LogSceneOutliner, Verbose,  TEXT("Clicking on Actor (world outliner): %s (%s)"), *Actor->GetClass()->GetName(), *Actor->GetActorLabel());
-						GEditor->SelectActor( Actor, bShouldSelect, bNotifyAfterSelect, bSelectEvenIfHidden );
-					}
-
-					// Commit selection changes
-					GEditor->GetSelectedActors()->EndBatchSelectOperation(/*bNotify*/false);
-
-					// Fire selection changed event
-					GEditor->NoteSelectionChange();
-				}
-
-				bActorSelectionDirty = true;
-			}
-		}
-	}
-
-	
-	void SSceneOutliner::OnLevelSelectionChanged(UObject* Obj)
-	{
-		// We only synchronize selection when in actor browsing mode
-		if( SharedData->Mode == ESceneOutlinerMode::ActorBrowsing )
+		if (!bIsReentrant)
 		{
-			// @todo outliner: Because we are not notified of which items are being added/removed from selection, we have
-			// no immediate means to incrementally update the tree when selection changes.
-
-			// Ideally, we can improve the filtering paradigm to better support incremental updates in cases such as these
-			if ( IsShowingOnlySelected() )
-			{
-				FullRefresh();
-			}
-			else if (!bIsReentrant)
-			{
-				OutlinerTreeView->ClearSelection();
-				bActorSelectionDirty = true;
-
-				// Scroll last item into view - this means if we are multi-selecting, we show newest selection. @TODO Not perfect though
-				if (AActor* LastSelectedActor = GEditor->GetSelectedActors()->GetBottom<AActor>())
-				{
-					FTreeItemPtr TreeItem = TreeItemMap.FindRef(LastSelectedActor);
-					if (TreeItem.IsValid())
-					{
-						if (!OutlinerTreeView->IsItemVisible(TreeItem))
-						{
-							ScrollItemIntoView(TreeItem);
-						}
-					}
-					else
-					{
-						OnItemAdded(LastSelectedActor, ENewItemAction::ScrollIntoView);
-					}
-				}
-			}
+			TGuardValue<bool> ReentrantGuard(bIsReentrant, true);
+			Mode->OnItemSelectionChanged(TreeItem, SelectInfo, FItemSelection(*OutlinerTreeView));
+
+			OnItemSelectionChanged.Broadcast(TreeItem, SelectInfo);
 		}
 	}
 
 	void SSceneOutliner::OnOutlinerTreeDoubleClick( FTreeItemPtr TreeItem )
 	{
-		if(SharedData->Mode == ESceneOutlinerMode::ActorBrowsing)
+		if (TreeItem->IsA<FFolderTreeItem>())
 		{
-			auto ExpandCollapseFolder = [&](const FFolderTreeItem& Folder){
-				auto Shared = const_cast<FFolderTreeItem&>(Folder).AsShared();
-				OutlinerTreeView->SetItemExpansion(Shared, !OutlinerTreeView->IsItemExpanded(Shared));
-			};
-
-			if (TreeItem->CanInteract())
-			{
-				TreeItem->Visit(FFunctionalVisitor()
-
-					.Actor([&](const FActorTreeItem&){
-						// Move all actors into view
-						FItemSelection Selection(*OutlinerTreeView);
-						if( Selection.Actors.Num() > 0 )
-						{
-							const bool bActiveViewportOnly = false;
-							GEditor->MoveViewportCamerasToActor( Selection.GetActorPtrs(), bActiveViewportOnly );
-						}
-					})
-
-					.Folder(ExpandCollapseFolder)
-
-					.World([](const FWorldTreeItem& WorldItem){
-						WorldItem.OpenWorldSettings();
-					})
-
-					.Component([](const FComponentTreeItem& ComponentItem) {
-						ComponentItem.OnDoubleClick();
-					})
-
-					.SubComponent([](const FSubComponentTreeItem& SubComponentItem) {
-						SubComponentItem.OnDoubleClick();
-					})
-				);
-
-			}
-			else
-			{
-				TreeItem->Visit(FFunctionalVisitor()
+			SetItemExpansion(TreeItem, !IsItemExpanded(TreeItem));
+		}
 
-					.Folder(ExpandCollapseFolder)
+		Mode->OnItemDoubleClick(TreeItem);
 
-					.Actor([&](const FActorTreeItem& Item){
-						// Move just this actor into view
-						if (AActor* Actor = Item.Actor.Get())
-						{
-							const bool bActiveViewportOnly = false;
-							GEditor->MoveViewportCamerasToActor( *Actor, bActiveViewportOnly );
-						}
-					})
-				);
-			}
-		}
-		else if (SharedData->Mode == ESceneOutlinerMode::Custom)
-		{
-			OnDoubleClickOnTreeEvent.Broadcast(TreeItem);
-		}
+		OnDoubleClickOnTreeEvent.Broadcast(TreeItem);
 	}
 
 	void SSceneOutliner::OnOutlinerTreeItemScrolledIntoView( FTreeItemPtr TreeItem, const TSharedPtr<ITableRow>& Widget )
@@ -3416,162 +1547,68 @@ namespace SceneOutliner
 		}
 	}
 
-	void SSceneOutliner::OnLevelAdded(ULevel* InLevel, UWorld* InWorld)
-	{
-		if ( SharedData->RepresentingWorld == InWorld )
-		{
-			FullRefresh();
-		}
-	}
-
-	void SSceneOutliner::OnLevelRemoved(ULevel* InLevel, UWorld* InWorld)
-	{
-		if (SharedData->RepresentingWorld == InWorld)
-		{
-			FullRefresh();
-		}
-	}
-
-	void SSceneOutliner::OnLevelActorsAdded(AActor* InActor)
+	void SSceneOutliner::OnHierarchyChangedEvent(FHierarchyChangedData Event)
 	{
-		if( !bIsReentrant )
+		if (Event.Type == FHierarchyChangedData::Added)
 		{
-			if( InActor && SharedData->RepresentingWorld == InActor->GetWorld() && IsActorDisplayable(InActor) )
+			if (Event.Item.IsValid() && !TreeItemMap.Find(Event.Item->GetID()))
 			{
-				if (!TreeItemMap.Find(InActor) && !PendingTreeItemMap.Find(InActor))
+				AddPendingItemAndChildren(Event.Item);
+				if (Event.ItemActions)
 				{
-					// Update the total actor count that match the filters
-					if (Filters->PassesAllFilters(FActorTreeItem(InActor)))
-					{
-						ApplicableActors.Emplace(InActor);
-					}
-
-					ConstructItemFor<FActorTreeItem>(InActor);
-
-					if (IsShowingActorComponents())
-					{
-						TArray<ISceneOutlinerTraversal*> ConstructTreeItemImp = IModularFeatures::Get().GetModularFeatureImplementations<ISceneOutlinerTraversal>("SceneOutlinerTraversal");
-						for (UActorComponent* Component : InActor->GetComponents())
-						{
-							if (Filters->PassesAllFilters(FComponentTreeItem(Component)))
-							{
-								bool IsHandled = false;
-								for (ISceneOutlinerTraversal* CustomImplementation : ConstructTreeItemImp)
-								{
-									IsHandled = CustomImplementation->ConstructTreeItem(*this, Component);
-									if (IsHandled)
-									{
-										break;
-									}
-								}
-								if (!IsHandled)
-								{
-									// add the actor's components - default implementation
-									ConstructItemFor<FComponentTreeItem>(Component);
-								}
-							}
-						}
-					}
+					NewItemActions.Add(Event.Item->GetID(), Event.ItemActions);
 				}
 			}
 		}
-	}
-
-	void SSceneOutliner::OnLevelActorsRemoved(AActor* InActor)
-	{
-		if( !bIsReentrant )
+		else if (Event.Type == FHierarchyChangedData::Removed)
 		{
-			if( InActor && SharedData->RepresentingWorld == InActor->GetWorld() )
+			FTreeItemPtr* Item = TreeItemMap.Find(Event.ItemID);
+			if (!Item)
 			{
-				ApplicableActors.Remove(InActor);
-				auto* ItemPtr = TreeItemMap.Find(InActor);
-				if (!ItemPtr)
-				{
-					ItemPtr = PendingTreeItemMap.Find(InActor);
-				}
-
-				if (ItemPtr)
-				{
-					PendingOperations.Emplace(FPendingTreeOperation::Removed, ItemPtr->ToSharedRef());
-					Refresh();
-				}
+				Item = PendingTreeItemMap.Find(Event.ItemID);
 			}
-		}
-	}
 
-	void SSceneOutliner::OnLevelActorsAttached(AActor* InActor, const AActor* InParent)
-	{
-		// InActor can be equal to InParent in cases of components being attached internally. The Scene Outliner does not need to do anything in this case.
-		if( !bIsReentrant && InActor != InParent )
-		{
-			if( InActor && SharedData->RepresentingWorld == InActor->GetWorld() )
+			if (Item)
 			{
-				if (auto* ItemPtr = TreeItemMap.Find(InActor))
-				{
-					PendingOperations.Emplace(FPendingTreeOperation::Moved, ItemPtr->ToSharedRef());
-					Refresh();
-				}
+				PendingOperations.Emplace(FPendingTreeOperation::Removed, Item->ToSharedRef());
+				Refresh();
 			}
 		}
-	}
-
-	void SSceneOutliner::OnLevelActorsDetached(AActor* InActor, const AActor* InParent)
-	{
-		// InActor can be equal to InParent in cases of components being attached internally. The Scene Outliner does not need to do anything in this case.
-		if( !bIsReentrant && InActor != InParent)
+		else if (Event.Type == FHierarchyChangedData::Moved)
 		{
-			if( InActor && SharedData->RepresentingWorld == InActor->GetWorld() )
+			FTreeItemPtr* Item = Event.Item ? &Event.Item : TreeItemMap.Find(Event.ItemID);
+			if (!Item)
 			{
-				if (auto* ItemPtr = TreeItemMap.Find(InActor))
-				{
-					PendingOperations.Emplace(FPendingTreeOperation::Moved, ItemPtr->ToSharedRef());
-					Refresh();
-				}
-				else
-				{
-					// We should find the item, but if we don't, do an add.
-					OnLevelActorsAdded(InActor);
-				}
+				Item = PendingTreeItemMap.Find(Event.ItemID);
 			}
-		}
-	}
-
-	/** Called by the engine when an actor's folder is changed */
-	void SSceneOutliner::OnLevelActorFolderChanged(const AActor* InActor, FName OldPath)
-	{
-		auto* ActorTreeItem = TreeItemMap.Find(InActor);
-		if (!ShouldShowFolders() || !InActor || !ActorTreeItem)
-		{
-			return;
-		}
-		
-		PendingOperations.Emplace(FPendingTreeOperation::Moved, ActorTreeItem->ToSharedRef());
-		Refresh();
-	}
 
-	void SSceneOutliner::OnLevelActorsRequestRename(const AActor* InActor)
-	{
-		auto SelectedItems = OutlinerTreeView->GetSelectedItems();
-		if( SelectedItems.Num() > 0)
-		{
-			// Ensure that the item we want to rename is visible in the tree
-			FTreeItemPtr ItemToRename = SelectedItems[SelectedItems.Num() - 1];
-			if (CanExecuteRenameRequest(ItemToRename) && ItemToRename->CanInteract())
+			if (Item)
 			{
-				PendingRenameItem = ItemToRename->AsShared();
-				ScrollItemIntoView(ItemToRename);
+				PendingOperations.Emplace(FPendingTreeOperation::Moved, Item->ToSharedRef());
+				Refresh();
 			}
 		}
-	}
+		else if (Event.Type == FHierarchyChangedData::FolderMoved)
+		{
+			FTreeItemPtr Item = TreeItemMap.FindRef(Event.ItemID);
+			if (Item.IsValid())
+			{
+				// Remove it from the map under the old ID (which is derived from the folder path)
+				TreeItemMap.Remove(Item->GetID());
 
-	void SSceneOutliner::OnMapChange(uint32 MapFlags)
-	{
-		FullRefresh();
-	}
+				// Now change the path and put it back in the map with its new ID
+				auto Folder = StaticCastSharedPtr<FFolderTreeItem>(Item);
+				Folder->Path = Event.NewPath;
+				Folder->LeafName = GetFolderLeafName(Event.NewPath);
 
-	void SSceneOutliner::OnNewCurrentLevel()
-	{
-		if (IsShowingOnlyCurrentLevel())
+				TreeItemMap.Add(Item->GetID(), Item);
+
+				// Add an operation to move the item in the hierarchy
+				PendingOperations.Emplace(FPendingTreeOperation::Moved, Item.ToSharedRef());
+				Refresh();
+			}
+		}
+		else if (Event.Type == FHierarchyChangedData::FullRefresh)
 		{
 			FullRefresh();
 		}
@@ -3579,7 +1616,7 @@ namespace SceneOutliner
 
 	void SSceneOutliner::PostUndo(bool bSuccess)
 	{
-		// Refresh our tree in case any changes have been made to the scene that might effect our actor list
+		// Refresh our tree in case any changes have been made to the scene that might effect our list
 		if( !bIsReentrant )
 		{
             bDisableIntermediateSorting = true;
@@ -3587,32 +1624,31 @@ namespace SceneOutliner
 		}
 	}
 
-	void SSceneOutliner::OnActorLabelChanged(AActor* ChangedActor)
+	void SSceneOutliner::OnItemLabelChanged(FTreeItemPtr ChangedItem)
 	{
-		if ( !ensure(ChangedActor) )
-		{
-			return;
-		}
-		
-		FTreeItemPtr TreeItem = TreeItemMap.FindRef(ChangedActor);
-		if (TreeItem.IsValid())
+		// If the item already exists
+		if (FTreeItemPtr* ExistingItem = TreeItemMap.Find(ChangedItem->GetID()))
 		{
-			if (SearchBoxFilter->PassesFilter(*TreeItem))
+			// The changed item flags will have been set already
+			if (!ChangedItem->Flags.bIsFilteredOut)
 			{
-				OutlinerTreeView->FlashHighlightOnItem(TreeItem);
+				OutlinerTreeView->FlashHighlightOnItem(*ExistingItem);
 				RequestSort();
 			}
 			else
 			{
-				// Do longer matches the filters, remove it
-				PendingOperations.Emplace(FPendingTreeOperation::Removed, TreeItem.ToSharedRef());
+				// No longer matches the filters, remove it
+				PendingOperations.Emplace(FPendingTreeOperation::Removed, ExistingItem->ToSharedRef());
 				Refresh();
 			}
 		}
-		else if (IsActorDisplayable(ChangedActor) && SharedData->RepresentingWorld == ChangedActor->GetWorld())
+		else
 		{
 			// Attempt to add the item if we didn't find it - perhaps it now matches the filter?
-			ConstructItemFor<FActorTreeItem>(ChangedActor);
+			if (ChangedItem && !ChangedItem->Flags.bIsFilteredOut)
+			{
+				AddPendingItemAndChildren(ChangedItem);
+			}
 		}
 	}
 
@@ -3630,36 +1666,13 @@ namespace SceneOutliner
 		SearchBoxFilter->SetRawFilterText( InFilterText );
 		FilterTextBoxWidget->SetError( SearchBoxFilter->GetFilterErrorText() );
 
-		// Scroll last item (if it passes the filter) into view - this means if we are multi-selecting, we show newest selection that passes the filter
-		if (AActor* LastSelectedActor = GEditor->GetSelectedActors()->GetBottom<AActor>())
-		{
-			// This part is different than that of OnLevelSelectionChanged(nullptr) because IsItemVisible(TreeItem) & ScrollItemIntoView(TreeItem) are applied to
-			// the current visual state, not to the one after applying the filter. Thus, the scroll would go to the place where the object was located
-			// before applying the FilterText
-
-			// If the object is already in the list, but it does not passes the filter, then we do not want to re-add it, because it will be removed by the filter
-			const FTreeItemPtr TreeItem = TreeItemMap.FindRef(LastSelectedActor);
-			if (TreeItem.IsValid() && !SearchBoxFilter->PassesFilter(*TreeItem))
-			{
-				return;
-			}
-
-			// If the object is not in the list, and it does not passes the filter, then we should not re-add it, because it would be removed by the filter again. Unfortunately,
-			// there is no code to check if a future element (i.e., one that is currently not in the TreeItemMap list) will pass the filter. Therefore, we kind of overkill it
-			// by re-adding that element (even though it will be removed). However, AddItemToTree(FTreeItemRef Item) and similar functions already check the element before
-			// adding it. So this solution is fine.
-			// This solution might affect the performance of the World Outliner when a key is pressed, but it will still work properly when the remove/del keys are pressed. Not
-			// updating the filter when !TreeItem.IsValid() would result in the focus not being updated when the remove/del keys are pressed.
-
-			// In any other case (i.e., if the object passes the current filter), re-add it
-			OnItemAdded(LastSelectedActor, ENewItemAction::ScrollIntoView);
-		}
+		Mode->OnFilterTextChanged(InFilterText);
 	}
 
 	void SSceneOutliner::OnFilterTextCommitted( const FText& InFilterText, ETextCommit::Type CommitInfo )
 	{
 		const FString CurrentFilterText = InFilterText.ToString();
-		// We'll only select actors if the user actually pressed the enter key.  We don't want to change
+		// We'll only select items if the user actually pressed the enter key. We don't want to change
 		// selection just because focus was lost from the search text field.
 		if( CommitInfo == ETextCommit::OnEnter )
 		{
@@ -3668,76 +1681,16 @@ namespace SceneOutliner
 			{
 				FItemSelection Selection;
 
-				// Gather all of the actors that match the filter text
+				// Gather all of the items that match the filter text
 				for (auto& Pair : TreeItemMap)
 				{
 					if (!Pair.Value->Flags.bIsFilteredOut)
 					{
-						Pair.Value->Visit(Selection);
-					}
-				}
-
-				// We only select level actors when in actor browsing mode
-				if( SharedData->Mode == ESceneOutlinerMode::ActorBrowsing )
-				{
-					// Start batching selection changes
-					GEditor->GetSelectedActors()->BeginBatchSelectOperation();
-
-					// Select actors (and only the actors) that match the filter text
-					const bool bNoteSelectionChange = false;
-					const bool bDeselectBSPSurfs = false;
-					const bool WarnAboutManyActors = true;
-					GEditor->SelectNone( bNoteSelectionChange, bDeselectBSPSurfs, WarnAboutManyActors );
-					for (auto* Actor : Selection.GetActorPtrs())
-					{
-						const bool bShouldSelect = true;
-						const bool bSelectEvenIfHidden = false;
-						GEditor->SelectActor( Actor, bShouldSelect, bNoteSelectionChange, bSelectEvenIfHidden );
+						Selection.Add(Pair.Value);
 					}
-
-					// Commit selection changes
-					GEditor->GetSelectedActors()->EndBatchSelectOperation(/*bNotify*/false);
-
-					// Fire selection changed event
-					GEditor->NoteSelectionChange();
-
-					// Set keyboard focus to the SceneOutliner, so the user can perform keyboard commands that interact
-					// with selected actors (such as Delete, to delete selected actors.)
-					SetKeyboardFocus();
 				}
 
-				// In 'actor picking' mode, we allow the user to commit their selection by pressing enter
-				// in the search window when a single actor is available
-				else if( SharedData->Mode == ESceneOutlinerMode::ActorPicker || SharedData->Mode == ESceneOutlinerMode::ComponentPicker )
-				{
-					// In actor picking mode, we check to see if we have a selected actor, and if so, fire
-					// off the notification to whoever is listening.  This may often cause the widget itself
-					// to be enqueued for destruction
-					if( Selection.Actors.Num() == 1 )
-					{
-						// Signal that an actor was selected. We assume it is valid as it won't have been added to ActorsToSelect if not.
-						OutlinerTreeView->SetSelection( Selection.Actors[0]->AsShared(), ESelectInfo::OnKeyPress );
-					}
-				}
-				// In the custom mode we want to mimic a similar result to the actor browser
-				else if (SharedData->Mode == ESceneOutlinerMode::Custom)
-				{
-					if (SelectionMode == ESelectionMode::Single || SelectionMode == ESelectionMode::SingleToggle)
-					{
-						OutlinerTreeView->SetSelection( Selection.Actors[0]->AsShared(), ESelectInfo::OnKeyPress );
-					}
-					else if (SelectionMode == ESelectionMode::Multi)
-					{
-						TArray<FTreeItemPtr> ItemsPtr;
-						ItemsPtr.Reserve(Selection.Actors.Num());
-						for (FActorTreeItem* Item : Selection.Actors)
-						{
-							ItemsPtr.Add(Item->AsShared());
-						}
-						OutlinerTreeView->ClearSelection();
-						OutlinerTreeView->SetItemSelection( ItemsPtr, true, ESelectInfo::OnKeyPress );
-					}
-				}
+				Mode->OnFilterTextCommited(Selection, CommitInfo);
 			}
 		}
 		else if (CommitInfo == ETextCommit::OnCleared)
@@ -3748,74 +1701,32 @@ namespace SceneOutliner
 
 	EVisibility SSceneOutliner::GetFilterStatusVisibility() const
 	{
-		return IsFilterActive() ? EVisibility::Visible : EVisibility::Collapsed;
+		return IsTextFilterActive() ? EVisibility::Visible : EVisibility::Collapsed;
 	}
 
 	EVisibility SSceneOutliner::GetEmptyLabelVisibility() const
 	{
-		return ( IsFilterActive() || RootTreeItems.Num() > 0 ) ? EVisibility::Collapsed : EVisibility::Visible;
+		return ( IsTextFilterActive() || RootTreeItems.Num() > 0 ) ? EVisibility::Collapsed : EVisibility::Visible;
 	}
 
 	FText SSceneOutliner::GetFilterStatusText() const
 	{
-		const int32 TotalActorCount = ApplicableActors.Num();
-
-		int32 SelectedActorCount = 0;
-		auto Count = [&](const FActorTreeItem&) { ++SelectedActorCount; };
-		for (const auto& Item : OutlinerTreeView->GetSelectedItems())
-		{
-			Item->Visit(FFunctionalVisitor().Actor(Count));
-		}
-
-		if ( !IsFilterActive() )
-		{
-			if (SelectedActorCount == 0) //-V547
-			{
-				return FText::Format( LOCTEXT("ShowingAllActorsFmt", "{0} actors"), FText::AsNumber( TotalActorCount ) );
-			}
-			else
-			{
-				return FText::Format( LOCTEXT("ShowingAllActorsSelectedFmt", "{0} actors ({1} selected)"), FText::AsNumber( TotalActorCount ), FText::AsNumber( SelectedActorCount ) );
-			}
-		}
-		else if( IsFilterActive() && FilteredActorCount == 0 )
-		{
-			return FText::Format( LOCTEXT("ShowingNoActorsFmt", "No matching actors ({0} total)"), FText::AsNumber( TotalActorCount ) );
-		}
-		else if (SelectedActorCount != 0) //-V547
-		{
-			return FText::Format( LOCTEXT("ShowingOnlySomeActorsSelectedFmt", "Showing {0} of {1} actors ({2} selected)"), FText::AsNumber( FilteredActorCount ), FText::AsNumber( TotalActorCount ), FText::AsNumber( SelectedActorCount ) );
-		}
-		else
-		{
-			return FText::Format( LOCTEXT("ShowingOnlySomeActorsFmt", "Showing {0} of {1} actors"), FText::AsNumber( FilteredActorCount ), FText::AsNumber( TotalActorCount ) );
-		}
+		return Mode->GetStatusText();
 	}
 
 	FSlateColor SSceneOutliner::GetFilterStatusTextColor() const
 	{
-		if ( !IsFilterActive() )
-		{
-			return FSlateColor::UseForeground();
-		}
-		else if( FilteredActorCount == 0 )
-		{
-			return FAppStyle::Get().GetSlateColor("Colors.AccentRed");
-		}
-		else
-		{
-			return FAppStyle::Get().GetSlateColor("Colors.AccentGreen");
-		}
+		return Mode->GetStatusTextColor();
 	}
 
-	bool SSceneOutliner::IsFilterActive() const
+	bool SSceneOutliner::IsTextFilterActive() const
 	{
-		return FilterTextBoxWidget->GetText().ToString().Len() > 0 && ApplicableActors.Num() != FilteredActorCount;
+		return FilterTextBoxWidget->GetText().ToString().Len() > 0;
 	}
 
 	const FSlateBrush* SSceneOutliner::GetFilterButtonGlyph() const
 	{
-		if( IsFilterActive() )
+		if( IsTextFilterActive() )
 		{
 			return FEditorStyle::GetBrush(TEXT("SceneOutliner.FilterCancel"));
 		}
@@ -3827,7 +1738,7 @@ namespace SceneOutliner
 
 	FString SSceneOutliner::GetFilterButtonToolTip() const
 	{
-		return IsFilterActive() ? LOCTEXT("ClearSearchFilter", "Clear search filter").ToString() : LOCTEXT("StartSearching", "Search").ToString();
+		return IsTextFilterActive() ? LOCTEXT("ClearSearchFilter", "Clear search filter").ToString() : LOCTEXT("StartSearching", "Search").ToString();
 
 	}
 
@@ -3869,142 +1780,13 @@ namespace SceneOutliner
 
 	bool SSceneOutliner::SupportsKeyboardFocus() const
 	{
-		// We only need to support keyboard focus if we're in actor browsing mode
-		if( SharedData->Mode == ESceneOutlinerMode::ActorBrowsing )
-		{
-			// Scene outliner needs keyboard focus so the user can press keys to activate commands, such as the Delete
-			// key to delete selected actors
-			return true;
-		}
-
-		return false;
+		return Mode->SupportsKeyboardFocus();
 	}
 
 	FReply SSceneOutliner::OnKeyDown( const FGeometry& MyGeometry, const FKeyEvent& InKeyEvent )
 	{
 		// @todo outliner: Use command system for these for discoverability? (allow bindings?)
-
-		// We only allow these operations in actor browsing mode
-		if( SharedData->Mode == ESceneOutlinerMode::ActorBrowsing )
-		{
-			// Rename key: Rename selected actors (not rebindable, because it doesn't make much sense to bind.)
-			if( InKeyEvent.GetKey() == EKeys::F2 )
-			{
-				if (OutlinerTreeView->GetNumItemsSelected() == 1)
-				{
-					FTreeItemPtr ItemToRename = OutlinerTreeView->GetSelectedItems()[0];
-					
-					if (CanExecuteRenameRequest(ItemToRename) && ItemToRename->CanInteract())
-					{
-						PendingRenameItem = ItemToRename->AsShared();
-						ScrollItemIntoView(ItemToRename);
-					}
-
-					return FReply::Handled();
-				}
-			}
-
-			// F5 forces a full refresh
-			else if ( InKeyEvent.GetKey() == EKeys::F5 )
-			{
-				FullRefresh();
-				return FReply::Handled();
-			}
-
-			// Delete key: Delete selected actors (not rebindable, because it doesn't make much sense to bind.)
-			// Use Delete and Backspace instead of Platform_Delete because the LevelEditor default Edit Delete is bound to both 
-			else if ( InKeyEvent.GetKey() == EKeys::Delete || InKeyEvent.GetKey() == EKeys::BackSpace )
-			{
-				const FItemSelection Selection(*OutlinerTreeView);
-
-				if( SharedData->CustomDelete.IsBound() )
-				{
-					SharedData->CustomDelete.Execute( Selection.GetWeakActors() );
-				}
-				else
-				{
-					if (CheckWorld())
-					{
-						GUnrealEd->Exec(SharedData->RepresentingWorld, TEXT("DELETE"));
-					}
-				}
-				return FReply::Handled();
-			}
-		}
-
-		return FReply::Unhandled();
-	}
-
-	void SSceneOutliner::SynchronizeActorSelection()
-	{
-		TGuardValue<bool> ReentrantGuard(bIsReentrant, true);
-
-		USelection* SelectedActors = GEditor->GetSelectedActors();
-
-		// Deselect actors in the tree that are no longer selected in the world
-		FItemSelection Selection(*OutlinerTreeView);
-		if (Selection.Actors.Num())
-		{
-			TArray<FTreeItemPtr> ActorItems;
-			for (FActorTreeItem* ActorItem : Selection.Actors)
-			{
-				if(!ActorItem->Actor.IsValid() || !ActorItem->Actor.Get()->IsSelected())
-				{
-					ActorItems.Add(ActorItem->AsShared());
-				}
-			}
-
-			OutlinerTreeView->SetItemSelection(ActorItems, false);
-		}
-		
-		// Show actor selection but only if sub objects are not selected
-		if (Selection.Components.Num() == 0 && Selection.SubComponents.Num() == 0)
-		{
-			// See if the tree view selector is pointing at a selected item
-			bool bSelectorInSelectionSet = false;
-
-			TArray<FTreeItemPtr> ActorItems;
-			for (FSelectionIterator SelectionIt(*SelectedActors); SelectionIt; ++SelectionIt)
-			{
-				AActor* Actor = CastChecked< AActor >(*SelectionIt);
-				if (FTreeItemPtr* ActorItem = TreeItemMap.Find(Actor))
-				{
-					if (!bSelectorInSelectionSet && OutlinerTreeView->Private_HasSelectorFocus(*ActorItem))
-					{
-						bSelectorInSelectionSet = true;
-					}
-
-					ActorItems.Add(*ActorItem);
-				}
-			}
-
-			// If NOT bSelectorInSelectionSet then we want to just move the selector to the first selected item.
-			ESelectInfo::Type SelectInfo = bSelectorInSelectionSet ? ESelectInfo::Direct : ESelectInfo::OnMouseClick;
-			OutlinerTreeView->SetItemSelection(ActorItems, true, SelectInfo);
-		}
-
-		// Broadcast selection changed delegate
-		FSceneOutlinerDelegates::Get().SelectionChanged.Broadcast();
-	}
-
-	void SSceneOutliner::OnComponentSelectionChanged(UActorComponent* Component)
-	{
-		if (!Component)
-			return;
-
-		if (FTreeItemPtr* ComponentItem = TreeItemMap.Find(Component))
-		{
-			if (ComponentItem->IsValid())
-			{
-				(*ComponentItem)->SynchronizeSubItemSelection(OutlinerTreeView);
-			}
-		}
-	}
-	
-	void SSceneOutliner::OnComponentsUpdated()
-	{
-		// #todo: A bit overkill, only one actors sub-components have changed
-		FullRefresh();
+		return Mode->OnKeyDown(InKeyEvent);
 	}
 
 	void SSceneOutliner::Tick(const FGeometry& AllottedGeometry, const double InCurrentTime, const float InDeltaTime)
@@ -4036,7 +1818,8 @@ namespace SceneOutliner
 		}
 		SortOutlinerTimer -= InDeltaTime;
 
-		if (bSortDirty && (!SharedData->bRepresentingPlayWorld || SortOutlinerTimer <= 0))
+		// Delay sorting when in PIE
+		if (bSortDirty && (GEditor->PlayWorld == nullptr || SortOutlinerTimer <= 0))
 		{
 			SortItems(RootTreeItems);
 			for (const auto& Pair : TreeItemMap)
@@ -4053,19 +1836,13 @@ namespace SceneOutliner
 			SortOutlinerTimer = SCENE_OUTLINER_RESORT_TIMER;
 		}
 
-
-		if (bActorSelectionDirty)
+		if (bSelectionDirty)
 		{
-			SynchronizeActorSelection();
-			bActorSelectionDirty = false;
+			Mode->SynchronizeSelection();
+			bSelectionDirty = false;
 		}
 	}
 
-	void SSceneOutliner::AddReferencedObjects(FReferenceCollector& Collector)
-	{
-		Collector.AddReferencedObject(SceneOutlinerSettings);
-	}
-
 	EColumnSortMode::Type SSceneOutliner::GetColumnSortMode( const FName ColumnId ) const
 	{
 		if (SortByColumn == ColumnId)
@@ -4108,39 +1885,9 @@ namespace SceneOutliner
 		}
 	}
 
-	void SSceneOutliner::FOnItemAddedToTree::Visit(FActorTreeItem& ActorItem) const
-	{
-		Outliner.FilteredActorCount += ActorItem.Flags.bIsFilteredOut ? 0 : 1;
-
-		// Synchronize selection
-		if (GEditor->GetSelectedActors()->IsSelected(ActorItem.Actor.Get()))
-		{
-			Outliner.OutlinerTreeView->SetItemSelection(ActorItem.AsShared(), true);
-		}
-	}
-
-	void SSceneOutliner::FOnItemAddedToTree::Visit(FFolderTreeItem& Folder) const
-	{
-		if (!Outliner.SharedData->RepresentingWorld)
-		{
-			return;
-		}
-
-		if (FActorFolderProps* Props = FActorFolders::Get().GetFolderProperties(*Outliner.SharedData->RepresentingWorld, Folder.Path))
-		{
-			Folder.Flags.bIsExpanded = Props->bIsExpanded;
-		}
-	}
-
-	void SSceneOutliner::OnSelectWorld(TWeakObjectPtr<UWorld> InWorld)
-	{
-		SharedData->UserChosenWorld = InWorld;
-		FullRefresh();
-	}
-
-	bool SSceneOutliner::IsWorldChecked(TWeakObjectPtr<UWorld> InWorld)
+	uint32 SSceneOutliner::GetTypeSortPriority(const ITreeItem& Item) const 
 	{
-		return (InWorld == SharedData->UserChosenWorld);
+		return Mode->GetTypeSortPriority(Item);
 	}
 
 	void SSceneOutliner::SetItemExpansionRecursive(FTreeItemPtr Model, bool bInExpansionState)
@@ -4157,6 +1904,35 @@ namespace SceneOutliner
 			}
 		}
 	}
+
+	TSharedPtr<FDragDropOperation> SSceneOutliner::CreateDragDropOperation(const TArray<FTreeItemPtr>& InTreeItems) const 
+	{ 
+		return Mode->CreateDragDropOperation(InTreeItems); 
+	}
+
+	/** Parse a drag drop operation into a payload */
+	bool SSceneOutliner::ParseDragDrop(FDragDropPayload& OutPayload, const FDragDropOperation& Operation) const 
+	{
+		return Mode->ParseDragDrop(OutPayload, Operation);
+	}
+
+	/** Validate a drag drop operation on a drop target */
+	FDragValidationInfo SSceneOutliner::ValidateDrop(const ITreeItem& DropTarget, const FDragDropPayload& Payload) const 
+	{ 
+		return Mode->ValidateDrop(DropTarget, Payload); 
+	}
+
+	/** Called when a payload is dropped onto a target */
+	void SSceneOutliner::OnDropPayload(ITreeItem& DropTarget, const FDragDropPayload& Payload, const FDragValidationInfo& ValidationInfo) const 
+	{ 
+		return Mode->OnDrop(DropTarget, Payload, ValidationInfo); 
+	}
+
+	/** Called when a payload is dragged over an item */
+	FReply SSceneOutliner::OnDragOverItem(const FDragDropEvent& Event, const ITreeItem& Item) const 
+	{ 
+		return Mode->OnDragOverItem(Event, Item); 
+	}
 }	// namespace SceneOutliner
 
 #undef LOCTEXT_NAMESPACE
