SceneOutliner refactor

diff --git a/Engine/Source/Editor/SceneOutliner/Private/SceneOutlinerActorInfoColumn.cpp b/Engine/Source/Editor/SceneOutliner/Private/SceneOutlinerActorInfoColumn.cpp
--- a/Engine/Source/Editor/SceneOutliner/Private/SceneOutlinerActorInfoColumn.cpp
+++ b/Engine/Source/Editor/SceneOutliner/Private/SceneOutlinerActorInfoColumn.cpp
@@ -8,49 +8,53 @@
 #include "Widgets/Text/STextBlock.h"
 #include "Widgets/Views/SListView.h"
 #include "EditorStyleSet.h"
-#include "ISceneOutliner.h"
-
-
 #include "EditorClassUtils.h"
 #include "SortHelper.h"
+#include "ISceneOutliner.h"
+#include "ActorTreeItem.h"
+#include "ComponentTreeItem.h"
+#include "FolderTreeItem.h"
+#include "WorldTreeItem.h"
 
 #define LOCTEXT_NAMESPACE "SceneOutlinerActorInfoColumn"
 
 namespace SceneOutliner
 {
 
-struct FGetInfo : TTreeItemGetter<FString>
+struct FGetInfo
 {
 	ECustomColumnMode::Type CurrentMode;
 
 	FGetInfo(ECustomColumnMode::Type InCurrentMode) : CurrentMode(InCurrentMode) {}
 
-	virtual FString Get(const FActorTreeItem& ActorItem) const override
+	FString operator()(const ITreeItem& Item) const
 	{
-		AActor* Actor = ActorItem.Actor.Get();
-		if (!Actor)
+		if (const FActorTreeItem* ActorItem = Item.CastTo<FActorTreeItem>())
 		{
-			return FString();
-		}
+			AActor* Actor = ActorItem->Actor.Get();
+			if (!Actor)
+			{
+				return FString();
+			}
 
-		switch(CurrentMode)
-		{
-		case ECustomColumnMode::Class:
-			return Actor->GetClass()->GetName();
+			switch (CurrentMode)
+			{
+			case ECustomColumnMode::Class:
+				return Actor->GetClass()->GetName();
 
-		case ECustomColumnMode::Level:
-			return FPackageName::GetShortName(Actor->GetOutermost()->GetName());
+			case ECustomColumnMode::Level:
+				return FPackageName::GetShortName(Actor->GetOutermost()->GetName());
 
-		case ECustomColumnMode::Socket:
-			return Actor->GetAttachParentSocketName().ToString();
+			case ECustomColumnMode::Socket:
+				return Actor->GetAttachParentSocketName().ToString();
 
-		case ECustomColumnMode::InternalName:
-			return Actor->GetFName().ToString();
+			case ECustomColumnMode::InternalName:
+				return Actor->GetFName().ToString();
 
-		case ECustomColumnMode::UncachedLights:
-			return FString::Printf(TEXT("%7d"), Actor->GetNumUncachedStaticLightingInteractions());
+			case ECustomColumnMode::UncachedLights:
+				return FString::Printf(TEXT("%7d"), Actor->GetNumUncachedStaticLightingInteractions());
 
-		case ECustomColumnMode::Layer:
+			case ECustomColumnMode::Layer:
 			{
 				FString Result;
 				for (const auto& Layer : Actor->Layers)
@@ -86,69 +90,55 @@ struct FGetInfo : TTreeItemGetter<FString>
 				return Result;
 			}
 
-		default:
-			return FString();
+			default:
+				return FString();
+			}
 		}
-	}
-
-	virtual FString Get(const FFolderTreeItem&) const override
-	{
-		switch(CurrentMode)
+		else if (Item.IsA<FFolderTreeItem>())
 		{
-		case ECustomColumnMode::Class:
-			return LOCTEXT("FolderTypeName", "Folder").ToString();
+			switch (CurrentMode)
+			{
+			case ECustomColumnMode::Class:
+				return LOCTEXT("FolderTypeName", "Folder").ToString();
 
-		default:
-			return FString();
+			default:
+				return FString();
+			}
 		}
-	}
-
-	virtual FString Get(const FWorldTreeItem&) const override
-	{
-		switch(CurrentMode)
+		else if (Item.IsA<FWorldTreeItem>())
 		{
-		case ECustomColumnMode::Class:
-			return LOCTEXT("WorldTypeName", "World").ToString();
+			switch (CurrentMode)
+			{
+			case ECustomColumnMode::Class:
+				return LOCTEXT("WorldTypeName", "World").ToString();
 
-		default:
-			return FString();
+			default:
+				return FString();
+			}
 		}
-	}
-
-	virtual FString Get(const FComponentTreeItem& ComponentItem) const override
-	{
-		UActorComponent* Component = ComponentItem.Component.Get();
-		if (!Component)
+		else if (const FComponentTreeItem* ComponentItem = Item.CastTo<FComponentTreeItem>())
 		{
-			return FString();
-		}
+			UActorComponent* Component = ComponentItem->Component.Get();
+			if (!Component)
+			{
+				return FString();
+			}
 
-		switch (CurrentMode)
-		{
-		case ECustomColumnMode::Class:
-			return LOCTEXT("ComponentTypeName", "Component").ToString();
+			switch (CurrentMode)
+			{
+			case ECustomColumnMode::Class:
+				return LOCTEXT("ComponentTypeName", "Component").ToString();
 
-		case ECustomColumnMode::InternalName:
-			return Component->GetFName().ToString();
+			case ECustomColumnMode::InternalName:
+				return Component->GetFName().ToString();
 
-		default:
-			return FString();
+			default:
+				return FString();
+			}
 		}
-	}
 
-	virtual FString Get(const FSubComponentTreeItem& Item) const override
-	{
-		switch (CurrentMode)
-		{
-		case ECustomColumnMode::Class:
-		case ECustomColumnMode::InternalName:
-			return Item.GetTypeName();
-
-		default:
-			return FString();
-		}
+		return FString();
 	}
-
 };
 
 
@@ -260,49 +250,44 @@ const TSharedRef< SWidget > FActorInfoColumn::ConstructRowWidget( FTreeItemRef T
 
 TSharedPtr<SWidget> FActorInfoColumn::ConstructClassHyperlink( ITreeItem& TreeItem )
 {
-	struct FConstructHyperlink : TTreeItemGetter<TSharedPtr<SWidget>>
+	if (const FActorTreeItem* ActorItem = TreeItem.CastTo<FActorTreeItem>())
 	{
-		TSharedPtr<SWidget> Get(const FActorTreeItem& ActorItem) const override
+		if (AActor* Actor = ActorItem->Actor.Get())
 		{
-			if (AActor* Actor = ActorItem.Actor.Get())
+			if (UClass* ActorClass = Actor->GetClass())
 			{
-				if (UClass* ActorClass = Actor->GetClass())
+				// Always show blueprints
+				const bool bIsBlueprintClass = UBlueprint::GetBlueprintFromClass(ActorClass) != nullptr;
+
+				// Also show game or game plugin native classes (but not engine classes as that makes the scene outliner pretty noisy)
+				bool bIsGameClass = false;
+				if (!bIsBlueprintClass)
 				{
-					// Always show blueprints
-					const bool bIsBlueprintClass = UBlueprint::GetBlueprintFromClass(ActorClass) != nullptr;
+					UPackage* Package = ActorClass->GetOutermost();
+					const FString ModuleName = FPackageName::GetShortName(Package->GetFName());
 
-					// Also show game or game plugin native classes (but not engine classes as that makes the scene outliner pretty noisy)
-					bool bIsGameClass = false;
-					if (!bIsBlueprintClass)
+					FModuleStatus PackageModuleStatus;
+					if (FModuleManager::Get().QueryModule(*ModuleName, /*out*/ PackageModuleStatus))
 					{
-						UPackage* Package = ActorClass->GetOutermost();
-						const FString ModuleName = FPackageName::GetShortName(Package->GetFName());
-
-						FModuleStatus PackageModuleStatus;
-						if (FModuleManager::Get().QueryModule(*ModuleName, /*out*/ PackageModuleStatus))
-						{
-							bIsGameClass = PackageModuleStatus.bIsGameModule;
-						}
+						bIsGameClass = PackageModuleStatus.bIsGameModule;
 					}
+				}
 
-					if (bIsBlueprintClass || bIsGameClass)
-					{
-						return FEditorClassUtils::GetSourceLink(ActorClass, Actor);
-					}
+				if (bIsBlueprintClass || bIsGameClass)
+				{
+					return FEditorClassUtils::GetSourceLink(ActorClass, Actor);
 				}
 			}
-
-			return nullptr;
 		}
-	};
+	}
 
-	return TreeItem.Get(FConstructHyperlink());
+	return nullptr;
 }
 
 void FActorInfoColumn::PopulateSearchStrings( const ITreeItem& Item, TArray< FString >& OutSearchStrings ) const
 {
 	{
-		FString String = Item.Get(FGetInfo(CurrentMode));
+		FString String = FGetInfo(CurrentMode)(Item);
 		if (String.Len())
 		{
 			OutSearchStrings.Add(String);
@@ -312,7 +297,7 @@ void FActorInfoColumn::PopulateSearchStrings( const ITreeItem& Item, TArray< FSt
 	// We always add the class
 	if (CurrentMode != ECustomColumnMode::Class)
 	{
-		FString String = Item.Get(FGetInfo(ECustomColumnMode::Class));
+		FString String = FGetInfo(ECustomColumnMode::Class)(Item);
 		if (String.Len())
 		{
 			OutSearchStrings.Add(String);
@@ -356,7 +341,7 @@ EVisibility FActorInfoColumn::GetColumnDataVisibility( bool bIsClassHyperlink )
 FText FActorInfoColumn::GetTextForItem( TWeakPtr<ITreeItem> TreeItem ) const
 {
 	auto Item = TreeItem.Pin();
-	return Item.IsValid() ? FText::FromString(Item->Get(FGetInfo(CurrentMode))) : FText::GetEmpty();
+	return Item.IsValid() ? FText::FromString(FGetInfo(CurrentMode)(*Item)) : FText::GetEmpty();
 }
 
 FText FActorInfoColumn::GetSelectedMode() const
diff --git a/Engine/Source/Editor/SceneOutliner/Private/SceneOutlinerDragDrop.cpp b/Engine/Source/Editor/SceneOutliner/Private/SceneOutlinerDragDrop.cpp
--- a/Engine/Source/Editor/SceneOutliner/Private/SceneOutlinerDragDrop.cpp
+++ b/Engine/Source/Editor/SceneOutliner/Private/SceneOutlinerDragDrop.cpp
@@ -3,35 +3,32 @@
 #include "SceneOutlinerDragDrop.h"
 #include "ITreeItem.h"
 #include "SceneOutlinerStandaloneTypes.h"
-#include "SubComponentTreeItem.h"
 
 #define LOCTEXT_NAMESPACE "SSceneOutliner"
 
 namespace SceneOutliner
 {
-
-	FSceneOutlinerDragDropOp::FSceneOutlinerDragDropOp(const FDragDropPayload& DraggedObjects)
-		: OverrideText()
-		, OverrideIcon(nullptr)
+	void FFolderDragDropOp::Init(FFolderPaths InFolders)
 	{
-		if (DraggedObjects.Actors)
-		{
-			ActorOp = MakeShareable(new FActorDragDropOp);
-			ActorOp->Init(DraggedObjects.Actors.GetValue());
-		}
+		Folders = MoveTemp(InFolders);
 
-		if (DraggedObjects.Folders)
+		CurrentIconBrush = FEditorStyle::Get().GetBrush(TEXT("SceneOutliner.FolderClosed"));
+		if (Folders.Num() == 1)
 		{
-			FolderOp = MakeShareable(new FFolderDragDropOp);
-			FolderOp->Init(DraggedObjects.Folders.GetValue());
+			CurrentHoverText = FText::FromName(GetFolderLeafName(Folders[0]));
 		}
-
-		if (DraggedObjects.SubComponents)
+		else
 		{
-			SubComponentOp = MakeShareable(new FSubComponentDragDropOp);
-			SubComponentOp->Init(DraggedObjects.SubComponents.GetValue());
+			CurrentHoverText = FText::Format(NSLOCTEXT("FFolderDragDropOp", "FormatFolders", "{0} Folders"), FText::AsNumber(Folders.Num()));
 		}
 	}
+
+	FSceneOutlinerDragDropOp::FSceneOutlinerDragDropOp()
+		: OverrideText()
+		, OverrideIcon(nullptr)
+	{
+		
+	}
 	
 	EVisibility FSceneOutlinerDragDropOp::GetOverrideVisibility() const
 	{
@@ -74,144 +71,18 @@ namespace SceneOutliner
 			]
 		];
 
-		if (FolderOp.IsValid())
+		for (auto& SubOp : SubOps)
 		{
-			auto Content = FolderOp->GetDefaultDecorator();
+			auto Content = SubOp->GetDefaultDecorator();
 			if (Content.IsValid())
 			{
 				Content->SetVisibility(TAttribute<EVisibility>(this, &FSceneOutlinerDragDropOp::GetDefaultVisibility));
-				VerticalBox->AddSlot()[ Content.ToSharedRef() ];
-			}
-		}
-
-		if (ActorOp.IsValid())
-		{
-			auto Content = ActorOp->GetDefaultDecorator();
-			if (Content.IsValid())
-			{
-				Content->SetVisibility(TAttribute<EVisibility>(this, &FSceneOutlinerDragDropOp::GetDefaultVisibility));
-				VerticalBox->AddSlot()[ Content.ToSharedRef() ];
+				VerticalBox->AddSlot()[Content.ToSharedRef()];
 			}
 		}
 
 		return VerticalBox;
 	}
-
-	void FFolderDragDropOp::Init(FFolderPaths InFolders)
-	{
-		Folders = MoveTemp(InFolders);
-
-		CurrentIconBrush = FEditorStyle::Get().GetBrush(TEXT("SceneOutliner.FolderClosed"));
-		if (Folders.Num() == 1)
-		{
-			CurrentHoverText = FText::FromName(GetFolderLeafName(Folders[0]));
-		}
-		else
-		{
-			CurrentHoverText = FText::Format(NSLOCTEXT("FFolderDragDropOp", "FormatFolders", "{0} Folders"), FText::AsNumber(Folders.Num()));
-		}
-	}
-
-	TSharedPtr<FDragDropOperation> CreateDragDropOperation(const TArray<FTreeItemPtr>& InTreeItems)
-	{
-		FDragDropPayload DraggedObjects;
-		for (const auto& Item : InTreeItems)
-		{
-			Item->PopulateDragDropPayload(DraggedObjects);
-		}
-
-		if (DraggedObjects.Folders)
-		{
-			TSharedPtr<FSceneOutlinerDragDropOp> OutlinerOp = MakeShareable(new FSceneOutlinerDragDropOp(DraggedObjects));
-			OutlinerOp->Construct();
-			return OutlinerOp;
-		}
-		else if (DraggedObjects.Actors)
-		{
-			return FActorDragDropGraphEdOp::New(DraggedObjects.Actors.GetValue());
-		}
-		else if (DraggedObjects.SubComponents)
-		{
-			TSharedPtr<FSceneOutlinerDragDropOp> OutlinerOp = MakeShareable(new FSceneOutlinerDragDropOp(DraggedObjects));
-			OutlinerOp->Construct();
-			return OutlinerOp;
-		}
-		return nullptr;
-	}
-
-	FDragDropPayload::FDragDropPayload()
-	{
-
-	}
-
-	bool FDragDropPayload::ParseDrag(const FDragDropOperation& Operation)
-	{
-		bool bApplicable = false;
-
-		if (Operation.IsOfType<FSceneOutlinerDragDropOp>())
-		{
-			const auto& OutlinerOp = static_cast<const FSceneOutlinerDragDropOp&>(Operation);
-			if (OutlinerOp.FolderOp.IsValid())
-			{
-				Folders = OutlinerOp.FolderOp->Folders;
-			}
-			if (OutlinerOp.ActorOp.IsValid())
-			{
-				Actors = OutlinerOp.ActorOp->Actors;
-			}
-			if (OutlinerOp.SubComponentOp.IsValid())
-			{
-				SubComponents = OutlinerOp.SubComponentOp->Items;
-			}
-			bApplicable = true;
-		}
-		else if (Operation.IsOfType<FActorDragDropOp>())
-		{
-			Actors = static_cast<const FActorDragDropOp&>(Operation).Actors;
-			bApplicable = true;
-		}
-		else if (Operation.IsOfType<FFolderDragDropOp>())
-		{
-			Folders = static_cast<const FFolderDragDropOp&>(Operation).Folders;
-			bApplicable = true;
-		}
-		else if (Operation.IsOfType<FSubComponentDragDropOp>())
-		{
-			SubComponents = static_cast<const FSubComponentDragDropOp&>(Operation).Items;
-			bApplicable = true;
-		}
-		return bApplicable;
-	}
-
-	void FSubComponentDragDropOp::Init(const FSubComponentItemArray& InItems)
-	{
-		for (int32 i = 0; i < InItems.Num(); i++)
-		{
-			if (InItems[i].IsValid())
-			{
-				Items.Add(InItems[i]);
-			}
-		}
-
-		// Set text and icon
-		UClass* CommonSelClass = NULL;
-		//CurrentIconBrush = FClassIconFinder::FindIconForActors(Items, CommonSelClass);
-		if (Items.Num() == 0)
-		{
-			CurrentHoverText = NSLOCTEXT("FSubComponentItemDragDropOp", "None", "None");
-		}
-		else if (Items.Num() == 1)
-		{
-			// Find icon for actor
-			auto TheItem = Items[0].Pin();
-			CurrentHoverText = FText::FromString(TheItem->GetDisplayString());
-		}
-		else
-		{
-			CurrentHoverText = FText::Format(NSLOCTEXT("FSubComponentItemDragDropOp", "FormatItems", "{0} Items"), FText::AsNumber(Items.Num()));
-		}
-	}
-
 }
 
 #undef LOCTEXT_NAMESPACE
diff --git a/Engine/Source/Editor/SceneOutliner/Private/SceneOutlinerFilters.cpp b/Engine/Source/Editor/SceneOutliner/Private/SceneOutlinerFilters.cpp
--- a/Engine/Source/Editor/SceneOutliner/Private/SceneOutlinerFilters.cpp
+++ b/Engine/Source/Editor/SceneOutliner/Private/SceneOutlinerFilters.cpp
@@ -51,6 +51,8 @@ void FOutlinerFilterInfo::ToggleFilterActive()
 	bActive = !bActive;
 
 	ApplyFilter(bActive);
+
+	OnToggleEvent.Broadcast(bActive);
 }
 
 bool FOutlinerFilterInfo::IsFilterActive() const
diff --git a/Engine/Source/Editor/SceneOutliner/Private/SceneOutlinerGutter.cpp b/Engine/Source/Editor/SceneOutliner/Private/SceneOutlinerGutter.cpp
--- a/Engine/Source/Editor/SceneOutliner/Private/SceneOutlinerGutter.cpp
+++ b/Engine/Source/Editor/SceneOutliner/Private/SceneOutlinerGutter.cpp
@@ -10,7 +10,7 @@
 #include "ScopedTransaction.h"
 #include "Widgets/Views/STreeView.h"
 #include "ISceneOutliner.h"
-
+#include "ISceneOutlinerMode.h"
 #include "SortHelper.h"
 
 #define LOCTEXT_NAMESPACE "SceneOutlinerGutter"
@@ -18,7 +18,21 @@
 namespace SceneOutliner
 {
 
-bool FGetVisibilityVisitor::RecurseChildren(const ITreeItem& Item) const
+void OnSetItemVisibility(ITreeItem& Item, const bool bNewVisibility)
+{
+	// Apply the same visibility to the children
+	Item.OnVisibilityChanged(bNewVisibility);
+	for (auto& ChildPtr : Item.GetChildren())
+	{
+		auto Child = ChildPtr.Pin();
+		if (Child.IsValid())
+		{
+			OnSetItemVisibility(*Child, bNewVisibility);
+		}
+	}
+}
+
+bool FGetVisibilityCache::RecurseChildren(const ITreeItem& Item) const
 {
 	if (const bool* Info = VisibilityInfo.Find(&Item))
 	{
@@ -30,7 +44,7 @@ bool FGetVisibilityVisitor::RecurseChildren(const ITreeItem& Item) const
 		for (const auto& ChildPtr : Item.GetChildren())
 		{
 			auto Child = ChildPtr.Pin();
-			if (Child.IsValid() && Child->Get(*this))
+			if (Child.IsValid() && GetVisibility(*Child))
 			{
 				bIsVisible = true;
 				break;
@@ -42,90 +56,28 @@ bool FGetVisibilityVisitor::RecurseChildren(const ITreeItem& Item) const
 	}
 }
 
-bool FGetVisibilityVisitor::Get(const FActorTreeItem& ActorItem) const
+bool FGetVisibilityCache::GetVisibility(const ITreeItem& Item) const
 {
-	if (const bool* Info = VisibilityInfo.Find(&ActorItem))
-	{
-		return *Info;
-	}
-	else
+	if (Item.HasVisibilityInfo())
 	{
-		const AActor* Actor = ActorItem.Actor.Get();
-
-		const bool bIsVisible = Actor && !Actor->IsTemporarilyHiddenInEditor(true);
-		VisibilityInfo.Add(&ActorItem, bIsVisible);
-
-		return bIsVisible;
-	}
-}
-
-bool FGetVisibilityVisitor::Get(const FWorldTreeItem& WorldItem) const
-{
-	return RecurseChildren(WorldItem);
-}
-
-bool FGetVisibilityVisitor::Get(const FFolderTreeItem& FolderItem) const
-{
-	return RecurseChildren(FolderItem);
-}
-
-struct FSetVisibilityVisitor : IMutableTreeItemVisitor
-{
-	/** Whether this item should be visible or not */
-	const bool bSetVisibility;
-
-	FSetVisibilityVisitor(bool bInSetVisibility) : bSetVisibility(bInSetVisibility) {}
-
-	virtual void Visit(FActorTreeItem& ActorItem) const override
-	{
-		AActor* Actor = ActorItem.Actor.Get();
-		if (Actor)
+		if (const bool* Info = VisibilityInfo.Find(&Item))
 		{
-			// Save the actor to the transaction buffer to support undo/redo, but do
-			// not call Modify, as we do not want to dirty the actor's package and
-			// we're only editing temporary, transient values
-			SaveToTransactionBuffer(Actor, false);
-			Actor->SetIsTemporarilyHiddenInEditor( !bSetVisibility );
-
-			// Apply the same visibility to the actors children
-			for (auto& ChildPtr : ActorItem.GetChildren())
-			{
-				auto Child = ChildPtr.Pin();
-				if (Child.IsValid())
-				{
-					FSetVisibilityVisitor Visibility( bSetVisibility );
-					Child->Visit( Visibility );
-				}
-			}
+			return *Info;
 		}
-	}
-
-	virtual void Visit(FWorldTreeItem& WorldItem) const override
-	{
-		for (auto& ChildPtr : WorldItem.GetChildren())
+		else
 		{
-			auto Child = ChildPtr.Pin();
-			if (Child.IsValid())
-			{
-				FSetVisibilityVisitor Visibility(bSetVisibility);
-				Child->Visit(Visibility);
-			}
+			const bool bIsVisible = Item.GetVisibility();
+			VisibilityInfo.Add(&Item, bIsVisible);
+			return bIsVisible;
 		}
 	}
-
-	virtual void Visit(FFolderTreeItem& FolderItem) const override
+	else
 	{
-		for (auto& ChildPtr : FolderItem.GetChildren())
-		{
-			auto Child = ChildPtr.Pin();
-			if (Child.IsValid())
-			{
-				FSetVisibilityVisitor Visibility(bSetVisibility);
-				Child->Visit(Visibility);
-			}
-		}
+		return RecurseChildren(Item);
 	}
-};
+
+	return false;
+}
 
 class FVisibilityDragDropOp : public FDragDropOperation, public TSharedFromThis<FVisibilityDragDropOp>
 {
@@ -133,7 +85,7 @@ public:
 	
 	DRAG_DROP_OPERATOR_TYPE(FVisibilityDragDropOp, FDragDropOperation)
 
-	/** Flag which defines whether to hide destination actors or not */
+	/** Flag which defines whether to hide destination items or not */
 	bool bHidden;
 
 	/** Undo transaction stolen from the gutter which is kept alive for the duration of the drag */
@@ -158,7 +110,7 @@ public:
 	}
 };
 
-/** Widget responsible for managing the visibility for a single actor */
+/** Widget responsible for managing the visibility for a single item */
 class SVisibilityWidget : public SImage
 {
 public:
@@ -209,7 +161,7 @@ private:
 		}
 	}
 
-	/** If a visibility drag drop operation has entered this widget, set its actor to the new visibility state */
+	/** If a visibility drag drop operation has entered this widget, set its item to the new visibility state */
 	virtual void OnDragEnter(const FGeometry& MyGeometry, const FDragDropEvent& DragDropEvent) override
 	{
 		auto VisibilityOp = DragDropEvent.GetOperationAs<FVisibilityDragDropOp>();
@@ -231,7 +183,7 @@ private:
 		}
 
 		// Open an undo transaction
-		UndoTransaction.Reset(new FScopedTransaction(LOCTEXT("SetActorVisibility", "Set Actor Visibility")));
+		UndoTransaction.Reset(new FScopedTransaction(LOCTEXT("SetOutlinerItemVisibility", "Set Item Visibility")));
 
 		const auto& Tree = Outliner->GetTree();
 
@@ -240,13 +192,11 @@ private:
 		// We operate on all the selected items if the specified item is selected
 		if (Tree.IsItemSelected(TreeItem.ToSharedRef()))
 		{
-			const FSetVisibilityVisitor Visitor(bVisible);
-
 			for (auto& SelectedItem : Tree.GetSelectedItems())
 			{
 				if (IsVisible(SelectedItem, Column) != bVisible)
 				{
-					SelectedItem->Visit(Visitor);
+					OnSetItemVisibility(*SelectedItem, bVisible);
 				}		
 			}
 
@@ -340,7 +290,7 @@ private:
 		return IsVisible(WeakTreeItem.Pin(), WeakColumn.Pin());
 	}
 
-	/** Set the actor this widget is responsible for to be hidden or shown */
+	/** Set the item this widget is responsible for to be hidden or shown */
 	void SetIsVisible(const bool bVisible)
 	{
 		TSharedPtr<ITreeItem> TreeItem = WeakTreeItem.Pin();
@@ -348,8 +298,7 @@ private:
 
 		if (TreeItem.IsValid() && Outliner.IsValid() && IsVisible() != bVisible)
 		{
-			FSetVisibilityVisitor Visitor(bVisible);
-			TreeItem->Visit(Visitor);
+			OnSetItemVisibility(*TreeItem, bVisible);
 			
 			Outliner->Refresh();
 
@@ -372,7 +321,7 @@ private:
 	/** Scoped undo transaction */
 	TUniquePtr<FScopedTransaction> UndoTransaction;
 
-	/** Visbility brushes for the various states */
+	/** Visibility brushes for the various states */
 	const FSlateBrush* VisibleHoveredBrush;
 	const FSlateBrush* VisibleNotHoveredBrush;
 	const FSlateBrush* NotVisibleHoveredBrush;
@@ -424,9 +373,9 @@ void FSceneOutlinerGutter::SortItems(TArray<FTreeItemPtr>& RootItems, const ECol
 {
 	FSortHelper<int32, bool>()
 		/** Sort by type first */
-		.Primary([](const ITreeItem& Item){ return Item.GetTypeSortPriority(); }, SortMode)
+		.Primary([this](const ITreeItem& Item){ return WeakOutliner.Pin()->GetMode()->GetTypeSortPriority(Item); }, SortMode)
 		/** Then by visibility */
-		.Secondary(FGetVisibilityVisitor(), SortMode)
+		.Secondary([](const ITreeItem& Item) {return FGetVisibilityCache().GetVisibility(Item); },					SortMode)
 		.Sort(RootItems);
 }
 
diff --git a/Engine/Source/Editor/SceneOutliner/Private/SceneOutlinerGutter.h b/Engine/Source/Editor/SceneOutliner/Private/SceneOutlinerGutter.h
--- a/Engine/Source/Editor/SceneOutliner/Private/SceneOutlinerGutter.h
+++ b/Engine/Source/Editor/SceneOutliner/Private/SceneOutlinerGutter.h
@@ -7,6 +7,7 @@
 #include "Widgets/Views/SHeaderRow.h"
 #include "SceneOutlinerPublicTypes.h"
 #include "ISceneOutlinerColumn.h"
+#include "ITreeItem.h"
 
 class ISceneOutliner;
 template<typename ItemType> class STableRow;
@@ -14,31 +15,11 @@ template<typename ItemType> class STableRow;
 namespace SceneOutliner
 {
 
-/** A 'getter' visitor that gets, and caches the visibility of a tree item */
-struct FGetVisibilityVisitor : TTreeItemGetter<bool>
-{
-	/** Map of tree item to visibility */
-	mutable TMap<const ITreeItem*, bool> VisibilityInfo;
-
-	/** Get an item's visibility based on its children */
-	bool RecurseChildren(const ITreeItem& Item) const;
-
-	/** Get an actor's visibility */
-	virtual bool Get(const FActorTreeItem& ActorItem) const override;
-
-	/** Get a World's visibility */
-	virtual bool Get(const FWorldTreeItem& WorldItem) const override;
-
-	/** Get a folder's visibility */
-	virtual bool Get(const FFolderTreeItem& FolderItem) const override;
-};
-
 /**
- * A gutter for the SceneOutliner which is capable of displaying a variety of Actor details
+ * A gutter for the SceneOutliner which handles setting and visualizing item visibility
  */
 class FSceneOutlinerGutter : public ISceneOutlinerColumn
 {
-
 public:
 
 	/**	Constructor */
@@ -66,7 +47,7 @@ public:
 	/** Check whether the specified item is visible */
 	FORCEINLINE bool IsItemVisible(const ITreeItem& Item)
 	{
-		return Item.Get(VisibilityCache);
+		return VisibilityCache.GetVisibility(Item);
 	}
 
 private:
@@ -74,8 +55,8 @@ private:
 	/** Weak pointer back to the scene outliner - required for setting visibility on current selection. */
 	TWeakPtr<ISceneOutliner> WeakOutliner;
 
-	/** Visitor used to get (and cache) visibilty for items. Cahced per-frame to avoid expensive recursion. */
-	FGetVisibilityVisitor VisibilityCache;
+	/** Get and cache visibility for items. Cached per-frame to avoid expensive recursion. */
+	FGetVisibilityCache VisibilityCache;
 };
 
 }	// namespace SceneOutliner
diff --git a/Engine/Source/Editor/SceneOutliner/Private/SceneOutlinerItemLabelColumn.cpp b/Engine/Source/Editor/SceneOutliner/Private/SceneOutlinerItemLabelColumn.cpp
--- a/Engine/Source/Editor/SceneOutliner/Private/SceneOutlinerItemLabelColumn.cpp
+++ b/Engine/Source/Editor/SceneOutliner/Private/SceneOutlinerItemLabelColumn.cpp
@@ -12,909 +12,18 @@
 #include "EditorStyleSet.h"
 #include "Editor.h"
 #include "ScopedTransaction.h"
-#include "Styling/SlateIconFinder.h"
-#include "ClassIconFinder.h"
 #include "SceneOutlinerStandaloneTypes.h"
+#include "ISceneOutlinerMode.h"
 #include "SceneOutlinerDragDrop.h"
 #include "Widgets/Views/SListView.h"
 #include "SortHelper.h"
-#include "EditorActorFolders.h"
-#include "Widgets/Text/SInlineEditableTextBlock.h"
 #include "Widgets/SToolTip.h"
-#include "IDocumentation.h"
-#include "ActorEditorUtils.h"
+
 
 #define LOCTEXT_NAMESPACE "SceneOutlinerItemLabelColumn"
 
 namespace SceneOutliner
 {
-
-struct FCommonLabelData
-{
-	TWeakPtr<ISceneOutliner> WeakSceneOutliner;
-	static const FLinearColor DarkColor;
-
-	TOptional<FLinearColor> GetForegroundColor(const FTreeItemPtr& TreeItem) const
-	{
-		if (!TreeItem.IsValid())
-		{
-			return DarkColor;
-		}
-
-		UWorld* World = TreeItem->GetSharedData().RepresentingWorld;
-
-		// Darken items that aren't suitable targets for an active drag and drop action
-		if (FSlateApplication::Get().IsDragDropping() && World)
-		{
-			TSharedPtr<FDragDropOperation> DragDropOp = FSlateApplication::Get().GetDragDroppingContent();
-
-			FDragDropPayload DraggedObjects;
-			if (DraggedObjects.ParseDrag(*DragDropOp) && !TreeItem->ValidateDrop(DraggedObjects, *World).IsValid())
-			{
-				return DarkColor;
-			}
-		}
-
-		if (!TreeItem->CanInteract())
-		{
-			return DarkColor;
-		}
-
-		return TOptional<FLinearColor>();
-	}
-
-	bool CanExecuteRenameRequest(const FTreeItemPtr& Item)
-	{
-		if (ISceneOutliner* SceneOutliner = WeakSceneOutliner.Pin().Get())
-		{
-			return SceneOutliner->CanExecuteRenameRequest(Item);
-		}
-		return false;
-	}
-};
-
-const FLinearColor FCommonLabelData::DarkColor(0.3f, 0.3f, 0.3f);
-
-struct SActorTreeLabel : FCommonLabelData, public SCompoundWidget
-{
-	SLATE_BEGIN_ARGS(SActorTreeLabel) {}
-	SLATE_END_ARGS()
-
-	void Construct(const FArguments& InArgs, FActorTreeItem& ActorItem, ISceneOutliner& SceneOutliner, const STableRow<FTreeItemPtr>& InRow)
-	{
-		WeakSceneOutliner = StaticCastSharedRef<ISceneOutliner>(SceneOutliner.AsShared());
-
-		TreeItemPtr = StaticCastSharedRef<FActorTreeItem>(ActorItem.AsShared());
-		ActorPtr = ActorItem.Actor;
-
-		HighlightText = SceneOutliner.GetFilterHighlightText();
-
-		TSharedPtr<SInlineEditableTextBlock> InlineTextBlock;
-
-		auto MainContent = SNew(SHorizontalBox)
-
-			// Main actor label
-			+ SHorizontalBox::Slot()
-			.VAlign(VAlign_Center)
-			[
-				SAssignNew(InlineTextBlock, SInlineEditableTextBlock)
-				.Text(this, &SActorTreeLabel::GetDisplayText)
-				.ToolTipText(this, &SActorTreeLabel::GetTooltipText)
-				.HighlightText(HighlightText)
-				.ColorAndOpacity(this, &SActorTreeLabel::GetForegroundColor)
-				.OnTextCommitted(this, &SActorTreeLabel::OnLabelCommitted)
-				.OnVerifyTextChanged(this, &SActorTreeLabel::OnVerifyItemLabelChanged)
-				.IsSelected(FIsSelected::CreateSP(&InRow, &STableRow<FTreeItemPtr>::IsSelectedExclusively))
-				.IsReadOnly_Lambda([Item = ActorItem.AsShared(), this]()
-					{
-						return !CanExecuteRenameRequest(Item);
-					})
-			]
-
-			+ SHorizontalBox::Slot()
-			.VAlign(VAlign_Center)
-			.AutoWidth()
-			.Padding( 0.0f, 0.f, 3.0f, 0.0f )
-			[
-				SNew( STextBlock )
-				.Text(this, &SActorTreeLabel::GetTypeText)
-				.Visibility(this, &SActorTreeLabel::GetTypeTextVisibility)
-				.HighlightText(HighlightText)
-			];
-
-		if (ActorItem.GetSharedData().Mode == ESceneOutlinerMode::ActorBrowsing)
-		{
-			ActorItem.RenameRequestEvent.BindSP( InlineTextBlock.Get(), &SInlineEditableTextBlock::EnterEditingMode );
-		}
-		
-		ChildSlot
-		[
-			SNew(SHorizontalBox)
-
-			+ SHorizontalBox::Slot()
-			.AutoWidth()
-			.VAlign(VAlign_Center)
-			.Padding(FDefaultTreeItemMetrics::IconPadding())
-			[
-				SNew(SBox)
-				.WidthOverride(FDefaultTreeItemMetrics::IconSize())
-				.HeightOverride(FDefaultTreeItemMetrics::IconSize())
-				[
-					SNew(SImage)
-					.Image(this, &SActorTreeLabel::GetIcon)
-					.ToolTipText(this, &SActorTreeLabel::GetIconTooltip)
-				]
-			]
-
-			+ SHorizontalBox::Slot()
-			.FillWidth(1.0f)
-			.VAlign(VAlign_Center)
-			.Padding(0.0f, 0.0f)
-			[
-				MainContent
-			]
-		];
-	}
-
-private:
-	TWeakPtr<FActorTreeItem> TreeItemPtr;
-	TWeakObjectPtr<AActor> ActorPtr;
-	TAttribute<FText> HighlightText;
-
-	FText GetDisplayText() const
-	{
-		const AActor* Actor = ActorPtr.Get();
-		return Actor ? FText::FromString(Actor->GetActorLabel()) : LOCTEXT("ActorLabelForMissingActor", "(Deleted Actor)");
-	}
-
-	FText GetTooltipText() const
-	{
-		if (const AActor* Actor = ActorPtr.Get())
-		{
-			FFormatNamedArguments Args;
-			Args.Add(TEXT("ID_Name"), LOCTEXT("CustomColumnMode_InternalName", "ID Name"));
-			Args.Add(TEXT("Name"), FText::FromString(Actor->GetName()));
-			return FText::Format(LOCTEXT("ActorNameTooltip", "{ID_Name}: {Name}"), Args);
-		}
-
-		return FText();
-	}
-
-	FText GetTypeText() const
-	{
-		if (const AActor* Actor = ActorPtr.Get())
-		{
-			return FText::FromName(Actor->GetClass()->GetFName());
-		}			
-
-		return FText();
-	}
-
-	EVisibility GetTypeTextVisibility() const
-	{
-		return HighlightText.Get().IsEmpty() ? EVisibility::Collapsed : EVisibility::Visible;
-	}
-
-	const FSlateBrush* GetIcon() const
-	{
-		if (const AActor* Actor = ActorPtr.Get())
-		{
-			if (WeakSceneOutliner.IsValid())
-			{
-				FName IconName = Actor->GetCustomIconName();
-				if (IconName == NAME_None)
-				{
-					IconName = Actor->GetClass()->GetFName();
-				}
-
-				const FSlateBrush* CachedBrush = WeakSceneOutliner.Pin()->GetCachedIconForClass(IconName);
-				if (CachedBrush != nullptr)
-				{
-					return CachedBrush;
-				}
-				else
-				{
-
-					const FSlateBrush* FoundSlateBrush = FClassIconFinder::FindIconForActor(Actor);
-					WeakSceneOutliner.Pin()->CacheIconForClass(IconName, FoundSlateBrush);
-					return FoundSlateBrush;
-				}
-			}
-			else
-			{
-				return nullptr;
-			}
-		}
-		else
-		{
-			return nullptr;
-		}
-	}
-
-	const FSlateBrush* GetIconOverlay() const
-	{
-		static const FName SequencerActorTag(TEXT("SequencerActor"));
-
-		if (const AActor* Actor = ActorPtr.Get())
-		{
-			if (Actor->ActorHasTag(SequencerActorTag))
-			{
-				return FEditorStyle::GetBrush("Sequencer.SpawnableIconOverlay");
-			}
-		}
-		return nullptr;
-	}
-
-	FText GetIconTooltip() const
-	{
-		auto TreeItem = TreeItemPtr.Pin();
-		if (!TreeItem.IsValid())
-		{
-			return FText();
-		}
-
-		FText ToolTipText;
-		if (AActor* Actor = ActorPtr.Get())
-		{
-			ToolTipText = FText::FromString( Actor->GetClass()->GetName() );
-			if (TreeItem->GetSharedData().Mode == ESceneOutlinerMode::ActorBrowsing)
-			{
-				USceneComponent* RootComponent = Actor->GetRootComponent();
-				if (RootComponent)
-				{
-					FFormatNamedArguments Args;
-					Args.Add( TEXT("ActorClassName"), ToolTipText );
-
-					if( RootComponent->Mobility == EComponentMobility::Static )
-					{
-						ToolTipText = FText::Format( LOCTEXT( "ComponentMobility_Static", "{ActorClassName} with static mobility" ), Args );
-					}
-					else if( RootComponent->Mobility == EComponentMobility::Stationary )
-					{
-						ToolTipText = FText::Format( LOCTEXT( "ComponentMobility_Stationary", "{ActorClassName} with stationary mobility" ), Args );
-					}
-					else if( RootComponent->Mobility == EComponentMobility::Movable )
-					{
-						ToolTipText = FText::Format( LOCTEXT( "ComponentMobility_Movable", "{ActorClassName} with movable mobility" ), Args );
-					}
-				}
-			}
-		}
-
-		return ToolTipText;
-	}
-
-	FSlateColor GetForegroundColor() const
-	{
-		auto TreeItem = TreeItemPtr.Pin();
-		if (auto BaseColor = FCommonLabelData::GetForegroundColor(TreeItem))
-		{
-			return BaseColor.GetValue();
-		}
-
-		AActor* Actor = ActorPtr.Get();
-		if (!Actor)
-		{
-			// Deleted actor!
-			return FLinearColor(0.2f, 0.2f, 0.25f);
-		}
-
-		if (TreeItem->GetSharedData().bRepresentingPlayWorld && !TreeItem->bExistsInCurrentWorldAndPIE)
-		{
-			// Highlight actors that are exclusive to PlayWorld
-			return FLinearColor::FromSRGBColor(FColor::FromHex("FFE65A"));
-
-		}
-
-		// also darken items that are non selectable in the active mode(s)
-		const bool bInSelected = true;
-		const bool bSelectEvenIfHidden = true;		// @todo outliner: Is this actually OK?
-		if (!GEditor->CanSelectActor(Actor, bInSelected, bSelectEvenIfHidden))
-		{
-			return FCommonLabelData::DarkColor;
-		}
-
-		return FSlateColor::UseForeground();
-	}
-
-	bool OnVerifyItemLabelChanged(const FText& InLabel, FText& OutErrorMessage)
-	{
-		return FActorEditorUtils::ValidateActorName(InLabel, OutErrorMessage);
-	}
-
-	void OnLabelCommitted(const FText& InLabel, ETextCommit::Type InCommitInfo)
-	{
-		auto* Actor = ActorPtr.Get();
-		if (Actor && Actor->IsActorLabelEditable() && !InLabel.ToString().Equals(Actor->GetActorLabel(), ESearchCase::CaseSensitive))
-		{
-			const FScopedTransaction Transaction( LOCTEXT( "SceneOutlinerRenameActorTransaction", "Rename Actor" ) );
-			FActorLabelUtilities::RenameExistingActor(Actor, InLabel.ToString());
-
-			auto Outliner = WeakSceneOutliner.Pin();
-			if (Outliner.IsValid())
-			{
-				Outliner->SetKeyboardFocus();
-			}
-		}
-	}
-};
-
-struct SWorldTreeLabel : FCommonLabelData, public SCompoundWidget
-{
-	SLATE_BEGIN_ARGS(SWorldTreeLabel){}
-	SLATE_END_ARGS()
-
-	void Construct(const FArguments& InArgs, FWorldTreeItem& WorldItem, ISceneOutliner& SceneOutliner, const STableRow<FTreeItemPtr>& InRow)
-	{
-		TreeItemPtr = StaticCastSharedRef<FWorldTreeItem>(WorldItem.AsShared());
-		WeakSceneOutliner = StaticCastSharedRef<ISceneOutliner>(SceneOutliner.AsShared());
-
-		ChildSlot
-		[
-			SNew(SHorizontalBox)
-
-			+ SHorizontalBox::Slot()
-			.AutoWidth()
-			.VAlign(VAlign_Center)
-			.Padding(FDefaultTreeItemMetrics::IconPadding())
-			[
-				SNew(SBox)
-				.WidthOverride(FDefaultTreeItemMetrics::IconSize())
-				.HeightOverride(FDefaultTreeItemMetrics::IconSize())
-				[
-					SNew(SImage)
-					.Image(FSlateIconFinder::FindIconBrushForClass(UWorld::StaticClass()))
-					.ToolTipText(LOCTEXT("WorldIcon_Tooltip", "World"))
-				]
-			]
-
-			+ SHorizontalBox::Slot()
-			.FillWidth(1.0f)
-			.VAlign(VAlign_Center)
-			.Padding(0.0f, 2.0f)
-			[
-				SNew(STextBlock)
-				.Text(this, &SWorldTreeLabel::GetDisplayText)
-				.HighlightText(SceneOutliner.GetFilterHighlightText())
-				.ColorAndOpacity(this, &SWorldTreeLabel::GetForegroundColor)
-				.ToolTip(IDocumentation::Get()->CreateToolTip(
-					TAttribute<FText>(this, &SWorldTreeLabel::GetTooltipText),
-					nullptr,
-					TEXT("Shared/LevelEditor/SceneOutliner"),
-					TEXT("WorldSettingsLabel")))
-			]
-		];
-	}
-
-private:
-	TWeakPtr<FWorldTreeItem> TreeItemPtr;
-
-	FText GetDisplayText() const
-	{
-		auto Item = TreeItemPtr.Pin();
-		return Item.IsValid() ? FText::FromString(Item->GetDisplayString()) : FText();
-	}
-
-	FText GetTooltipText() const
-	{
-		auto Item = TreeItemPtr.Pin();
-		FText PersistentLevelDisplayName = Item.IsValid() ? FText::FromString(Item->GetWorldName()) : FText();
-		if(Item->CanInteract())
-		{
-			return FText::Format(LOCTEXT("WorldLabel_Tooltip", "The world settings for {0}, double-click to edit"), PersistentLevelDisplayName);
-		}
-		else
-		{
-			return FText::Format(LOCTEXT("WorldLabel_TooltipNonInteractive", "The world {0}"), PersistentLevelDisplayName);
-		}
-	}
-
-	FSlateColor GetForegroundColor() const
-	{
-		if (auto BaseColor = FCommonLabelData::GetForegroundColor(TreeItemPtr.Pin()))
-		{
-			return BaseColor.GetValue();
-		}
-		
-		return FSlateColor::UseForeground();
-	}
-};
-
-struct SFolderTreeLabel : FCommonLabelData, public SCompoundWidget
-{
-	SLATE_BEGIN_ARGS(SFolderTreeLabel){}
-	SLATE_END_ARGS()
-
-	void Construct(const FArguments& InArgs, FFolderTreeItem& FolderItem, ISceneOutliner& SceneOutliner, const STableRow<FTreeItemPtr>& InRow)
-	{
-		TreeItemPtr = StaticCastSharedRef<FFolderTreeItem>(FolderItem.AsShared());
-		WeakSceneOutliner = StaticCastSharedRef<ISceneOutliner>(SceneOutliner.AsShared());
-
-		TSharedPtr<SInlineEditableTextBlock> InlineTextBlock = SNew(SInlineEditableTextBlock)
-			.Text(this, &SFolderTreeLabel::GetDisplayText)
-			.HighlightText(SceneOutliner.GetFilterHighlightText())
-			.ColorAndOpacity(this, &SFolderTreeLabel::GetForegroundColor)
-			.OnTextCommitted(this, &SFolderTreeLabel::OnLabelCommitted)
-			.OnVerifyTextChanged(this, &SFolderTreeLabel::OnVerifyItemLabelChanged)
-			.IsSelected(FIsSelected::CreateSP(&InRow, &STableRow<FTreeItemPtr>::IsSelectedExclusively))
-			.IsReadOnly_Lambda([Item = FolderItem.AsShared(), this]()
-			{
-				return !CanExecuteRenameRequest(Item);
-			});
-
-		if (SceneOutliner.GetSharedData().Mode == ESceneOutlinerMode::ActorBrowsing)
-		{
-			FolderItem.RenameRequestEvent.BindSP( InlineTextBlock.Get(), &SInlineEditableTextBlock::EnterEditingMode );
-		}
-
-		ChildSlot
-		[
-			SNew(SHorizontalBox)
-
-			+ SHorizontalBox::Slot()
-			.AutoWidth()
-			.VAlign(VAlign_Center)
-			.Padding(FDefaultTreeItemMetrics::IconPadding())
-			[
-				SNew(SBox)
-				.WidthOverride(FDefaultTreeItemMetrics::IconSize())
-				[
-					SNew(SImage)
-					.Image(this, &SFolderTreeLabel::GetIcon)
-				]
-			]
-
-			+ SHorizontalBox::Slot()
-			.FillWidth(1.0f)
-			.VAlign(VAlign_Center)
-			.Padding(0.0f, 2.0f)
-			[
-				InlineTextBlock.ToSharedRef()
-			]
-		];
-	}
-
-private:
-	TWeakPtr<FFolderTreeItem> TreeItemPtr;
-
-	FText GetDisplayText() const
-	{
-		auto Folder = TreeItemPtr.Pin();
-		return Folder.IsValid() ? FText::FromName(Folder->LeafName) : FText();
-	}
-
-	const FSlateBrush* GetIcon() const
-	{	
-		auto TreeItem = TreeItemPtr.Pin();
-		if (!TreeItem.IsValid())
-		{
-			return FEditorStyle::Get().GetBrush(TEXT("SceneOutliner.FolderClosed"));
-		}
-
-		if (TreeItem->Flags.bIsExpanded && TreeItem->GetChildren().Num())
-		{
-			return FEditorStyle::Get().GetBrush(TEXT("SceneOutliner.FolderOpen"));
-		}
-		else
-		{
-			return FEditorStyle::Get().GetBrush(TEXT("SceneOutliner.FolderClosed"));
-		}
-	}
-
-	FSlateColor GetForegroundColor() const
-	{
-		if (auto BaseColor = FCommonLabelData::GetForegroundColor(TreeItemPtr.Pin()))
-		{
-			return BaseColor.GetValue();
-		}
-
-		return FSlateColor::UseForeground();
-	}
-
-	bool OnVerifyItemLabelChanged( const FText& InLabel, FText& OutErrorMessage)
-	{
-		auto TreeItem = TreeItemPtr.Pin();
-		if (!TreeItem.IsValid())
-		{
-			OutErrorMessage = LOCTEXT( "RenameFailed_TreeItemDeleted", "Tree item no longer exists" );
-			return false;
-		}
-
-		FText TrimmedLabel = FText::TrimPrecedingAndTrailing(InLabel);
-
-		if (TrimmedLabel.IsEmpty())
-		{
-			OutErrorMessage = LOCTEXT( "RenameFailed_LeftBlank", "Names cannot be left blank" );
-			return false;
-		}
-
-		if (TrimmedLabel.ToString().Len() >= NAME_SIZE)
-		{
-			FFormatNamedArguments Arguments;
-			Arguments.Add( TEXT("CharCount"), NAME_SIZE );
-			OutErrorMessage = FText::Format(LOCTEXT("RenameFailed_TooLong", "Names must be less than {CharCount} characters long."), Arguments);
-			return false;
-		}
-
-		FString LabelString = TrimmedLabel.ToString();
-		if (TreeItem->LeafName.ToString() == LabelString)
-		{
-			return true;
-		}
-
-		int32 Dummy = 0;
-		if (LabelString.FindChar('/', Dummy) || LabelString.FindChar('\\', Dummy))
-		{
-			OutErrorMessage = LOCTEXT("RenameFailed_InvalidChar", "Folder names cannot contain / or \\.");
-			return false;
-		}
-
-		// Validate that this folder doesn't exist already
-		FName NewPath = GetParentPath(TreeItem->Path);
-		if (NewPath.IsNone())
-		{
-			NewPath = FName(*LabelString);
-		}
-		else
-		{
-			NewPath = FName(*(NewPath.ToString() / LabelString));
-		}
-
-		if (FActorFolders::Get().GetFolderProperties(*GWorld, NewPath))
-		{
-			OutErrorMessage = LOCTEXT("RenameFailed_AlreadyExists", "A folder with this name already exists at this level");
-			return false;
-		}
-
-		return true;
-	}
-
-	void OnLabelCommitted(const FText& InLabel, ETextCommit::Type InCommitInfo)
-	{
-		auto TreeItem = TreeItemPtr.Pin();
-		if (TreeItem.IsValid() && !InLabel.ToString().Equals(TreeItem->LeafName.ToString(), ESearchCase::CaseSensitive))
-		{
-			// Rename the item
-			FName NewPath = GetParentPath(TreeItem->Path);
-			if (NewPath.IsNone())
-			{
-				NewPath = FName(*InLabel.ToString());
-			}
-			else
-			{
-				NewPath = FName(*(NewPath.ToString() / InLabel.ToString()));
-			}
-
-			FActorFolders::Get().RenameFolderInWorld(*GWorld, TreeItem->Path, NewPath);
-
-			auto Outliner = WeakSceneOutliner.Pin();
-			if (Outliner.IsValid())
-			{
-				Outliner->SetKeyboardFocus();
-			}
-		}
-	}
-};
-
-struct SComponentTreeLabel : FCommonLabelData, public SCompoundWidget
-{
-	SLATE_BEGIN_ARGS(SComponentTreeLabel) {}
-	SLATE_END_ARGS()
-
-	void Construct(const FArguments& InArgs, FComponentTreeItem& ComponentItem, ISceneOutliner& SceneOutliner, const STableRow<FTreeItemPtr>& InRow)
-	{
-		TreeItemPtr = StaticCastSharedRef<FComponentTreeItem>(ComponentItem.AsShared());
-		WeakSceneOutliner = StaticCastSharedRef<ISceneOutliner>(SceneOutliner.AsShared());
-
-		ComponentPtr = ComponentItem.Component;
-		
-		HighlightText = SceneOutliner.GetFilterHighlightText();
-
-		TSharedPtr<SInlineEditableTextBlock> InlineTextBlock;
-
-		auto MainContent = SNew(SHorizontalBox)
-
-			// Main actor label
-			+ SHorizontalBox::Slot()
-			[
-				SAssignNew(InlineTextBlock, SInlineEditableTextBlock)
-				.Text(this, &SComponentTreeLabel::GetDisplayText)
-				.ToolTipText(this, &SComponentTreeLabel::GetTooltipText)
-				//.HighlightText(HighlightText)
-				.ColorAndOpacity(this, &SComponentTreeLabel::GetForegroundColor)
-				//.OnTextCommitted(this, &SComponentTreeLabel::OnLabelCommitted)
-				//.OnVerifyTextChanged(this, &SComponentTreeLabel::OnVerifyItemLabelChanged)
-				.IsReadOnly_Lambda([Item = ComponentItem.AsShared(), this]()
-				{
-					return !CanExecuteRenameRequest(Item);
-				})
-				.IsSelected(FIsSelected::CreateSP(&InRow, &STableRow<FTreeItemPtr>::IsSelectedExclusively))
-			];
-
-		if (ComponentItem.GetSharedData().Mode == ESceneOutlinerMode::ActorBrowsing)
-		{
-			ComponentItem.RenameRequestEvent.BindSP(InlineTextBlock.Get(), &SInlineEditableTextBlock::EnterEditingMode);
-		}
-
-		ChildSlot
-			[
-				SNew(SHorizontalBox)
-
-				+ SHorizontalBox::Slot()
-				.AutoWidth()
-				.VAlign(VAlign_Center)
-				.Padding(FDefaultTreeItemMetrics::IconPadding())
-				[
-					SNew(SBox)
-					.WidthOverride(FDefaultTreeItemMetrics::IconSize())
-					.HeightOverride(FDefaultTreeItemMetrics::IconSize())
-					[
-						SNew(SImage)
-						.Image(this, &SComponentTreeLabel::GetIcon)
-						.ToolTipText(this, &SComponentTreeLabel::GetIconTooltip)
-					]
-				]
-
-				+ SHorizontalBox::Slot()
-				.FillWidth(1.0f)
-				.VAlign(VAlign_Center)
-				.Padding(0.0f, 2.0f)
-				[
-					MainContent
-				]
-			];
-
-	}
-
-private:
-	TWeakPtr<FComponentTreeItem> TreeItemPtr;
-	TWeakObjectPtr<UActorComponent> ComponentPtr;
-	TAttribute<FText> HighlightText;
-
-	FText GetDisplayText() const
-	{
-		auto Item = TreeItemPtr.Pin();
-		return Item.IsValid() ? FText::FromString(Item->GetDisplayString()) : FText();
-	}
-
-	FText GetTypeText() const
-	{
-		if (const UActorComponent* Component = ComponentPtr.Get())
-		{
-			return FText::FromName(Component->GetClass()->GetFName());
-		}
-
-		return FText();
-	}
-
-	EVisibility GetTypeTextVisibility() const
-	{
-		return HighlightText.Get().IsEmpty() ? EVisibility::Collapsed : EVisibility::Visible;
-	}
-
-	FText GetTooltipText() const
-	{
-		if (const UActorComponent* Component = ComponentPtr.Get())
-		{
-			FFormatNamedArguments Args;
-			Args.Add(TEXT("ID_Name"), LOCTEXT("CustomColumnMode_InternalName", "ID Name"));
-			Args.Add(TEXT("Name"), FText::FromString(Component->GetName()));
-			return FText::Format(LOCTEXT("ComponentNameTooltip", "{ID_Name}: {Name}"), Args);
-		}
-
-		return FText();
-	}
-
-	const FSlateBrush* GetIcon() const
-	{
-		if (const UActorComponent* Component = ComponentPtr.Get())
-		{
-			if (WeakSceneOutliner.IsValid())
-			{
-				const FSlateBrush* CachedBrush = WeakSceneOutliner.Pin()->GetCachedIconForClass(Component->GetClass()->GetFName());
-				if (CachedBrush != nullptr)
-				{
-					return CachedBrush;
-				}
-				else
-				{
-					const FSlateBrush* FoundSlateBrush = FSlateIconFinder::FindIconBrushForClass(Component->GetClass());
-					WeakSceneOutliner.Pin()->CacheIconForClass(Component->GetClass()->GetFName(), const_cast<FSlateBrush*>(FoundSlateBrush));
-					return FoundSlateBrush;
-				}
-			}
-			else
-			{
-				return nullptr;
-			}
-		}
-		else
-		{
-			return nullptr;
-		}
-	}
-
-	FText GetIconTooltip() const
-	{
-		FText ToolTipText;
-		return ToolTipText;
-	}
-
-	FSlateColor GetForegroundColor() const
-	{
-		if (auto BaseColor = FCommonLabelData::GetForegroundColor(TreeItemPtr.Pin()))
-		{
-			return BaseColor.GetValue();
-		}
-
-		return FSlateColor::UseForeground();
-	}
-};
-
-struct SSubComponentTreeLabel : FCommonLabelData, public SCompoundWidget
-{
-	SLATE_BEGIN_ARGS(SSubComponentTreeLabel) {}
-	SLATE_END_ARGS()
-
-	void Construct(const FArguments& InArgs, FSubComponentTreeItem& SubComponentItem, ISceneOutliner& SceneOutliner, const STableRow<FTreeItemPtr>& InRow)
-	{
-		TreeItemPtr = StaticCastSharedRef<FSubComponentTreeItem>(SubComponentItem.AsShared());
-		WeakSceneOutliner = StaticCastSharedRef<ISceneOutliner>(SceneOutliner.AsShared());
-
-		ComponentPtr = SubComponentItem.ParentComponent;
-
-		TSharedPtr<SInlineEditableTextBlock> InlineTextBlock;
-
-		auto MainContent = SNew(SHorizontalBox)
-
-			// Main actor label
-			+ SHorizontalBox::Slot()
-			[
-				SAssignNew(InlineTextBlock, SInlineEditableTextBlock)
-				.Text(this, &SSubComponentTreeLabel::GetDisplayText)
-				.ToolTipText(this, &SSubComponentTreeLabel::GetTooltipText)
-				//.HighlightText(HighlightText)
-				.ColorAndOpacity(this, &SSubComponentTreeLabel::GetForegroundColor)
-				.OnTextCommitted(this, &SSubComponentTreeLabel::OnLabelCommitted)
-				.OnVerifyTextChanged(this, &SSubComponentTreeLabel::OnVerifyItemLabelChanged)
-				.IsSelected(FIsSelected::CreateSP(&InRow, &STableRow<FTreeItemPtr>::IsSelectedExclusively))
-				.IsReadOnly_Lambda([Item = SubComponentItem.AsShared(), this]()
-				{
-					return !CanExecuteRenameRequest(Item);
-				})
-			];
-
-		if (SubComponentItem.GetSharedData().Mode == ESceneOutlinerMode::ActorBrowsing)
-		{
-			SubComponentItem.RenameRequestEvent.BindSP(InlineTextBlock.Get(), &SInlineEditableTextBlock::EnterEditingMode);
-		}
-
-		ChildSlot
-		[
-			SNew(SHorizontalBox)
-
-			+ SHorizontalBox::Slot()
-			.AutoWidth()
-			.VAlign(VAlign_Center)
-			.Padding(FDefaultTreeItemMetrics::IconPadding())
-			[
-				SNew(SBox)
-				.WidthOverride(FDefaultTreeItemMetrics::IconSize())
-				.HeightOverride(FDefaultTreeItemMetrics::IconSize())
-				[
-					SNew(SImage)
-					.Image(this, &SSubComponentTreeLabel::GetIcon)
-					.ToolTipText(this, &SSubComponentTreeLabel::GetIconTooltip)
-				]
-			]
-
-			+ SHorizontalBox::Slot()
-			.FillWidth(1.0f)
-			.VAlign(VAlign_Center)
-			.Padding(0.0f, 2.0f)
-			[
-				MainContent
-			]
-		];
-	}
-
-private:
-	TWeakPtr<FSubComponentTreeItem> TreeItemPtr;
-	TWeakObjectPtr<UActorComponent> ComponentPtr;
-
-	FText GetDisplayText() const
-	{
-		auto Item = TreeItemPtr.Pin();
-		return Item.IsValid() ? FText::FromString(Item->GetDisplayString()) : FText();
-	}
-
-	FSlateColor GetForegroundColor() const
-	{
-		if (auto BaseColor = FCommonLabelData::GetForegroundColor(TreeItemPtr.Pin()))
-		{
-			return BaseColor.GetValue();
-		}
-
-		return FSlateColor::UseForeground();
-	}
-	
-	FText GetTooltipText() const
-	{
-		if (const UActorComponent* Component = ComponentPtr.Get())
-		{
-			FFormatNamedArguments Args;
-			Args.Add(TEXT("ID_Name"), LOCTEXT("CustomColumnMode_InternalName", "ID Name"));
-			Args.Add(TEXT("Name"), FText::FromString(Component->GetName()));
-			return FText::Format(LOCTEXT("ComponentNameTooltip", "{ID_Name}: {Name}"), Args);
-		}
-
-		return FText();
-	}
-
-	const FSlateBrush* GetIcon() const
-	{
-		UClass* StaticClass = UActorComponent::StaticClass();
-		auto Item = TreeItemPtr.Pin();
-		if (Item.IsValid())
-		{
-			StaticClass = Item->GetIconClass();
-		}
-
-		if (WeakSceneOutliner.IsValid())
-		{
-			const FSlateBrush* CachedBrush = WeakSceneOutliner.Pin()->GetCachedIconForClass(StaticClass->GetFName());
-			if (CachedBrush != nullptr)
-			{
-				return CachedBrush;
-			}
-			else
-			{
-				const FSlateBrush* FoundSlateBrush = FSlateIconFinder::FindIconBrushForClass(StaticClass);
-				WeakSceneOutliner.Pin()->CacheIconForClass(StaticClass->GetFName(), const_cast<FSlateBrush*>(FoundSlateBrush));
-				return FoundSlateBrush;
-			}
-		}
-		else
-		{
-			return nullptr;
-		}
-	}
-
-	FText GetIconTooltip() const
-	{
-		FText ToolTipText;
-		return ToolTipText;
-	}
-
-	virtual bool OnVerifyItemLabelChanged(const FText& InLabel, FText& OutErrorMessage)
-	{
-		bool IsValidName = false;
-		auto Item = TreeItemPtr.Pin();
-		if (Item.IsValid())
-		{
-			IsValidName = Item->ValidateSubComponentName(InLabel, OutErrorMessage);
-		}
-
-		return IsValidName;
-	}
-
-	virtual void OnLabelCommitted(const FText& InLabel, ETextCommit::Type InCommitInfo)
-	{
-		auto Item = TreeItemPtr.Pin();
-		if (Item.IsValid())
-		{
-			Item->RenameSubComponent(InLabel);
-
-			auto Outliner = WeakSceneOutliner.Pin();
-			if (Outliner.IsValid())
-			{
-				Outliner->SetKeyboardFocus();
-			}
-		}
-	}
-};
-
-
 FName FItemLabelColumn::GetColumnID()
 {
 	return GetID();
@@ -928,9 +37,9 @@ SHeaderRow::FColumn::FArguments FItemLabelColumn::ConstructHeaderRowColumn()
 }
 const TSharedRef<SWidget> FItemLabelColumn::ConstructRowWidget(FTreeItemRef TreeItem, const STableRow<FTreeItemPtr>& Row)
 {
-	FColumnWidgetGenerator Generator(*this, Row);
-	TreeItem->Visit(Generator);
-	return Generator.Widget.ToSharedRef();
+	ISceneOutliner* Outliner = WeakSceneOutliner.Pin().Get();
+	check(Outliner);
+	return TreeItem->GenerateLabelWidget(*Outliner, Row);
 }
 
 void FItemLabelColumn::PopulateSearchStrings( const ITreeItem& Item, TArray< FString >& OutSearchStrings ) const
@@ -943,46 +52,10 @@ void FItemLabelColumn::SortItems(TArray<FTreeItemPtr>& OutItems, const EColumnSo
 	typedef FSortHelper<int32, FNumericStringWrapper> FSort;
 
 	FSort()
-		.Primary([](const ITreeItem& Item){ return Item.GetTypeSortPriority(); }, 	SortMode)
-		.Secondary([](const ITreeItem& Item){ return FNumericStringWrapper(Item.GetDisplayString()); }, 	SortMode)
+		.Primary([this](const ITreeItem& Item){ return WeakSceneOutliner.Pin()->GetTypeSortPriority(Item); },SortMode)
+		.Secondary([](const ITreeItem& Item){ return FNumericStringWrapper(Item.GetDisplayString()); }, 			SortMode)
 		.Sort(OutItems);
 }
-
-TSharedRef<SWidget> FItemLabelColumn::GenerateWidget( FActorTreeItem& TreeItem, const STableRow<FTreeItemPtr>& InRow )
-{
-	ISceneOutliner* Outliner = WeakSceneOutliner.Pin().Get();
-	check(Outliner);
-	return SNew(SActorTreeLabel, TreeItem, *Outliner, InRow);
-}
-
-TSharedRef<SWidget> FItemLabelColumn::GenerateWidget( FWorldTreeItem& TreeItem, const STableRow<FTreeItemPtr>& InRow )
-{
-	ISceneOutliner* Outliner = WeakSceneOutliner.Pin().Get();
-	check(Outliner);
-	return SNew(SWorldTreeLabel, TreeItem, *Outliner, InRow);
-}
-
-TSharedRef<SWidget> FItemLabelColumn::GenerateWidget( FFolderTreeItem& TreeItem, const STableRow<FTreeItemPtr>& InRow )
-{
-	ISceneOutliner* Outliner = WeakSceneOutliner.Pin().Get();
-	check(Outliner);
-	return SNew(SFolderTreeLabel, TreeItem, *Outliner, InRow);
-}
-
-TSharedRef<SWidget> FItemLabelColumn::GenerateWidget( FComponentTreeItem& TreeItem, const STableRow<FTreeItemPtr>& InRow)
-{
-	ISceneOutliner* Outliner = WeakSceneOutliner.Pin().Get();
-	check(Outliner);
-	return SNew(SComponentTreeLabel, TreeItem, *Outliner, InRow);
-}
-
-TSharedRef<SWidget> FItemLabelColumn::GenerateWidget(FSubComponentTreeItem& TreeItem, const STableRow<FTreeItemPtr>& InRow)
-{
-	ISceneOutliner* Outliner = WeakSceneOutliner.Pin().Get();
-	check(Outliner);
-	return SNew(SSubComponentTreeLabel, TreeItem, *Outliner, InRow);
-}
-
 }	// namespace SceneOutliner
 
 #undef LOCTEXT_NAMESPACE
diff --git a/Engine/Source/Editor/SceneOutliner/Private/SceneOutlinerItemLabelColumn.h b/Engine/Source/Editor/SceneOutliner/Private/SceneOutlinerItemLabelColumn.h
--- a/Engine/Source/Editor/SceneOutliner/Private/SceneOutlinerItemLabelColumn.h
+++ b/Engine/Source/Editor/SceneOutliner/Private/SceneOutlinerItemLabelColumn.h
@@ -13,7 +13,6 @@ template<typename ItemType> class STableRow;
 
 namespace SceneOutliner
 {
-
 /** A column for the SceneOutliner that displays the item label */
 class FItemLabelColumn : public ISceneOutlinerColumn
 {
@@ -46,28 +45,6 @@ public:
 private:
 	TWeakPtr<ISceneOutliner> WeakSceneOutliner;
 
-private:
-
-	TSharedRef<SWidget> GenerateWidget( FActorTreeItem& TreeItem, const STableRow<FTreeItemPtr>& InRow );
-	TSharedRef<SWidget> GenerateWidget( FWorldTreeItem& TreeItem, const STableRow<FTreeItemPtr>& InRow );
-	TSharedRef<SWidget> GenerateWidget( FFolderTreeItem& TreeItem, const STableRow<FTreeItemPtr>& InRow );
-	TSharedRef<SWidget> GenerateWidget( FComponentTreeItem& TreeItem, const STableRow<FTreeItemPtr>& InRow);
-	TSharedRef<SWidget> GenerateWidget( FSubComponentTreeItem& TreeItem, const STableRow<FTreeItemPtr>& InRow);
-
-	struct FColumnWidgetGenerator : FColumnGenerator
-	{
-		FItemLabelColumn& Column;
-		const STableRow<FTreeItemPtr>& Row;
-
-		FColumnWidgetGenerator(FItemLabelColumn& InColumn, const STableRow<FTreeItemPtr>& InRow) : Column(InColumn), Row(InRow) {}
-
-		virtual TSharedRef<SWidget> GenerateWidget(FFolderTreeItem& Item) const override { return Column.GenerateWidget(Item, Row); }
-		virtual TSharedRef<SWidget> GenerateWidget(FWorldTreeItem& Item) const override { return Column.GenerateWidget(Item, Row); }
-		virtual TSharedRef<SWidget> GenerateWidget(FActorTreeItem& Item) const override { return Column.GenerateWidget(Item, Row); }
-		virtual TSharedRef<SWidget> GenerateWidget(FComponentTreeItem& Item) const override { return Column.GenerateWidget(Item, Row); }
-		virtual TSharedRef<SWidget> GenerateWidget(FSubComponentTreeItem& Item) const override { return Column.GenerateWidget(Item, Row); }
-	};
-	friend FColumnWidgetGenerator;
 };
 
 }	// namespace SceneOutliner
diff --git a/Engine/Source/Editor/SceneOutliner/Private/SceneOutlinerModule.cpp b/Engine/Source/Editor/SceneOutliner/Private/SceneOutlinerModule.cpp
--- a/Engine/Source/Editor/SceneOutliner/Private/SceneOutlinerModule.cpp
+++ b/Engine/Source/Editor/SceneOutliner/Private/SceneOutlinerModule.cpp
@@ -4,78 +4,119 @@
 #include "Modules/ModuleManager.h"
 #include "Widgets/DeclarativeSyntaxSupport.h"
 #include "Framework/Application/SlateApplication.h"
-#include "SceneOutlinerSettings.h"
 #include "SSceneOutliner.h"
 
-#include "ICustomSceneOutliner.h"
 #include "SceneOutlinerActorInfoColumn.h"
 #include "SceneOutlinerGutter.h"
 #include "SceneOutlinerItemLabelColumn.h"
 #include "SceneOutlinerPublicTypes.h"
 
+#include "ActorPickingMode.h"
+#include "ActorBrowsingMode.h"
+#include "ActorTreeItem.h"
+#include "ComponentTreeItem.h"
+
 /* FSceneOutlinerModule interface
  *****************************************************************************/
 
-namespace SceneOutliner
-{
-	void OnSceneOutlinerItemClicked(TSharedRef<ITreeItem> Item, FOnActorPicked OnActorPicked)
-	{
-		Item->Visit(
-			FFunctionalVisitor()
-			.Actor([&](const FActorTreeItem& ActorItem){
-				if (AActor* Actor = ActorItem.Actor.Get())
-				{
-					OnActorPicked.ExecuteIfBound(Actor);
-				}
-			})
-		);
-	}
-}
-
-
 void FSceneOutlinerModule::StartupModule()
 {
-	RegisterDefaultColumnType< SceneOutliner::FSceneOutlinerGutter >(SceneOutliner::FDefaultColumnInfo(SceneOutliner::FColumnInfo(SceneOutliner::EColumnVisibility::Visible, 0), ESceneOutlinerMode::ActorBrowsing));
-	RegisterDefaultColumnType< SceneOutliner::FItemLabelColumn >(SceneOutliner::FDefaultColumnInfo(SceneOutliner::FColumnInfo(SceneOutliner::EColumnVisibility::Visible, 10)));
-	RegisterDefaultColumnType< SceneOutliner::FActorInfoColumn >(SceneOutliner::FDefaultColumnInfo(SceneOutliner::FColumnInfo(SceneOutliner::EColumnVisibility::Visible, 20)));
+	RegisterDefaultColumnType< SceneOutliner::FItemLabelColumn >(SceneOutliner::FColumnInfo(SceneOutliner::EColumnVisibility::Visible, 10));
+
+	// Register builtin column types which are not active by default
+	RegisterColumnType<SceneOutliner::FSceneOutlinerGutter>();
+	RegisterColumnType<SceneOutliner::FActorInfoColumn>();
 }
 
 
 void FSceneOutlinerModule::ShutdownModule()
 {
-	UnRegisterColumnType< SceneOutliner::FSceneOutlinerGutter >();
-	UnRegisterColumnType< SceneOutliner::FItemLabelColumn >();
-	UnRegisterColumnType< SceneOutliner::FActorInfoColumn >();
+	UnRegisterColumnType<SceneOutliner::FSceneOutlinerGutter>();
+	UnRegisterColumnType<SceneOutliner::FItemLabelColumn>();
+	UnRegisterColumnType<SceneOutliner::FActorInfoColumn>();
 }
 
-
-TSharedRef< ISceneOutliner > FSceneOutlinerModule::CreateSceneOutliner( const SceneOutliner::FInitializationOptions& InitOptions, const FOnActorPicked& OnActorPickedDelegate ) const
+TSharedRef<ISceneOutliner> FSceneOutlinerModule::CreateSceneOutliner(const SceneOutliner::FInitializationOptions& InitOptions) const
 {
-	auto OnItemPicked = FOnSceneOutlinerItemPicked::CreateStatic( &SceneOutliner::OnSceneOutlinerItemClicked, OnActorPickedDelegate );
-	return CreateSceneOutliner(InitOptions, OnItemPicked);
+	return SNew(SceneOutliner::SSceneOutliner, InitOptions)
+		.IsEnabled(FSlateApplication::Get().GetNormalExecutionAttribute());
 }
 
-
-TSharedRef< ISceneOutliner > FSceneOutlinerModule::CreateSceneOutliner( const SceneOutliner::FInitializationOptions& InitOptions, const FOnSceneOutlinerItemPicked& OnItemPickedDelegate ) const
+TSharedRef<ISceneOutliner> FSceneOutlinerModule::CreateActorPicker(const SceneOutliner::FInitializationOptions& InInitOptions, const FOnActorPicked& OnActorPickedDelegate, TWeakObjectPtr<UWorld> SpecifiedWorld) const
 {
-	return SNew( SceneOutliner::SSceneOutliner, InitOptions )
-		.IsEnabled( FSlateApplication::Get().GetNormalExecutionAttribute() )
-		.OnItemPickedDelegate( OnItemPickedDelegate );
-}
+	auto OnItemPicked = FOnSceneOutlinerItemPicked::CreateLambda([OnActorPickedDelegate](TSharedRef<SceneOutliner::ITreeItem> Item)
+		{
+			if (SceneOutliner::FActorTreeItem* ActorItem = Item->CastTo<SceneOutliner::FActorTreeItem>())
+			{
+				if (ActorItem->IsValid())
+				{
+					OnActorPickedDelegate.ExecuteIfBound(ActorItem->Actor.Get());
 
-TSharedRef< ICustomSceneOutliner >FSceneOutlinerModule::CreateCustomSceneOutliner( SceneOutliner::FInitializationOptions& InitOptions ) const
-{
-	InitOptions.Mode = ESceneOutlinerMode::Custom;
-	return SNew( SceneOutliner::SSceneOutliner, InitOptions );
+				}
+			}
+		});
+
+	SceneOutliner::FCreateOutlinerMode ModeFactory = SceneOutliner::FCreateOutlinerMode::CreateLambda([&OnItemPicked, &SpecifiedWorld](SceneOutliner::SSceneOutliner* Outliner)
+		{
+			return new SceneOutliner::FActorPickingMode(Outliner, true, OnItemPicked, SpecifiedWorld);
+		});
+
+	
+
+	SceneOutliner::FInitializationOptions InitOptions(InInitOptions);
+	InitOptions.ModeFactory = ModeFactory;
+	if (InitOptions.ColumnMap.Num() == 0)
+	{
+		InitOptions.ColumnMap.Add(SceneOutliner::FBuiltInColumnTypes::Label(), SceneOutliner::FColumnInfo(SceneOutliner::EColumnVisibility::Visible, 0));
+		InitOptions.ColumnMap.Add(SceneOutliner::FBuiltInColumnTypes::ActorInfo(), SceneOutliner::FColumnInfo(SceneOutliner::EColumnVisibility::Visible, 10));
+	}
+	return CreateSceneOutliner(InitOptions);
 }
 
+TSharedRef<ISceneOutliner> FSceneOutlinerModule::CreateComponentPicker(const SceneOutliner::FInitializationOptions& InInitOptions, const FOnComponentPicked& OnComponentPickedDelegate, TWeakObjectPtr<UWorld> SpecifiedWorld) const
+{
+	auto OnItemPicked = FOnSceneOutlinerItemPicked::CreateLambda([OnComponentPickedDelegate](TSharedRef<SceneOutliner::ITreeItem> Item)
+		{
+			if (SceneOutliner::FComponentTreeItem* ComponentItem = Item->CastTo<SceneOutliner::FComponentTreeItem>())
+			{
+				if (ComponentItem->IsValid())
+				{
+					OnComponentPickedDelegate.ExecuteIfBound(ComponentItem->Component.Get());
+				}
+			}
+		});
 
-/* Class constructors
- *****************************************************************************/
+	SceneOutliner::FCreateOutlinerMode ModeFactory = SceneOutliner::FCreateOutlinerMode::CreateLambda([&OnItemPicked, &SpecifiedWorld](SceneOutliner::SSceneOutliner* Outliner)
+		{
+			return new SceneOutliner::FActorPickingMode(Outliner, false, OnItemPicked, SpecifiedWorld);
+		});
 
-USceneOutlinerSettings::USceneOutlinerSettings( const FObjectInitializer& ObjectInitializer )
-	: Super(ObjectInitializer)
-{ }
+	SceneOutliner::FInitializationOptions InitOptions(InInitOptions);
+	InitOptions.ModeFactory = ModeFactory;
+	if (InitOptions.ColumnMap.Num() == 0)
+	{
+		InitOptions.ColumnMap.Add(SceneOutliner::FBuiltInColumnTypes::Label(), SceneOutliner::FColumnInfo(SceneOutliner::EColumnVisibility::Visible, 0));
+		InitOptions.ColumnMap.Add(SceneOutliner::FBuiltInColumnTypes::ActorInfo(), SceneOutliner::FColumnInfo(SceneOutliner::EColumnVisibility::Visible, 10));
+	}
+	return CreateSceneOutliner(InitOptions);
+}
 
+TSharedRef< ISceneOutliner > FSceneOutlinerModule::CreateActorBrowser(const SceneOutliner::FInitializationOptions& InInitOptions, TWeakObjectPtr<UWorld> SpecifiedWorld) const
+{
+	SceneOutliner::FCreateOutlinerMode ModeFactory = SceneOutliner::FCreateOutlinerMode::CreateLambda([&SpecifiedWorld](SceneOutliner::SSceneOutliner* Outliner)
+		{
+			return new SceneOutliner::FActorBrowsingMode(Outliner, SpecifiedWorld);
+		});
+
+	SceneOutliner::FInitializationOptions InitOptions(InInitOptions);
+	InitOptions.ModeFactory = ModeFactory;
+	if (InitOptions.ColumnMap.Num() == 0)
+	{
+		InitOptions.UseDefaultColumns();
+		InitOptions.ColumnMap.Add(SceneOutliner::FBuiltInColumnTypes::Gutter(), SceneOutliner::FColumnInfo(SceneOutliner::EColumnVisibility::Visible, 0));
+		InitOptions.ColumnMap.Add(SceneOutliner::FBuiltInColumnTypes::ActorInfo(), SceneOutliner::FColumnInfo(SceneOutliner::EColumnVisibility::Visible, 20));
+	}
+	return CreateSceneOutliner(InitOptions);
+}
 
 IMPLEMENT_MODULE(FSceneOutlinerModule, SceneOutliner);
diff --git a/Engine/Source/Editor/SceneOutliner/Private/SceneOutlinerStandaloneTypes.cpp b/Engine/Source/Editor/SceneOutliner/Private/SceneOutlinerStandaloneTypes.cpp
--- a/Engine/Source/Editor/SceneOutliner/Private/SceneOutlinerStandaloneTypes.cpp
+++ b/Engine/Source/Editor/SceneOutliner/Private/SceneOutlinerStandaloneTypes.cpp
@@ -3,12 +3,57 @@
 #include "SceneOutlinerStandaloneTypes.h"
 
 #include "EditorActorFolders.h"
+#include "ITreeItem.h"
+#include "ISceneOutliner.h"
+#include "ISceneOutlinerMode.h"
 
 
 #define LOCTEXT_NAMESPACE "SceneOutlinerStandaloneTypes"
 
 namespace SceneOutliner
 {
+	uint32 FTreeItemType::NextUniqueID = 0;
+	const FTreeItemType ITreeItem::Type;
+
+	const FLinearColor FCommonLabelData::DarkColor(0.3f, 0.3f, 0.3f);
+
+	TOptional<FLinearColor> FCommonLabelData::GetForegroundColor(const ITreeItem& TreeItem) const
+	{
+		if (!TreeItem.IsValid())
+		{
+			return DarkColor;
+		}
+
+		// Darken items that aren't suitable targets for an active drag and drop action
+		if (FSlateApplication::Get().IsDragDropping())
+		{
+			TSharedPtr<FDragDropOperation> DragDropOp = FSlateApplication::Get().GetDragDroppingContent();
+
+			FDragDropPayload DraggedObjects;
+			const auto Outliner = WeakSceneOutliner.Pin();
+			if (Outliner->GetMode()->ParseDragDrop(DraggedObjects, *DragDropOp) && !Outliner->GetMode()->ValidateDrop(TreeItem, DraggedObjects).IsValid())
+			{
+				return DarkColor;
+			}
+		}
+
+		if (!TreeItem.CanInteract())
+		{
+			return DarkColor;
+		}
+
+		return TOptional<FLinearColor>();
+	}
+
+	bool FCommonLabelData::CanExecuteRenameRequest(const ITreeItem& Item) const
+	{
+		if (const ISceneOutliner* SceneOutliner = WeakSceneOutliner.Pin().Get())
+		{
+			return SceneOutliner->CanExecuteRenameRequest(Item);
+		}
+		return false;
+	}
+
 	/** Parse a new path (including leaf-name) into this tree item. Does not do any notification */
 	FName GetFolderLeafName(FName InPath)
 	{
@@ -24,19 +69,15 @@ namespace SceneOutliner
 		}
 	}
 
-	FName MoveFolderTo(FName InPath, FName NewParent, UWorld& World)
+	bool PathIsChildOf(const FName& PotentialChild, const FName& Parent)
 	{
-		FName LeafName = GetFolderLeafName(InPath);
-
-		// Get unique name
-		FName NewPath = FActorFolders::Get().GetFolderName(World, NewParent, LeafName);
-
-		if (FActorFolders::Get().RenameFolderInWorld(World, InPath, NewPath))
-		{
-			return NewPath;
-		}
-
-		return FName();
+		const FString ParentString = Parent.ToString();
+		const FString ChildString = PotentialChild.ToString();
+		const int32 ParentLen = ParentString.Len();
+		return
+			ChildString.Len() > ParentLen &&
+			ChildString[ParentLen] == '/' &&
+			ChildString.Left(ParentLen) == ParentString;
 	}
 
 }	// namespace SceneOutliner
diff --git a/Engine/Source/Editor/SceneOutliner/Private/WorldTreeItem.cpp b/Engine/Source/Editor/SceneOutliner/Private/WorldTreeItem.cpp
--- a/Engine/Source/Editor/SceneOutliner/Private/WorldTreeItem.cpp
+++ b/Engine/Source/Editor/SceneOutliner/Private/WorldTreeItem.cpp
@@ -4,45 +4,170 @@
 #include "Modules/ModuleManager.h"
 #include "Textures/SlateIcon.h"
 #include "Framework/Commands/UIAction.h"
+#include "Widgets/SToolTip.h"
 #include "EditorStyleSet.h"
 #include "ScopedTransaction.h"
 #include "SceneOutlinerDragDrop.h"
 #include "SSceneOutliner.h"
-
+#include "Styling/SlateIconFinder.h"
 #include "ToolMenus.h"
 #include "LevelEditor.h"
 #include "EditorActorFolders.h"
+#include "ISceneOutliner.h"
+#include "ISceneOutlinerMode.h"
+#include "IDocumentation.h"
+#include "FolderTreeItem.h"
 
 #define LOCTEXT_NAMESPACE "SceneOutliner_WorldTreeItem"
 
 namespace SceneOutliner
 {
-
-FWorldTreeItem::FWorldTreeItem(UWorld* InWorld)
-	: World(InWorld)
-	, ID(InWorld)
+FText GetWorldDescription(UWorld* World)
 {
+	FText Description;
+	if (World)
+	{
+		FText PostFix;
+		const FWorldContext* WorldContext = nullptr;
+		for (const FWorldContext& Context : GEngine->GetWorldContexts())
+		{
+			if (Context.World() == World)
+			{
+				WorldContext = &Context;
+				break;
+			}
+		}
+
+		if (World->WorldType == EWorldType::PIE)
+		{
+			switch (World->GetNetMode())
+			{
+			case NM_Client:
+				if (WorldContext)
+				{
+					PostFix = FText::Format(LOCTEXT("ClientPostfixFormat", "(Client {0})"), FText::AsNumber(WorldContext->PIEInstance - 1));
+				}
+				else
+				{
+					PostFix = LOCTEXT("ClientPostfix", "(Client)");
+				}
+				break;
+			case NM_DedicatedServer:
+			case NM_ListenServer:
+				PostFix = LOCTEXT("ServerPostfix", "(Server)");
+				break;
+			case NM_Standalone:
+				PostFix = LOCTEXT("PlayInEditorPostfix", "(Play In Editor)");
+				break;
+			}
+		}
+		else if (World->WorldType == EWorldType::Editor)
+		{
+			PostFix = LOCTEXT("EditorPostfix", "(Editor)");
+		}
+
+		Description = FText::Format(LOCTEXT("WorldFormat", "{0} {1}"), FText::FromString(World->GetFName().GetPlainNameString()), PostFix);
+	}
 
+	return Description;
 }
 
-FTreeItemPtr FWorldTreeItem::FindParent(const FTreeItemMap& ExistingItems) const
-{
-	return nullptr;
-}
 
-FTreeItemPtr FWorldTreeItem::CreateParent() const
+const FTreeItemType FWorldTreeItem::Type(&ITreeItem::Type);
+
+struct SWorldTreeLabel : FCommonLabelData, public SCompoundWidget
 {
-	return nullptr;
-}
+	SLATE_BEGIN_ARGS(SWorldTreeLabel) {}
+	SLATE_END_ARGS()
+
+		void Construct(const FArguments& InArgs, FWorldTreeItem& WorldItem, ISceneOutliner& SceneOutliner, const STableRow<FTreeItemPtr>& InRow)
+	{
+		TreeItemPtr = StaticCastSharedRef<FWorldTreeItem>(WorldItem.AsShared());
+		WeakSceneOutliner = StaticCastSharedRef<ISceneOutliner>(SceneOutliner.AsShared());
+
+		ChildSlot
+			[
+				SNew(SHorizontalBox)
+
+				+ SHorizontalBox::Slot()
+			.AutoWidth()
+			.VAlign(VAlign_Center)
+			.Padding(FDefaultTreeItemMetrics::IconPadding())
+			[
+				SNew(SBox)
+				.WidthOverride(FDefaultTreeItemMetrics::IconSize())
+			.HeightOverride(FDefaultTreeItemMetrics::IconSize())
+			[
+				SNew(SImage)
+				.Image(FSlateIconFinder::FindIconBrushForClass(UWorld::StaticClass()))
+			.ToolTipText(LOCTEXT("WorldIcon_Tooltip", "World"))
+			]
+			]
+
+		+ SHorizontalBox::Slot()
+			.FillWidth(1.0f)
+			.VAlign(VAlign_Center)
+			.Padding(0.0f, 2.0f)
+			[
+				SNew(STextBlock)
+				.Text(this, &SWorldTreeLabel::GetDisplayText)
+			.HighlightText(SceneOutliner.GetFilterHighlightText())
+			.ColorAndOpacity(this, &SWorldTreeLabel::GetForegroundColor)
+			.ToolTip(IDocumentation::Get()->CreateToolTip(
+				TAttribute<FText>(this, &SWorldTreeLabel::GetTooltipText),
+				nullptr,
+				TEXT("Shared/LevelEditor/SceneOutliner"),
+				TEXT("WorldSettingsLabel")))
+			]
+			];
+	}
+
+private:
+	TWeakPtr<FWorldTreeItem> TreeItemPtr;
+
+	FText GetDisplayText() const
+	{
+		auto Item = TreeItemPtr.Pin();
+		return Item.IsValid() ? FText::FromString(Item->GetDisplayString()) : FText();
+	}
+
+	FText GetTooltipText() const
+	{
+		auto Item = TreeItemPtr.Pin();
+		FText PersistentLevelDisplayName = Item.IsValid() ? FText::FromString(Item->GetWorldName()) : FText();
+		if (Item->CanInteract())
+		{
+			return FText::Format(LOCTEXT("WorldLabel_Tooltip", "The world settings for {0}, double-click to edit"), PersistentLevelDisplayName);
+		}
+		else
+		{
+			return FText::Format(LOCTEXT("WorldLabel_TooltipNonInteractive", "The world {0}"), PersistentLevelDisplayName);
+		}
+	}
 
-void FWorldTreeItem::Visit(const ITreeItemVisitor& Visitor) const
+	FSlateColor GetForegroundColor() const
+	{
+		if (auto BaseColor = FCommonLabelData::GetForegroundColor(*TreeItemPtr.Pin()))
+		{
+			return BaseColor.GetValue();
+		}
+
+		return FSlateColor::UseForeground();
+	}
+};
+
+FWorldTreeItem::FWorldTreeItem(UWorld* InWorld)
+	: ITreeItem(Type)
+	, World(InWorld)
+	, ID(InWorld)
 {
-	Visitor.Visit(*this);
 }
 
-void FWorldTreeItem::Visit(const IMutableTreeItemVisitor& Visitor)
+FWorldTreeItem::FWorldTreeItem(TWeakObjectPtr<UWorld> InWorld)
+	: ITreeItem(Type)
+	, World(InWorld)
+	, ID(InWorld.Get())
 {
-	Visitor.Visit(*this);
 }
 
 FTreeItemID FWorldTreeItem::GetID() const
@@ -68,11 +193,6 @@ FString FWorldTreeItem::GetWorldName() const
 	return FString();
 }
 
-int32 FWorldTreeItem::GetTypeSortPriority() const
-{
-	return ETreeItemSortOrder::World;
-}
-
 bool FWorldTreeItem::CanInteract() const
 {
 	if (UWorld* WorldPtr = World.Get())
@@ -89,37 +209,13 @@ void FWorldTreeItem::GenerateContextMenu(UToolMenu* Menu, SSceneOutliner& Outlin
 	const FSlateIcon WorldSettingsIcon(FEditorStyle::GetStyleSetName(), "LevelEditor.WorldProperties.Tab");
 	const FSlateIcon NewFolderIcon(FEditorStyle::GetStyleSetName(), "SceneOutliner.NewFolderIcon");
 	FToolMenuSection& Section = Menu->AddSection("Section");
-	Section.AddMenuEntry("CreateFolder", LOCTEXT("CreateFolder", "Create Folder"), FText(), NewFolderIcon, FUIAction(FExecuteAction::CreateSP(this, &FWorldTreeItem::CreateFolder, TWeakPtr<SSceneOutliner>(SharedOutliner))));
+	Section.AddMenuEntry("CreateFolder", LOCTEXT("CreateFolder", "Create Folder"), FText(), NewFolderIcon, FUIAction(FExecuteAction::CreateSP(&Outliner, &SSceneOutliner::CreateFolder)));
 	Section.AddMenuEntry("OpenWorldSettings", LOCTEXT("OpenWorldSettings", "World Settings"), FText(), WorldSettingsIcon, FExecuteAction::CreateSP(this, &FWorldTreeItem::OpenWorldSettings));
 }
 
-void FWorldTreeItem::CreateFolder(TWeakPtr<SSceneOutliner> WeakOutliner)
-{
-	auto Outliner = WeakOutliner.Pin();
-
-	if (Outliner.IsValid() && SharedData->RepresentingWorld)
-	{
-		const FScopedTransaction Transaction(LOCTEXT("UndoAction_CreateFolder", "Create Folder"));
-
-		const FName NewFolderName = FActorFolders::Get().GetDefaultFolderName(*SharedData->RepresentingWorld, "");
-		FActorFolders::Get().CreateFolder(*SharedData->RepresentingWorld, NewFolderName);
-
-		// At this point the new folder will be in our newly added list, so select it and open a rename when it gets refreshed
-		Outliner->OnItemAdded(NewFolderName, ENewItemAction::Select | ENewItemAction::Rename);
-	}
-}
-
-FDragValidationInfo FWorldTreeItem::ValidateDrop(FDragDropPayload& DraggedObjects, UWorld& InWorld) const
-{
-	// Dropping on the world means 'moving to the root' in folder terms
-	FFolderDropTarget Target(NAME_None);
-	return Target.ValidateDrop(DraggedObjects, InWorld);
-}
-
-void FWorldTreeItem::OnDrop(FDragDropPayload& DraggedObjects, UWorld& InWorld, const FDragValidationInfo& ValidationInfo, TSharedRef<SWidget> DroppedOnWidget)
+TSharedRef<SWidget> FWorldTreeItem::GenerateLabelWidget(ISceneOutliner& Outliner, const STableRow<FTreeItemPtr>& InRow)
 {
-	FFolderDropTarget Target(NAME_None);
-	return Target.OnDrop(DraggedObjects, InWorld, ValidationInfo, DroppedOnWidget);
+	return SNew(SWorldTreeLabel, *this, Outliner, InRow);
 }
 
 void FWorldTreeItem::OpenWorldSettings() const
diff --git a/Engine/Source/Editor/SceneOutliner/Public/ActorBrowsingMode.h b/Engine/Source/Editor/SceneOutliner/Public/ActorBrowsingMode.h
new file mode 100644
--- /dev/null
+++ b/Engine/Source/Editor/SceneOutliner/Public/ActorBrowsingMode.h
@@ -0,0 +1,109 @@
+// Copyright Epic Games, Inc. All Rights Reserved.
+
+#pragma once
+
+#include "ActorModeInteractive.h"
+
+class UActorBrowsingModeSettings;
+
+namespace SceneOutliner
+{
+	class FActorBrowsingMode : public FActorModeInteractive
+	{
+	public:
+		FActorBrowsingMode(SSceneOutliner* InSceneOutliner, TWeakObjectPtr<UWorld> InSpecifiedWorldToDisplay = nullptr);
+		virtual ~FActorBrowsingMode();
+
+		/* Begin ISceneOutlinerMode Interface */
+		virtual void Rebuild() override;
+		virtual FCreateOutlinerMode CreateFolderPickerMode() const override;
+		virtual void CreateViewContent(FMenuBuilder& MenuBuilder) override;
+		virtual TSharedPtr<FDragDropOperation> CreateDragDropOperation(const TArray<FTreeItemPtr>& InTreeItems) const override;
+		virtual bool ParseDragDrop(FDragDropPayload& OutPayload, const FDragDropOperation& Operation) const override;
+		virtual FDragValidationInfo ValidateDrop(const ITreeItem& DropTarget, const FDragDropPayload& Payload) const override;
+		virtual TSharedPtr<SWidget> CreateContextMenu() override;
+		virtual void OnItemAdded(FTreeItemPtr Item) override;
+		virtual void OnItemRemoved(FTreeItemPtr Item) override;
+		virtual void OnItemSelectionChanged(FTreeItemPtr Item, ESelectInfo::Type SelectionType, const FItemSelection& Selection) override;
+		virtual void OnItemDoubleClick(FTreeItemPtr Item) override;
+		virtual void OnFilterTextCommited(FItemSelection& Selection, ETextCommit::Type CommitType) override;
+		virtual void OnItemPassesFilters(const ITreeItem& Item) override;
+		virtual FReply OnKeyDown(const FKeyEvent& InKeyEvent) override;
+		virtual void OnDuplicateSelected() override;
+		virtual void OnDrop(ITreeItem& DropTarget, const FDragDropPayload& Payload, const FDragValidationInfo& ValidationInfo) const override;
+		virtual bool CanRenameItem(const ITreeItem& Item) const override;
+		virtual FText GetStatusText() const override;
+		virtual FSlateColor GetStatusTextColor() const override;
+		virtual ESelectionMode::Type GetSelectionMode() const override { return ESelectionMode::Multi; }
+		virtual bool SupportsKeyboardFocus() const override { return true; }
+		virtual bool ShouldShowFolders() const override { return true; }
+		virtual bool SupportsCreateNewFolder() const override { return true; }
+		virtual bool ShowStatusBar() const override { return true; }
+		virtual bool ShowViewButton() const override { return true; }
+		virtual bool ShowFilterOptions() const override { return true; }
+		virtual bool CanRename() const override;
+		virtual bool CanCut() const override;
+		virtual bool CanCopy() const override;
+		virtual bool CanPaste() const override;
+		virtual bool CanSupportDragAndDrop() const { return true; }
+		virtual FName CreateNewFolder() override;
+		virtual FName CreateFolder(const FName& ParentPath, const FName& LeafName) override;
+		virtual bool ReparentItemToFolder(const FName& FolderPath, const FTreeItemPtr& Item) override;
+		virtual void SelectFoldersDescendants(const TArray<FFolderTreeItem*>& FolderItems, bool bSelectImmediateChildrenOnly) override;
+		/* End ISceneOutlinerMode Interface */
+	public:
+		/* External events this mode must respond to */
+
+		/** Called by the engine when a component is updated */
+		void OnComponentsUpdated();
+		/** Called by the engine when an actor is deleted */
+		void OnLevelActorDeleted(AActor* Actor);
+
+		/** Called by engine when edit cut actors begins */
+		void OnEditCutActorsBegin();
+
+		/** Called by engine when edit cut actors ends */
+		void OnEditCutActorsEnd();
+
+		/** Called by engine when edit copy actors begins */
+		void OnEditCopyActorsBegin();
+
+		/** Called by engine when edit copy actors ends */
+		void OnEditCopyActorsEnd();
+
+		/** Called by engine when edit paste actors begins */
+		void OnEditPasteActorsBegin();
+
+		/** Called by engine when edit paste actors ends */
+		void OnEditPasteActorsEnd();
+
+		/** Called by engine when edit duplicate actors begins */
+		void OnDuplicateActorsBegin();
+
+		/** Called by engine when edit duplicate actors ends */
+		void OnDuplicateActorsEnd();
+
+		/** Called by engine when edit delete actors begins */
+		void OnDeleteActorsBegin();
+
+		/** Called by engine when edit delete actors ends */
+		void OnDeleteActorsEnd();
+	private:
+		/** Build and up the context menu */
+		TSharedPtr<SWidget> BuildContextMenu();
+		/** Register the context menu with the engine */
+		void RegisterContextMenu();
+		bool CanPasteFoldersOnlyFromClipboard() const;
+
+		/** Filter factories */
+		static TSharedRef<FOutlinerFilter> CreateShowOnlySelectedActorsFilter();
+		static TSharedRef<FOutlinerFilter> CreateHideTemporaryActorsFilter();
+		static TSharedRef<FOutlinerFilter> CreateIsInCurrentLevelFilter();
+		static TSharedRef<FOutlinerFilter> CreateHideComponentsFilter();
+	private:
+		/** Number of actors which have passed through the filters */
+		uint32 FilteredActorCount = 0;
+		/** List of actors which passed the regular filters and may or may not have passed the search filter */
+		TSet<TWeakObjectPtr<AActor>> ApplicableActors;
+	};
+}
\ No newline at end of file
diff --git a/Engine/Source/Editor/SceneOutliner/Public/ActorFolderHierarchy.h b/Engine/Source/Editor/SceneOutliner/Public/ActorFolderHierarchy.h
new file mode 100644
--- /dev/null
+++ b/Engine/Source/Editor/SceneOutliner/Public/ActorFolderHierarchy.h
@@ -0,0 +1,31 @@
+// Copyright Epic Games, Inc. All Rights Reserved.
+
+#pragma once
+
+#include "CoreMinimal.h"
+#include "ISceneOutlinerHierarchy.h"
+
+namespace SceneOutliner
+{
+	class FActorFolderHierarchy : public ISceneOutlinerHierarchy
+	{
+	public:
+		FActorFolderHierarchy(ISceneOutlinerMode* InMode, const TWeakObjectPtr<UWorld>& World);
+		virtual ~FActorFolderHierarchy() {}
+
+
+		/* Begin ISceneOutlinerHierarchy Implementation */
+		virtual void FindChildren(const ITreeItem& Item, const TMap<FTreeItemID, FTreeItemPtr>& Items, TArray<FTreeItemPtr>& OutChildItems) const override;
+		virtual FTreeItemPtr FindParent(const ITreeItem& Item, const TMap<FTreeItemID, FTreeItemPtr>& Items) const override;
+		virtual void CreateItems(TArray<FTreeItemPtr>& OutItems) const override;
+		virtual void CreateChildren(const FTreeItemPtr& Item, TArray<FTreeItemPtr>& OutChildren) const override;
+		virtual FTreeItemPtr CreateParentItem(const FTreeItemPtr& Item) const override;
+		/* End ISceneOutlinerHierarchy Implementation */
+	private:
+		/** Adds all the direct and indirect children of a world to OutItems */
+		void CreateWorldChildren(UWorld* World, TArray<FTreeItemPtr>& OutItems) const;
+	private:
+		/** The world which this hierarchy is representing */
+		TWeakObjectPtr<UWorld> RepresentingWorld;
+	};
+}
\ No newline at end of file
diff --git a/Engine/Source/Editor/SceneOutliner/Public/ActorFolderPickingMode.h b/Engine/Source/Editor/SceneOutliner/Public/ActorFolderPickingMode.h
new file mode 100644
--- /dev/null
+++ b/Engine/Source/Editor/SceneOutliner/Public/ActorFolderPickingMode.h
@@ -0,0 +1,29 @@
+// Copyright Epic Games, Inc. All Rights Reserved.
+
+#pragma once
+
+#include "ActorMode.h"
+
+namespace SceneOutliner
+{
+
+	class FActorFolderPickingMode : public FActorMode
+	{
+	public:
+		FActorFolderPickingMode(SSceneOutliner* InSceneOutliner, FOnSceneOutlinerItemPicked InOnItemPicked, TWeakObjectPtr<UWorld> InSpecifiedWorldToDisplay = nullptr);
+		virtual ~FActorFolderPickingMode() {}
+
+		/* Begin ISceneOutlinerMode Implementation */
+		virtual void OnFilterTextCommited(FItemSelection& Selection, ETextCommit::Type CommitType) override;
+		virtual void OnItemSelectionChanged(FTreeItemPtr Item, ESelectInfo::Type SelectionType, const FItemSelection& Selection) override;
+		virtual void CreateViewContent(FMenuBuilder& MenuBuilder) override;
+		virtual bool ShowViewButton() const override { return true; }
+		virtual bool ShouldShowFolders() const { return true; }
+	protected:
+		virtual TUniquePtr<ISceneOutlinerHierarchy> CreateHierarchy() override;
+		/* End ISceneOutlinerMode Implementation */
+
+		/** Delegate to call when an item is picked */
+		FOnSceneOutlinerItemPicked OnItemPicked;
+	};
+}
\ No newline at end of file
diff --git a/Engine/Source/Editor/SceneOutliner/Public/ActorFolderTreeItem.h b/Engine/Source/Editor/SceneOutliner/Public/ActorFolderTreeItem.h
new file mode 100644
--- /dev/null
+++ b/Engine/Source/Editor/SceneOutliner/Public/ActorFolderTreeItem.h
@@ -0,0 +1,34 @@
+// Copyright Epic Games, Inc. All Rights Reserved.
+
+#pragma once
+
+#include "CoreMinimal.h"
+#include "FolderTreeItem.h"
+
+namespace SceneOutliner
+{
+	struct SCENEOUTLINER_API FActorFolderTreeItem : public FFolderTreeItem
+	{
+	public:
+		/** Static type identifier for this tree item class */
+		static const FTreeItemType Type;
+
+		FActorFolderTreeItem(FName InPath, TWeakObjectPtr<UWorld> InWorld);
+
+		/** The world which this folder belongs to */
+		TWeakObjectPtr<UWorld> World;
+
+		/* Begin ITreeItem Implementation */
+		virtual bool IsValid() const { return World.IsValid(); }
+		virtual void OnExpansionChanged() override;
+		virtual void Delete(FName InNewParentPath) override;
+		virtual TSharedRef<SWidget> GenerateLabelWidget(ISceneOutliner& Outliner, const STableRow<FTreeItemPtr>& InRow) override;
+		/* End FFolderTreeItem Implementation */
+		
+		/* Begin FFolderTreeItem Implementation */
+		virtual FName MoveTo(const FName& NewParent) override;
+	private:
+		virtual void CreateSubFolder(TWeakPtr<SSceneOutliner> WeakOutliner) override;
+		/* End FFolderTreeItem Implementation */
+	};
+}
\ No newline at end of file
diff --git a/Engine/Source/Editor/SceneOutliner/Public/ActorHierarchy.h b/Engine/Source/Editor/SceneOutliner/Public/ActorHierarchy.h
new file mode 100644
--- /dev/null
+++ b/Engine/Source/Editor/SceneOutliner/Public/ActorHierarchy.h
@@ -0,0 +1,79 @@
+// Copyright Epic Games, Inc. All Rights Reserved.
+
+#pragma once
+
+#include "CoreMinimal.h"
+#include "ISceneOutlinerHierarchy.h"
+
+namespace SceneOutliner
+{
+	class FActorHierarchy : public ISceneOutlinerHierarchy
+	{
+	public:
+		virtual ~FActorHierarchy();
+
+		static TUniquePtr<FActorHierarchy> Create(ISceneOutlinerMode* Mode, const TWeakObjectPtr<UWorld>& World, TFunction<bool(const AActor*)> InIsActorDisplayableCallback);
+
+		/** Find all direct children of a tree item in an existing item map, if any exist. */
+		virtual void FindChildren(const ITreeItem& Item, const TMap<FTreeItemID, FTreeItemPtr>& Items, TArray<FTreeItemPtr>& OutChildItems) const override;
+		/** Find the parent of a tree item in an existing item map, it if exists. */
+		virtual FTreeItemPtr FindParent(const ITreeItem& Item, const TMap<FTreeItemID, FTreeItemPtr>& Items) const override;
+
+		/** Create a linearization of all applicable items in the hierarchy */
+		virtual void CreateItems(TArray<FTreeItemPtr>& OutItems) const override;
+		/** Create a linearization of all direct and indirect children of a given item in the hierarchy */
+		virtual void CreateChildren(const FTreeItemPtr& Item, TArray<FTreeItemPtr>& OutChildren) const override;
+		/** Forcibly create a parent item for a given tree item */
+		virtual FTreeItemPtr CreateParentItem(const FTreeItemPtr& Item) const override;
+
+		void SetShowingComponents(bool bInShowingComponents) { bShowingComponents = bInShowingComponents; }
+	private:
+		/** Adds all the direct and indirect children of a world to OutItems */
+		void CreateWorldChildren(UWorld* World, TArray<FTreeItemPtr>& OutItems) const;
+		/** Create all component items for an actor if we are showing components and place them in OutItems */
+		void CreateComponentItems(const AActor* Actor, TArray<FTreeItemPtr>& OutItems) const;
+	private:
+		// Update the hierarchy when actor or world changing events occur
+
+		void OnLevelActorAdded(AActor* InActor);
+		void OnLevelActorDeleted(AActor* InActor);
+		
+		void OnLevelActorAttached(AActor* InActor, const AActor* InParent);
+		void OnLevelActorDetached(AActor* InActor, const AActor* InParent);
+
+		void OnLevelActorFolderChanged(const AActor* InActor, FName OldPath);
+
+		void OnComponentsUpdated();
+
+		void OnLevelActorListChanged();
+		
+		void OnLevelAdded(ULevel* InLevel, UWorld* InWorld);
+		void OnLevelRemoved(ULevel* InLevel, UWorld* InWorld);
+
+		/** Called when a folder is to be created */
+		void OnBroadcastFolderCreate(UWorld& InWorld, FName NewPath);
+
+		/** Called when a folder is to be moved */
+		void OnBroadcastFolderMove(UWorld& InWorld, FName OldPath, FName NewPath);
+
+		/** Called when a folder is to be deleted */
+		void OnBroadcastFolderDelete(UWorld& InWorld, FName Path);
+
+	private:
+		/** Send a an event indicating a full refresh of the hierarchy is required */
+		void FullRefreshEvent();
+
+		bool bShowingComponents = false;
+
+		TWeakObjectPtr<UWorld> RepresentingWorld;
+
+		/** Used to determine if an actor can be displayed in the outliner */
+		TFunction<bool(const AActor*)> IsActorDisplayableCallback;
+
+	private:
+		FActorHierarchy(ISceneOutlinerMode* Mode, const TWeakObjectPtr<UWorld>& Worlds, TFunction<bool(const AActor*)> InIsActorDisplayableCallback);
+
+		FActorHierarchy(const FActorHierarchy&) = delete;
+		FActorHierarchy& operator=(const FActorHierarchy&) = delete;
+	};
+}
\ No newline at end of file
diff --git a/Engine/Source/Editor/SceneOutliner/Public/ActorMode.h b/Engine/Source/Editor/SceneOutliner/Public/ActorMode.h
new file mode 100644
--- /dev/null
+++ b/Engine/Source/Editor/SceneOutliner/Public/ActorMode.h
@@ -0,0 +1,65 @@
+// Copyright Epic Games, Inc. All Rights Reserved.
+
+#pragma once
+
+#include "ISceneOutlinerMode.h"
+
+namespace SceneOutliner
+{
+	/** Functor which can be used to get weak actor pointers from a selection */
+	struct FWeakActorSelector
+	{
+		bool operator()(const TWeakPtr<ITreeItem>& Item, TWeakObjectPtr<AActor>& DataOut) const;
+	};
+
+	/** Functor which can be used to get actors from a selection including component parents */
+	struct FActorSelector
+	{
+		bool operator()(const TWeakPtr<ITreeItem>& Item, AActor*& ActorPtrOut) const;
+	};
+
+	class SCENEOUTLINER_API FActorMode : public ISceneOutlinerMode
+	{
+	public:
+		struct EItemSortOrder
+		{
+			enum Type { World = 0, Folder = 10, Actor = 20 };
+		};
+	
+		FActorMode(SSceneOutliner* InSceneOutliner, bool bHideComponents, TWeakObjectPtr<UWorld> InSpecifiedWorldToDisplay = nullptr);
+		virtual ~FActorMode();
+
+		virtual void Rebuild() override;
+
+		void BuildWorldPickerMenu(FMenuBuilder& MenuBuilder);
+
+		virtual void SynchronizeSelection() override { SynchronizeActorSelection(); }
+
+		virtual void OnFilterTextChanged(const FText& InFilterText) override;
+
+		virtual int32 GetTypeSortPriority(const ITreeItem& Item) const override;
+	private:
+		/** Called when the user selects a world in the world picker menu */
+		void OnSelectWorld(TWeakObjectPtr<UWorld> World);
+	private:
+		/* Private Helpers */
+
+		void ChooseRepresentingWorld();
+		bool IsWorldChecked(TWeakObjectPtr<UWorld> World) const;
+	protected:
+		void SynchronizeActorSelection();
+		bool IsActorDisplayable(const AActor* InActor) const;
+
+		virtual TUniquePtr<ISceneOutlinerHierarchy> CreateHierarchy() override;
+	protected:
+		// Should the hide components filter be enabled
+		bool bHideComponents;
+		/** The world which we are currently representing */
+		TWeakObjectPtr<UWorld> RepresentingWorld;
+		/** The world which the user manually selected */
+		TWeakObjectPtr<UWorld> UserChosenWorld;
+
+		/** If this mode was created to display a specific world, don't allow it to be reassigned */
+		const TWeakObjectPtr<UWorld> SpecifiedWorldToDisplay;
+	};
+}
\ No newline at end of file
diff --git a/Engine/Source/Editor/SceneOutliner/Public/ActorModeInteractive.h b/Engine/Source/Editor/SceneOutliner/Public/ActorModeInteractive.h
new file mode 100644
--- /dev/null
+++ b/Engine/Source/Editor/SceneOutliner/Public/ActorModeInteractive.h
@@ -0,0 +1,26 @@
+// Copyright Epic Games, Inc. All Rights Reserved.
+
+#pragma once
+
+#include "ActorMode.h"
+
+namespace SceneOutliner
+{
+	class FActorModeInteractive : public FActorMode
+	{
+	public:
+		FActorModeInteractive(SSceneOutliner* InSceneOutliner, bool bHideComponents, TWeakObjectPtr<UWorld> InSpecifiedWorldToDisplay = nullptr);
+		virtual ~FActorModeInteractive();
+
+		virtual bool IsInteractive() const override { return true; }
+	private:
+		/* Events */
+
+		void OnMapChange(uint32 MapFlags);
+		void OnNewCurrentLevel();
+
+		void OnLevelSelectionChanged(UObject* Obj);
+		void OnActorLabelChanged(AActor* ChangedActor);
+		void OnLevelActorRequestsRename(const AActor* Actor);
+	};
+}
diff --git a/Engine/Source/Editor/SceneOutliner/Public/ActorPickingMode.h b/Engine/Source/Editor/SceneOutliner/Public/ActorPickingMode.h
new file mode 100644
--- /dev/null
+++ b/Engine/Source/Editor/SceneOutliner/Public/ActorPickingMode.h
@@ -0,0 +1,27 @@
+// Copyright Epic Games, Inc. All Rights Reserved.
+
+#pragma once
+
+#include "ActorMode.h"
+
+namespace SceneOutliner
+{
+	class FActorPickingMode : public FActorMode
+	{
+	public:
+		FActorPickingMode(SSceneOutliner* InSceneOutliner, bool bInHideComponents, FOnSceneOutlinerItemPicked OnItemPickedDelegate, TWeakObjectPtr<UWorld> InSpecifiedWorldToDisplay = nullptr);
+
+		virtual ~FActorPickingMode() {};
+	public:
+		virtual void OnItemSelectionChanged(FTreeItemPtr Item, ESelectInfo::Type SelectionType, const FItemSelection& Selection) override;
+
+		/** Allow the user to commit their selection by pressing enter if it is valid */
+		virtual void OnFilterTextCommited(FItemSelection& Selection, ETextCommit::Type CommitType) override;
+
+		virtual void CreateViewContent(FMenuBuilder& MenuBuilder) override;
+	public:
+		virtual bool ShowViewButton() const override { return true; }
+	private:
+		FOnSceneOutlinerItemPicked OnItemPicked;
+	};
+}
\ No newline at end of file
diff --git a/Engine/Source/Editor/SceneOutliner/Public/ActorTreeItem.h b/Engine/Source/Editor/SceneOutliner/Public/ActorTreeItem.h
--- a/Engine/Source/Editor/SceneOutliner/Public/ActorTreeItem.h
+++ b/Engine/Source/Editor/SceneOutliner/Public/ActorTreeItem.h
@@ -9,85 +9,48 @@
 
 namespace SceneOutliner
 {
-	/** Helper class to manage moving abritrary data onto an actor */
-	struct FActorDropTarget : IDropTarget
+	/** A tree item that represents an actor in the world */
+	struct SCENEOUTLINER_API FActorTreeItem : ITreeItem
 	{
-		/** The actor this tree item is associated with. */
-		TWeakObjectPtr<AActor> Actor;
-		
-		/** Construct this object out of an Actor */
-		FActorDropTarget(AActor* InActor) : Actor(InActor) {}
-
 	public:
+		DECLARE_DELEGATE_RetVal_OneParam(bool, FFilterPredicate, const AActor*);
+		DECLARE_DELEGATE_RetVal_OneParam(bool, FInteractivePredicate, const AActor*);
+		
+		bool Filter(FFilterPredicate Pred) const
+		{
+			return Pred.Execute(Actor.Get());
+		}
 
-		/** Called to test whether the specified payload can be dropped onto this tree item */
-		virtual FDragValidationInfo ValidateDrop(FDragDropPayload& DraggedObjects, UWorld& World) const override;
-
-		/** Called to drop the specified objects on this item. Only called if ValidateDrop() allows. */
-		virtual void OnDrop(FDragDropPayload& DraggedObjects, UWorld& World, const FDragValidationInfo& ValidationInfo, TSharedRef<SWidget> DroppedOnWidget) override;
-
-	protected:
-
-		/** Attach our actor to the specified parent and socket */
-		static void PerformAttachment(FName SocketName, TWeakObjectPtr<AActor> Parent, const FActorArray NewAttachments);
-
-		/** Detach the specified actor from this actor */
-		void DetachActorFromParent(AActor* ChildActor);
-
-	};
+		bool GetInteractiveState(FInteractivePredicate Pred) const 
+		{
+			return Pred.Execute(Actor.Get());
+		}
 
-	/** A tree item that represents an actor in the world */
-	struct FActorTreeItem : ITreeItem
-	{
 		/** The actor this tree item is associated with. */
 		mutable TWeakObjectPtr<AActor> Actor;
 
 		/** Constant identifier for this tree item */
 		const FObjectKey ID;
 
+		/** Static type identifier for this tree item class */
+		static const FTreeItemType Type;
+
 		/** Construct this item from an actor */
 		FActorTreeItem(AActor* InActor);
 
-		/** Get this item's parent item. It is valid to return nullptr if this item has no parent */
-		virtual FTreeItemPtr FindParent(const FTreeItemMap& ExistingItems) const override;
-
-		/** Create this item's parent. It is valid to return nullptr if this item has no parent */
-		virtual FTreeItemPtr CreateParent() const override;
-
- 		/** Visit this tree item */
- 		virtual void Visit(const ITreeItemVisitor& Visitor) const override;
- 		virtual void Visit(const IMutableTreeItemVisitor& Visitor) override;
-
-	public:
-
-		/** Get the ID that represents this tree item. Used to reference this item in a map */
+		/* Begin ITreeItem Implementation */
+		virtual bool IsValid() const override { return Actor.IsValid(); }
 		virtual FTreeItemID GetID() const override;
-
-		/** Get the raw string to display for this tree item - used for sorting */
 		virtual FString GetDisplayString() const override;
-
-		/** Get the sort priority given to this item's type */
-		virtual int32 GetTypeSortPriority() const override;
-
-		/** Check whether it should be possible to interact with this tree item */
 		virtual bool CanInteract() const override;
-
+		virtual TSharedRef<SWidget> GenerateLabelWidget(ISceneOutliner& Outliner, const STableRow<FTreeItemPtr>& InRow) override;
+		virtual void OnVisibilityChanged(const bool bNewVisibility) override;
+		virtual bool HasVisibilityInfo() const override { return true; }
+		virtual bool GetVisibility() const override;
+		/* End ITreeItem Implementation */
 	public:
-
-		/** Populate the specified drag/drop payload with any relevant information for this type */
-		virtual void PopulateDragDropPayload(FDragDropPayload& Payload) const override;
-
-		/** Called to test whether the specified payload can be dropped onto this tree item */
-		virtual FDragValidationInfo ValidateDrop(FDragDropPayload& DraggedObjects, UWorld& World) const override;
-
-		/** Called to drop the specified objects on this item. Only called if ValidateDrop() allows. */
-		virtual void OnDrop(FDragDropPayload& DraggedObjects, UWorld& World, const FDragValidationInfo& ValidationInfo, TSharedRef<SWidget> DroppedOnWidget) override;
-
-	public:
-
 		/** true if this item exists in both the current world and PIE. */
 		bool bExistsInCurrentWorldAndPIE;
-
 	};
 
 }		// namespace SceneOutliner
