SceneOutliner refactor

diff --git a/Engine/Source/Editor/SceneOutliner/Public/ComponentTreeItem.h b/Engine/Source/Editor/SceneOutliner/Public/ComponentTreeItem.h
--- a/Engine/Source/Editor/SceneOutliner/Public/ComponentTreeItem.h
+++ b/Engine/Source/Editor/SceneOutliner/Public/ComponentTreeItem.h
@@ -11,81 +11,45 @@
 
 namespace SceneOutliner
 {
-	/** Helper class to manage moving arbitrary data onto an Component if any */
-	struct FComponentDropTarget : IDropTarget
+	/** A tree item that represents an Component in the world */
+	struct SCENEOUTLINER_API FComponentTreeItem : ITreeItem
 	{
-		/** The Component this tree item is associated with. */
-		TWeakObjectPtr<UActorComponent> Component;
-		
-		/** Construct this object out of an Component */
-		FComponentDropTarget(UActorComponent* InComponent) : Component(InComponent) {}
-
 	public:
+		DECLARE_DELEGATE_RetVal_OneParam(bool, FFilterPredicate, const UActorComponent*);
+		DECLARE_DELEGATE_RetVal_OneParam(bool, FInteractivePredicate, const UActorComponent*);
 
-		/** Called to test whether the specified payload can be dropped onto this tree item */
-		virtual FDragValidationInfo ValidateDrop(FDragDropPayload& DraggedObjects, UWorld& World) const override;
+		bool Filter(FFilterPredicate Pred) const
+		{
+			return Pred.Execute(Component.Get());
+		}
 
-		/** Called to drop the specified objects on this item. Only called if ValidateDrop() allows. */
-		virtual void OnDrop(FDragDropPayload& DraggedObjects, UWorld& World, const FDragValidationInfo& ValidationInfo, TSharedRef<SWidget> DroppedOnWidget) override;
-
-	protected:
-
-	};
+		bool GetInteractiveState(FInteractivePredicate Pred) const
+		{
+			return Pred.Execute(Component.Get());
+		}
 
-	/** A tree item that represents an Component in the world */
-	struct SCENEOUTLINER_API FComponentTreeItem : ITreeItem
-	{
 		/** The Component this tree item is associated with. */
 		mutable TWeakObjectPtr<UActorComponent> Component;
 
 		/** Constant identifier for this tree item */
 		const FObjectKey ID;
 
+		/** Static type identifier for this tree item class */
+		static const FTreeItemType Type;
+
 		/** Construct this item from an Component */
 		FComponentTreeItem(UActorComponent* InComponent);
 
-		/** Get this item's parent item. It is valid to return nullptr if this item has no parent */
-		virtual FTreeItemPtr FindParent(const FTreeItemMap& ExistingItems) const override;
-
-		/** Create this item's parent. It is valid to return nullptr if this item has no parent */
-		virtual FTreeItemPtr CreateParent() const override;
-
- 		/** Visit this tree item */
- 		virtual void Visit(const ITreeItemVisitor& Visitor) const override;
- 		virtual void Visit(const IMutableTreeItemVisitor& Visitor) override;
-
-		/** Called when tree item double clicked - only works in certain modes */
-		virtual void OnDoubleClick() const {}
-	public:
-
-		/** Get the ID that represents this tree item. Used to reference this item in a map */
+		/* Begin ITreeItem Implementation */
+		virtual bool IsValid() const override { return Component.IsValid(); }
 		virtual FTreeItemID GetID() const override;
-
-		/** Get the raw string to display for this tree item - used for sorting */
 		virtual FString GetDisplayString() const override;
-
-		/** Get the sort priority given to this item's type */
-		virtual int32 GetTypeSortPriority() const override;
-
-		/** Check whether it should be possible to interact with this tree item */
 		virtual bool CanInteract() const override;
-
+		virtual TSharedRef<SWidget> GenerateLabelWidget(ISceneOutliner& Outliner, const STableRow<FTreeItemPtr>& InRow) override;
+		/* End ITreeItem Implementation */
 	public:
-
-		/** Populate the specified drag/drop payload with any relevant information for this type */
-		virtual void PopulateDragDropPayload(FDragDropPayload& Payload) const override;
-
-		/** Called to test whether the specified payload can be dropped onto this tree item */
-		virtual FDragValidationInfo ValidateDrop(FDragDropPayload& DraggedObjects, UWorld& World) const override;
-
-		/** Called to drop the specified objects on this item. Only called if ValidateDrop() allows. */
-		virtual void OnDrop(FDragDropPayload& DraggedObjects, UWorld& World, const FDragValidationInfo& ValidationInfo, TSharedRef<SWidget> DroppedOnWidget) override;
-
-	public:
-
 		/** true if this item exists in both the current world and PIE. */
 		bool bExistsInCurrentWorldAndPIE;
-
 		/** Cache the string displayed */
 		FString CachedDisplayString;
 	};
diff --git a/Engine/Source/Editor/SceneOutliner/Public/FolderTreeItem.h b/Engine/Source/Editor/SceneOutliner/Public/FolderTreeItem.h
--- a/Engine/Source/Editor/SceneOutliner/Public/FolderTreeItem.h
+++ b/Engine/Source/Editor/SceneOutliner/Public/FolderTreeItem.h
@@ -10,27 +10,31 @@ class UToolMenu;
 
 namespace SceneOutliner
 {
-	/** Helper class to manage moving abritrary data onto a folder */
-	struct FFolderDropTarget : IDropTarget
+	struct FFolderPathSelector
 	{
-		/** The path that we are dropping on */
-		FName DestinationPath;
-
-		/** Constructor that takes a path to this folder (including leaf-name) */
-		FFolderDropTarget(FName InDestinationPath) : DestinationPath(InDestinationPath) {}
+		bool operator()(TWeakPtr<ITreeItem> Item, FName& DataOut) const;
+	};
 
+	/** A tree item that represents a folder in the world */
+	struct SCENEOUTLINER_API FFolderTreeItem : ITreeItem
+	{
 	public:
+		/** Static type identifier for this tree item class */
+		static const FTreeItemType Type;
 
-		/** Called to test whether the specified payload can be dropped onto this tree item */
-		virtual FDragValidationInfo ValidateDrop(FDragDropPayload& DraggedObjects, UWorld& World) const override;
+		DECLARE_DELEGATE_RetVal_OneParam(bool, FFilterPredicate, FName);
+		DECLARE_DELEGATE_RetVal_OneParam(bool, FInteractivePredicate, FName);
 
-		/** Called to drop the specified objects on this item. Only called if ValidateDrop() allows. */
-		virtual void OnDrop(FDragDropPayload& DraggedObjects, UWorld& World, const FDragValidationInfo& ValidationInfo, TSharedRef<SWidget> DroppedOnWidget) override;
-	};
+		bool Filter(FFilterPredicate Pred) const
+		{
+			return Pred.Execute(Path);
+		}
+
+		bool GetInteractiveState(FInteractivePredicate Pred) const
+		{
+			return Pred.Execute(Path);
+		}
 
-	/** A tree item that represents a folder in the world */
-	struct FFolderTreeItem : ITreeItem
-	{
 		/** The path of this folder. / separated. */
 		FName Path;
 
@@ -39,56 +43,24 @@ namespace SceneOutliner
 
 		/** Constructor that takes a path to this folder (including leaf-name) */
 		FFolderTreeItem(FName InPath);
+		/** Constructor that takes a path to this folder and a subclass tree item type (used for subclassing FFolderTreeItem) */
+		FFolderTreeItem(FName InPath, FTreeItemType Type);
 
-		/** Get this item's parent item. It is valid to return nullptr if this item has no parent */
-		virtual FTreeItemPtr FindParent(const FTreeItemMap& ExistingItems) const override;
-
-		/** Create this item's parent. It is valid to return nullptr if this item has no parent */
-		FTreeItemPtr CreateParent() const override;
-
- 		/** Visit this tree item */
- 		virtual void Visit(const ITreeItemVisitor& Visitor) const override;
- 		virtual void Visit(const IMutableTreeItemVisitor& Visitor) override;
- 		
-	public:
-
-		/** Called when the expansion changes on this folder */
-		virtual void OnExpansionChanged() override;
-
-		/** Get the ID that represents this tree item. Used to reference this item in a map */
+		/* Begin ITreeItem Implementation */
+		virtual bool IsValid() const override { return true; }
 		virtual FTreeItemID GetID() const override;
-
-		/** Get the raw string to display for this tree item - used for sorting */
 		virtual FString GetDisplayString() const override;
-
-		/** Get the sort priority given to this item's type */
-		virtual int32 GetTypeSortPriority() const override;
-		
-		/** Check whether it should be possible to interact with this tree item */
 		virtual bool CanInteract() const override;
-		
-		/** Generate a context menu for this item. Only called if *only* this item is selected. */
 		virtual void GenerateContextMenu(UToolMenu* Menu, SSceneOutliner& Outliner) override;
-
-		/** Populate the specified drag/drop payload with any relevant information for this type */
-		virtual void PopulateDragDropPayload(FDragDropPayload& Payload) const override;
-
-	public:
-
-		/** Called to test whether the specified payload can be dropped onto this tree item */
-		virtual FDragValidationInfo ValidateDrop(FDragDropPayload& DraggedObjects, UWorld& World) const override;
-
-		/** Called to drop the specified objects on this item. Only called if ValidateDrop() allows. */
-		virtual void OnDrop(FDragDropPayload& DraggedObjects, UWorld& World, const FDragValidationInfo& ValidationInfo, TSharedRef<SWidget> DroppedOnWidget) override;
-
 		/** Delete this folder, children will be reparented to provided new parent path */
-		void Delete(FName InNewParentPath);
+		virtual void Delete(FName InNewParentPath) {}
+		/* End ITreeItem Implementation */
 
+		/** Move this folder to a new parent */
+		virtual FName MoveTo(const FName& NewParent) { return FName(); }
 	private:
-
 		/** Create a new folder as a child of this one */
-		void CreateSubFolder(TWeakPtr<SSceneOutliner> WeakOutliner);
-
+		virtual void CreateSubFolder(TWeakPtr<SSceneOutliner> WeakOutliner) {}
 		/** Duplicate folder hierarchy */
 		void DuplicateHierarchy(TWeakPtr<SSceneOutliner> WeakOutliner);
 	};
diff --git a/Engine/Source/Editor/SceneOutliner/Public/ISceneOutliner.h b/Engine/Source/Editor/SceneOutliner/Public/ISceneOutliner.h
--- a/Engine/Source/Editor/SceneOutliner/Public/ISceneOutliner.h
+++ b/Engine/Source/Editor/SceneOutliner/Public/ISceneOutliner.h
@@ -44,7 +44,7 @@ public:
 	virtual void CacheIconForClass(FName InClassName, const FSlateBrush* InSlateBrush) = 0;
 
 	/** Should the scene outliner accept a request to rename a item of the tree */
-	virtual bool CanExecuteRenameRequest(const SceneOutliner::FTreeItemPtr& ItemPtr) const = 0;
+	virtual bool CanExecuteRenameRequest(const SceneOutliner::ITreeItem& ItemPtr) const = 0;
 
 	/** 
 	 * Add a filter to the scene outliner 
@@ -90,6 +90,8 @@ public:
 	/** Return the sorting mode for the specified ColumnId */
 	virtual EColumnSortMode::Type GetColumnSortMode( const FName ColumnId ) const = 0;
 
+	virtual uint32 GetTypeSortPriority(const SceneOutliner::ITreeItem& Item) const = 0;
+
 	/** Request that the tree be sorted at a convenient time */
 	virtual void RequestSort() = 0;
 
@@ -110,4 +112,15 @@ public:
 
 	/** Returns true if edit paste can be executed */
 	virtual bool Paste_CanExecute() = 0;
+
+	/** Can the scene outliner rows generated on drag event */
+	virtual bool CanSupportDragAndDrop() const = 0;
+
+	/** Set the item selection of the outliner based on a selector function. Any items which return true will be added */
+	virtual void SetSelection(const TFunctionRef<bool(SceneOutliner::ITreeItem&)> Selector) = 0;
+
+	/** Get the active SceneOutlinerMode */
+	const SceneOutliner::ISceneOutlinerMode* GetMode() const { return Mode; }
+protected:
+	SceneOutliner::ISceneOutlinerMode* Mode;
 };
diff --git a/Engine/Source/Editor/SceneOutliner/Public/ISceneOutlinerHierarchy.h b/Engine/Source/Editor/SceneOutliner/Public/ISceneOutlinerHierarchy.h
new file mode 100644
--- /dev/null
+++ b/Engine/Source/Editor/SceneOutliner/Public/ISceneOutlinerHierarchy.h
@@ -0,0 +1,38 @@
+// Copyright Epic Games, Inc. All Rights Reserved.
+
+#pragma once
+
+#include "CoreMinimal.h"
+
+#include "SceneOutlinerStandaloneTypes.h"
+#include "SceneOutlinerFwd.h"
+#include "ITreeItem.h"
+
+namespace SceneOutliner
+{
+	class ISceneOutlinerHierarchy : public TSharedFromThis<ISceneOutlinerHierarchy>
+	{
+	public:
+		ISceneOutlinerHierarchy(ISceneOutlinerMode* InMode) : Mode(InMode) {}
+		virtual ~ISceneOutlinerHierarchy() {}
+
+		/** Find all direct children of a tree item in an existing item map, if any exist. */
+		virtual void FindChildren(const ITreeItem& Item, const TMap<FTreeItemID, FTreeItemPtr>& Items, TArray<FTreeItemPtr>& OutChildItems) const = 0;
+		/** Find the parent of a tree item in an existing item map, it if exists. */
+		virtual FTreeItemPtr FindParent(const ITreeItem& Item, const TMap<FTreeItemID, FTreeItemPtr>& Items) const = 0;
+		
+		/** Create a linearization of all applicable items in the hierarchy */
+		virtual void CreateItems(TArray<FTreeItemPtr>& OutItems) const = 0;
+		/** Create a linearization of all direct and indirect children of a given item in the hierarchy */
+		virtual void CreateChildren(const FTreeItemPtr& Item, TArray<FTreeItemPtr>& OutChildren) const = 0;
+		/** Forcibly create a parent item for a given tree item */
+		virtual FTreeItemPtr CreateParentItem(const FTreeItemPtr& Item) const = 0;
+
+		DECLARE_EVENT_OneParam(ISceneOutlinerHierarchy, FHierarchyChangedEvent, FHierarchyChangedData)
+		FHierarchyChangedEvent& OnHierarchyChanged() { return HierarchyChangedEvent; }
+
+	protected:
+		ISceneOutlinerMode* Mode;
+		FHierarchyChangedEvent HierarchyChangedEvent;
+	};
+}
\ No newline at end of file
diff --git a/Engine/Source/Editor/SceneOutliner/Public/ITreeItem.h b/Engine/Source/Editor/SceneOutliner/Public/ITreeItem.h
--- a/Engine/Source/Editor/SceneOutliner/Public/ITreeItem.h
+++ b/Engine/Source/Editor/SceneOutliner/Public/ITreeItem.h
@@ -4,24 +4,15 @@
 
 #include "CoreMinimal.h"
 #include "Widgets/SWidget.h"
-#include "SceneOutlinerVisitorTypes.h"
+#include "Widgets/Views/STableRow.h"
+#include "SceneOutlinerStandaloneTypes.h"
 
 class UToolMenu;
 
 namespace SceneOutliner
 {
-	/** Interface used for validating movement (ie, drag/dropping) operations */
-	struct IDropTarget
-	{
-		/** Called to test whether the specified payload can be dropped onto this tree item */
-		virtual FDragValidationInfo ValidateDrop(FDragDropPayload& DraggedObjects, UWorld& World) const = 0;
-
-		/** Called to drop the specified objects on this item. Only called if ValidateDrop() allows. */
-		virtual void OnDrop(FDragDropPayload& DraggedObjects, UWorld& World, const FDragValidationInfo& ValidationInfo, TSharedRef<SWidget> DroppedOnWidget) = 0;
-	};
-
 	/** Base tree item interface  */
-	struct ITreeItem : IDropTarget, TSharedFromThis<ITreeItem>
+	struct ITreeItem : TSharedFromThis<ITreeItem>
 	{
 		/** Friendship required for access to various internals */
 		friend SSceneOutliner;
@@ -59,29 +50,20 @@ namespace SceneOutliner
 	protected:
 
 		/** Default constructor */
-		ITreeItem() : SharedData(nullptr), Parent(nullptr) {}
+		ITreeItem(FTreeItemType InType) : Parent(nullptr), TreeType(InType) {}
 		virtual ~ITreeItem() {}
 
-		/** Data that is common between all outliner items - owned by the Outliner itself */
-		TSharedPtr<FSharedOutlinerData> SharedData;
-
 		/** This item's parent, if any. */
 		mutable TWeakPtr<ITreeItem> Parent;
 
 		/** Array of children contained underneath this item */
 		mutable TArray<TWeakPtr<ITreeItem>> Children;
 
-		/** Array of sub-component data contained within this item */
-		mutable TArray<FTreeItemRef> SubComponentItems;
-
-	public:
-
-		const FSharedOutlinerData& GetSharedData() const
-		{
-			checkSlow(SharedData.IsValid());
-			return *SharedData;
-		}
+		/** Static type identifier for the base class tree item */
+		static const FTreeItemType Type;
 
+		/** Tree item type identifier */
+		FTreeItemType TreeType;
 	public:
 
 		/** Get this item's parent. Can be nullptr. */
@@ -93,6 +75,7 @@ namespace SceneOutliner
 		/** Add a child to this item */
 		void AddChild(FTreeItemRef Child)
 		{
+			check(!Children.Contains(Child));
 			Child->Parent = AsShared();
 			Children.Add(MoveTemp(Child));
 		}
@@ -111,32 +94,28 @@ namespace SceneOutliner
 		{
 			return Children;
 		}
-
-		/** Find this item's parent in the specified map. It is valid to return nullptr if this item has no parent */
-		virtual FTreeItemPtr FindParent(const FTreeItemMap& ExistingItems) const = 0;
-
-		/** Create this item's parent. It is valid to return nullptr if this item has no parent */
-		virtual FTreeItemPtr CreateParent() const = 0;
-
-		/** Get this item's sub-component items, if any. */
-		FORCEINLINE const TArray<FTreeItemRef>& GetSubComponentItems() const
+	public:
+		/** Attempt to cast this item to another type if it is of that type. Returns nullptr if it fails */
+		template <typename T>
+		T* CastTo()
 		{
-			return SubComponentItems;
+			return TreeType.IsA(T::Type) ? StaticCast<T*>(this) : nullptr;
+		}
+		/** Attempt to cast this item to another type if it is of that type. Returns nullptr if it fails */
+		template <typename T>
+		const T* CastTo() const
+		{
+			return TreeType.IsA(T::Type) ? StaticCast<const T*>(this) : nullptr;
+		}
+		/** Returns true if this item is of the specified type */
+		template <typename T>
+		bool IsA() const
+		{
+			return TreeType.IsA(T::Type);
 		}
 
-	public:
-
-	 	/** Visit this tree item */
-	 	virtual void Visit(const ITreeItemVisitor& Visitor) const = 0;
-	 	virtual void Visit(const IMutableTreeItemVisitor& Visitor) = 0;
-
-	 	/** Get some data from this tree item using a 'getter' visitor. */
-	 	template<typename T>
-	 	T Get(const TTreeItemGetter<T>& Getter) const
-	 	{
-	 		Visit(Getter);
-	 		return Getter.Result();
-	 	}
+		/** Returns true if the data the item references is valid */
+		virtual bool IsValid() const = 0;
 
 	public:
 
@@ -146,22 +125,25 @@ namespace SceneOutliner
 		/** Get the raw string to display for this tree item - used for sorting */
 		virtual FString GetDisplayString() const = 0;
 
-		/** Get the sort priority given to this item's type */
-		virtual int32 GetTypeSortPriority() const = 0;
-
 		/** Check whether it should be possible to interact with this tree item */
 		virtual bool CanInteract() const = 0;
 		
 		/** Called when this item is expanded or collapsed */
 		virtual void OnExpansionChanged() {};
 
+		/** Generate the label widget for this item */
+		virtual TSharedRef<SWidget> GenerateLabelWidget(ISceneOutliner& Outliner, const STableRow<FTreeItemPtr>& InRow) { return SNullWidget::NullWidget; }
+
 		/** Generate a context menu for this item. Only called if *only* this item is selected. */
 		virtual void GenerateContextMenu(UToolMenu* Menu, SSceneOutliner& Outliner) {}
 
-		/** Populate the specified drag/drop payload with any relevant information for this type */
-		virtual void PopulateDragDropPayload(FDragDropPayload& Payload) const = 0;
+		/** Called when this items visibility changed */
+		virtual void OnVisibilityChanged(const bool bNewVisibility) {}
 
-		virtual void SynchronizeSubItemSelection(TSharedPtr<class SOutlinerTreeView> OutlinerTreeView) {}
+		/** Returns true if this item can set its own visibility */
+		virtual bool HasVisibilityInfo() const { return false; }
+		/** Query this items visibility state. Only called if the item type has visibility info */
+		virtual bool GetVisibility() const { return false; }
 
 	};
 
diff --git a/Engine/Source/Editor/SceneOutliner/Public/SOutlinerTreeView.h b/Engine/Source/Editor/SceneOutliner/Public/SOutlinerTreeView.h
--- a/Engine/Source/Editor/SceneOutliner/Public/SOutlinerTreeView.h
+++ b/Engine/Source/Editor/SceneOutliner/Public/SOutlinerTreeView.h
@@ -22,14 +22,13 @@ namespace SceneOutliner
 
 		void FlashHighlightOnItem( FTreeItemPtr FlashHighlightOnItem );
 
+		const TWeakPtr<SSceneOutliner>& GetOutlinerPtr() { return SceneOutlinerWeak; }
 	protected:
 
 		virtual FReply OnDragOver( const FGeometry& MyGeometry, const FDragDropEvent& DragDropEvent ) override;
 
 		virtual void OnDragLeave( const FDragDropEvent& DragDropEvent ) override;
 
-		//virtual FReply OnDragDetected( const FGeometry& MyGeometry, const FPointerEvent& MouseEvent ) override;
-
 		virtual FReply OnDrop( const FGeometry& MyGeometry, const FDragDropEvent& DragDropEvent ) override;
 
 		/** Weak reference to the outliner widget that owns this list */
diff --git a/Engine/Source/Editor/SceneOutliner/Public/SSceneOutliner.h b/Engine/Source/Editor/SceneOutliner/Public/SSceneOutliner.h
--- a/Engine/Source/Editor/SceneOutliner/Public/SSceneOutliner.h
+++ b/Engine/Source/Editor/SceneOutliner/Public/SSceneOutliner.h
@@ -15,18 +15,24 @@
 #include "Widgets/SWidget.h"
 #include "Widgets/Views/STableRow.h"
 #include "Widgets/Views/STableViewBase.h"
+#include "Delegates/DelegateCombinations.h"
+#include "Framework/Views/ITypedTableView.h"
+#include "Types/SlateEnums.h"
 
-#include "ICustomSceneOutliner.h"
 #include "ISceneOutliner.h"
+#include "SceneOutlinerFwd.h"
+
 #include "SOutlinerTreeView.h"
 #include "SceneOutlinerPublicTypes.h"
 #include "SceneOutlinerStandaloneTypes.h"
 
+#include "ISceneOutlinerHierarchy.h"
+#include "SceneOutlinerDragDrop.h"
+
 class FMenuBuilder;
 class UToolMenu;
 class ISceneOutlinerColumn;
 class SComboButton;
-class USceneOutlinerSettings;
 
 template<typename ItemType> class STreeView;
 
@@ -36,6 +42,10 @@ template<typename ItemType> class STreeView;
  */
 namespace SceneOutliner
 {
+	DECLARE_EVENT_OneParam(SSceneOutliner, FTreeItemPtrEvent, FTreeItemPtr);
+
+	DECLARE_EVENT_TwoParams(SSceneOutliner, FOnItemSelectionChanged, FTreeItemPtr, ESelectInfo::Type);
+
 	typedef TTextFilter< const ITreeItem& > TreeItemTextFilter;
 
 	/** Structure that defines an operation that should be applied to the tree */
@@ -56,6 +66,8 @@ namespace SceneOutliner
 	{
 		enum Type
 		{
+			/** Do nothing when it is created */
+			None			= 0,
 			/** Select the item when it is created */
 			Select			= 1 << 0,
 			/** Scroll the item into view when it is created */
@@ -65,28 +77,137 @@ namespace SceneOutliner
 		};
 	}
 
-	/** Get a description of a world to display in the scene outliner */
-	FText GetWorldDescription(UWorld* World);
+	/**
+	 * Stores a set of selected items with parsing functions for the scene outliner
+	 */
+	struct FItemSelection
+	{
+		/** Set of selected items */
+		mutable TArray<TWeakPtr<ITreeItem>> SelectedItems;
+
+		FItemSelection() {}
+
+		FItemSelection(const TArray<FTreeItemPtr>& InSelectedItems)
+			: SelectedItems(InSelectedItems) {}
+
+		FItemSelection(SOutlinerTreeView& Tree)
+			: FItemSelection(Tree.GetSelectedItems()) {}
+
+		/** Returns true if the selection has an item of a specified type */
+		template <typename TreeType>
+		bool Has() const
+		{
+			for (const TWeakPtr<ITreeItem>& Item : SelectedItems)
+			{
+				if (const auto ItemPtr = Item.Pin())
+				{
+					if (ItemPtr->IsA<TreeType>())
+					{
+						return true;
+					}
+				}
+			}
+			return false;
+		}
+
+		/** Returns the total number of items in the selection */
+		uint32 Num() const
+		{
+			return SelectedItems.Num();
+		}
+
+		/** Returns the number of items of a specific type in the selection */
+		template <typename TreeType>
+		uint32 Num() const
+		{
+			uint32 Result = 0;
+			for (const TWeakPtr<ITreeItem>& Item : SelectedItems)
+			{
+				if (const auto ItemPtr = Item.Pin())
+				{
+					if (ItemPtr->IsA<TreeType>())
+					{
+						++Result;
+					}
+				}
+			}
+			return Result;
+		}
+
+		/** Add a new item to the selection */
+		void Add(FTreeItemPtr NewItem)
+		{
+			SelectedItems.Add(NewItem);
+		}
+
+		/** Get all items of a specified type */
+		template <typename TreeType>
+		void Get(TArray<TreeType*>& OutArray) const
+		{
+			for (const TWeakPtr<ITreeItem>& Item : SelectedItems)
+			{
+				if (const auto ItemPtr = Item.Pin())
+				{
+					if (TreeType* CastedItem = ItemPtr->CastTo<TreeType>())
+					{
+						OutArray.Add(CastedItem);
+					}
+				}
+			}
+		}
+
+		/** Apply a function to each item of a specified type */
+		template <typename TreeType>
+		void ForEachItem(TFunctionRef<void(TreeType&)> Func) const
+		{
+			for (const TWeakPtr<ITreeItem>& Item : SelectedItems)
+			{
+				if (const auto ItemPtr = Item.Pin())
+				{
+					if (TreeType* CastedItem = ItemPtr->CastTo<TreeType>())
+					{
+						Func(*CastedItem);
+					}
+				}
+			}
+		}
+
+		/** Use a selector to retrieve a specific data type from items in the selection. Will only add an item's data if the selector returns true for that item. */
+		template <typename DataType>
+		TArray<DataType> GetData(TFunctionRef<bool(const TWeakPtr<ITreeItem>&, DataType&)> Selector) const
+		{
+			TArray<DataType> Result;
+			for (TWeakPtr<ITreeItem>& Item : SelectedItems)
+			{
+				DataType Data;
+				if (Selector(Item, Data))
+				{
+					Result.Add(Data);
+				}
+			}
+			return Result;
+		}
+	};
 
 	/**
 	 * Scene Outliner widget
 	 */
-	class SSceneOutliner : public ICustomSceneOutliner, public FEditorUndoClient, public FGCObject
+	class SSceneOutliner : public ISceneOutliner, public FEditorUndoClient, public FGCObject
 	{
 
 	public:
 
-		SLATE_BEGIN_ARGS( SSceneOutliner ) {}
-			SLATE_ARGUMENT( FOnSceneOutlinerItemPicked, OnItemPickedDelegate )
+		SLATE_BEGIN_ARGS(SSceneOutliner)
+		{}
 		SLATE_END_ARGS()
 
-		/**
-		 * Construct this widget.  Called by the SNew() Slate macro.
-		 *
-		 * @param	InArgs		Declaration used by the SNew() macro to construct this widget
-		 * @param	InitOptions	Programmer-driven initialization options for this widget
-		 */
-		void Construct( const FArguments& InArgs, const FInitializationOptions& InitOptions );
+			/**
+			 * Construct this widget.  Called by the SNew() Slate macro.
+			 *
+			 * @param	InArgs		Declaration used by the SNew() macro to construct this widget
+			 * @param	InitOptions	Programmer-driven initialization options for this widget
+			 */
+			void Construct(const FArguments& InArgs, const FInitializationOptions& InitOptions);
 
 		/** Default constructor - initializes data that is shared between all tree items */
 		SSceneOutliner() : SharedData(MakeShareable(new FSharedOutlinerData)) {}
@@ -95,13 +216,15 @@ namespace SceneOutliner
 		~SSceneOutliner();
 
 		/** SWidget interface */
-		virtual void Tick( const FGeometry& AllottedGeometry, const double InCurrentTime, const float InDeltaTime ) override;
+		virtual void Tick(const FGeometry& AllottedGeometry, const double InCurrentTime, const float InDeltaTime) override;
 		virtual bool SupportsKeyboardFocus() const override;
-		virtual FReply OnKeyDown( const FGeometry& MyGeometry, const FKeyEvent& InKeyEvent ) override;
+		virtual FReply OnKeyDown(const FGeometry& MyGeometry, const FKeyEvent& InKeyEvent) override;
 
 		/** Sends a requests to the Scene Outliner to refresh itself the next chance it gets */
 		virtual void Refresh() override;
 
+		void RefreshSelection();
+
 		//~ Begin FEditorUndoClient Interface
 		virtual void PostUndo(bool bSuccess) override;
 		virtual void PostRedo(bool bSuccess) override { PostUndo(bSuccess); }
@@ -119,6 +242,11 @@ namespace SceneOutliner
 			return *OutlinerTreeView;
 		}
 
+		virtual const TSharedPtr< SOutlinerTreeView>& GetTreeView() const
+		{
+			return OutlinerTreeView;
+		}
+
 		/** @return Returns a string to use for highlighting results in the outliner list */
 		virtual TAttribute<FText> GetFilterHighlightText() const override;
 
@@ -131,24 +259,8 @@ namespace SceneOutliner
 		/** Sets the cached icon for this class name */
 		virtual void CacheIconForClass(FName InClassName, const FSlateBrush* InSlateBrush) override;
 
-		/** Adds a new item for the specified type and refreshes the tree, provided it matches the filter terms */
-		template<typename TreeItemType, typename DataType>
-		void ConstructItemFor(const DataType& Data)
-		{
-			// We test the filters with a temporary so we don't allocate on the heap unnecessarily
-			const TreeItemType Temporary(Data);
-			if (Filters->PassesAllFilters(Temporary) && SearchBoxFilter->PassesFilter(Temporary))
-			{
-				FTreeItemRef NewItem = MakeShareable(new TreeItemType(Data));
-				PendingOperations.Emplace(FPendingTreeOperation::Added, NewItem);
-				PendingTreeItemMap.Add(NewItem->GetID(), NewItem);
-				ConstructSubComponentItems(NewItem);
-				Refresh();
-			}
-		}
-
 		/** Should the scene outliner accept a request to rename a object */
-		virtual bool CanExecuteRenameRequest(const FTreeItemPtr& ItemPtr) const override;
+		virtual bool CanExecuteRenameRequest(const ITreeItem& ItemPtr) const override;
 
 		/**
 		 * Add a filter to the scene outliner
@@ -215,156 +327,88 @@ namespace SceneOutliner
 		/** Returns true if edit paste can be executed */
 		virtual bool Paste_CanExecute();
 
-		/** Returns true if clipboard contains folders only */
-		bool CanPasteFoldersOnlyFromClipboard();
-
 		/** Can the scene outliner rows generated on drag event */
-		bool CanSupportDragAndDrop() const;
+		virtual bool CanSupportDragAndDrop() const override;
 
 		/** Tells the scene outliner that it should do a full refresh, which will clear the entire tree and rebuild it from scratch. */
 		virtual void FullRefresh() override;
 
 	public:
-		/** Methods for the custom scene outliner interface */
+		/** Event to react to a user double click on a item */
+		FTreeItemPtrEvent& GetDoubleClickEvent() { return OnDoubleClickOnTreeEvent; }
 
 		/**
-		 * Set the selection mode of the scene outliner.
-		 * @param SelectionMode The new selection mode
+		 * Allow the system that use the scene outliner to react when it's selection is changed
+		 * Note: This event will only be broadcast on a user input.
 		 */
-		virtual ICustomSceneOutliner& SetSelectionMode(ESelectionMode::Type InSelectionMode) override;
+		FOnItemSelectionChanged& GetOnItemSelectionChanged() { return OnItemSelectionChanged; }
 
-		/**
-		 * Tell the scene outliner to use this visitor before accepting a rename request from a actor or from the prebuild column Item Label
-		 * @param CanRenameItem The visitor that will be used to validate that a item can be renamed (return true to rename)
-		 */
-		virtual ICustomSceneOutliner& SetCanRenameItem(TUniquePtr<TTreeItemGetter<bool>>&& CanRenameItem) override;
+		/** Set the item selection of the outliner based on a selector function. Any items which return true will be added */
+		virtual void SetSelection(const TFunctionRef<bool(SceneOutliner::ITreeItem&)> Selector) override;
 
-		/**
-		 * Tell the scene outliner to use this visitor to dertimine if a newly added item should be selected
-		 * @param ShouldSelectItemWhenAdded The visitor be used to select a new item (return true if the item should be selected)
-		 */
-		virtual ICustomSceneOutliner& SetShouldSelectItemWhenAdded(TUniquePtr<TTreeItemGetter<bool>>&& ShouldSelectItemWhenAdded) override;
+		/** Set the selection status of a set of items in the scene outliner */
+		void SetItemSelection(const TArray<FTreeItemPtr>& InItems, bool bSelected, ESelectInfo::Type SelectInfo = ESelectInfo::Direct);
 
-		/**
-		 * Set the behavior for when a item is dragged
-		 * Note: to avoid having to different user experience from the world outliner. The callback is only called from a left click drag.
-		 * @param Callback The function that will be called when a drag from a item row is detected
-		 */
-		virtual ICustomSceneOutliner& SetOnItemDragDetected(TUniqueFunction<FReply (const SceneOutliner::ITreeItem&)> Callback) override;
+		/** Set the selection status of a single item in the scene outliner */
+		void SetItemSelection(const FTreeItemPtr& InItem, bool bSelected, ESelectInfo::Type SelectInfo = ESelectInfo::Direct);
 
-		/**
-		 * Set the behavior for when a drag pass over a Item of the scene outliner
-		 * @param Callback The function that will be called at each update when there is a drag over a item
-		 */
-		virtual ICustomSceneOutliner& SetOnDragOverItem(TUniqueFunction<FReply (const FDragDropEvent&, const SceneOutliner::ITreeItem&)> Callback) override;
+		/** Adds a set of items to the current selection */
+		void AddToSelection(const TArray<FTreeItemPtr>& InItems, ESelectInfo::Type SelectInfo = ESelectInfo::Direct);
 
-		/**
-		 * Set the behavior for when a drag and drop is dropped on the scene outliner
-		 * @param Callback The function that will be called
-		 */
-		virtual ICustomSceneOutliner& SetOnDropOnItem(TUniqueFunction<FReply (const FDragDropEvent&, const SceneOutliner::ITreeItem&)> Callback) override;
+		/** Remove a set of items from the current selection */
+		void RemoveFromSelection(const TArray<FTreeItemPtr>& InItems);
 
+		/** Remove an item from the current selection */
+		void RemoveFromSelection(const FTreeItemPtr& InItem);
 
 		/**
-		 * Set the behavior for when a drag and drop enter the zone of a item
-		 * @param Callback The function that will be called
+		 * Returns the list of currently selected tree items
 		 */
-		virtual ICustomSceneOutliner& SetOnDragEnterItem(TUniqueFunction<void (const FDragDropEvent&, const SceneOutliner::ITreeItem&)> Callback) override;
+		virtual TArray<FTreeItemPtr> GetSelectedItems() const { return OutlinerTreeView->GetSelectedItems(); }
 
 		/**
-		 * Set the behavior for when a drag and drop leave the zone of a item
-		 * @param Callback The function that will be called
+		 * Returns the currently selected items.
 		 */
-		virtual ICustomSceneOutliner& SetOnDragLeaveItem(TUniqueFunction<void (const FDragDropEvent&, const SceneOutliner::ITreeItem&)> Callback) override;
-
-		const TUniqueFunction<FReply (const SceneOutliner::ITreeItem&)>& GetOnItemDragDetected() const;
-		const TUniqueFunction<FReply (const FDragDropEvent&, const SceneOutliner::ITreeItem&)>& GetOnDragOverItem() const;
-		const TUniqueFunction<FReply (const FDragDropEvent&, const SceneOutliner::ITreeItem&)>& GetOnDropOnItem() const;
-		const TUniqueFunction<void (const FDragDropEvent&, const SceneOutliner::ITreeItem&)>& GetOnDragEnterItem() const;
-		const TUniqueFunction<void (const FDragDropEvent&, const SceneOutliner::ITreeItem&)>& GetOnDragLeaveItem() const;
+		virtual FItemSelection GetSelection() const { return FItemSelection(*OutlinerTreeView); }
 
 		/**
-		 * Tell this scene outliner to use the shared setting or not.
-		 * The shared settings are those used by the world ouliner tab in the level editor
-		 * Note: by default it does not use the shared settings
+		 * Add a folder to the selection of the scene outliner
+		 * @param FolderName The name of the folder to add to selection
 		 */
-		virtual ICustomSceneOutliner& SetUseSharedSceneOutlinerSettings(bool bUseSharedSettings) override;
+		void AddFolderToSelection(const FName& FolderName);
 
 		/**
-		 * Is the scene outliner using the shared settings? (The same as the world outliner)
-		 * @return True if the scene outliner use the shared settings
+		 * Remove a folder from the selection of the scene outliner
+		 * @param FolderName The name of the folder to remove from the selection
 		 */
-		virtual bool IsUsingSharedSceneOutlinerSettings() const override;
-
-			/** Set the hide temporary actors filter */
-		virtual ICustomSceneOutliner& SetHideTemporaryActors(bool bHideTemporaryActors) override;
+		void RemoveFolderFromSelection(const FName& FolderName);
 
-		/** Set the show only in current level setting  */
-		virtual ICustomSceneOutliner& SetShowOnlyCurrentLevel(bool bShowOnlyCurrentLevel) override;
+		/** Deselect all selected items */
+		void ClearSelection();
 
-		/** Set the show only selected setting */
-		virtual ICustomSceneOutliner& SetShownOnlySelected(bool bShownOnlySelected) override;
+		/** Sets the next item to rename */
+		void SetPendingRenameItem(const FTreeItemPtr& InItem) { PendingRenameItem = InItem; Refresh(); }
 
-		/** Set the show actor components setting */
-		virtual ICustomSceneOutliner& SetShowActorComponents(bool bShowActorComponents) override;
+		/** Retrieve an ITreeItem by its ID if it exists in the tree */
+		FTreeItemPtr GetTreeItem(FTreeItemID, bool bIncludePending = false);
 
-		/** Event to react to a user double click on a item */
-		virtual FTreeItemPtrEvent& GetDoubleClickEvent() override { return OnDoubleClickOnTreeEvent; }
+		/** Get the outliner filter collection */
+		TSharedPtr<FOutlinerFilters>& GetFilters() { return Filters; }
 
-		/**
-		 * Allow the system that use the scene outliner to react when it's selection is changed
-		 * Note: This event will only be broadcast on a user input.
-		 */
-		virtual FOnItemSelectionChanged& GetOnItemSelectionChanged() override { return OnItemSelectionChanged; }
+		/** Create a drag drop operation */
+		TSharedPtr<FDragDropOperation> CreateDragDropOperation(const TArray<FTreeItemPtr>& InTreeItems) const;
 
-		/**
-		 * Set the selection of the scene outliner
-		 * The items that return true will be the ones selected
-		 * @param ItemSelector A visitor that will be used set the selection.
-		 */
-		virtual void SetSelection(const SceneOutliner::TTreeItemGetter<bool>& ItemSelector) override;
+		/** Parse a drag drop operation into a payload */
+		bool ParseDragDrop(FDragDropPayload& OutPayload, const FDragDropOperation& Operation) const;
 
-		/**
-		 * Add some items to selection of the scene outliner
-		 * The items that return true will be the ones added to the selection
-		 * @param ItemSelector A visitor that will be used to add some items to the selection.
-		 */
-		virtual void AddToSelection(const SceneOutliner::TTreeItemGetter<bool>& ItemSelector) override;
+		/** Validate a drag drop operation on a drop target */
+		FDragValidationInfo ValidateDrop(const ITreeItem& DropTarget, const FDragDropPayload& Payload) const;
 
-		/**
-		 * Remove some items from selection of the scene outliner
-		 * The items that return true will be the ones removed from the selection
-		 * @param ItemSelector A visitor that will be used to remove some items from the selection.
-		 */
-		virtual void RemoveFromSelection(const SceneOutliner::TTreeItemGetter<bool>& ItemSelector) override;
-
-		/**
-		 * Add a object to the selection of the scene outliner
-		 * @param Object The Object that will be added to the selection
-		 */
-		virtual void AddObjectToSelection(const UObject* Object) override;
-
-		/**
-		 * Remove a object from the selection of the scene outliner
-		 * @param Object The Object that will be removed from the selection
-		 */
-		virtual void RemoveObjectFromSelection(const UObject* Object) override;
-
-		/**
-		 * Add a folder to the selection of the scene outliner
-		 * @param FolderName The name of the folder to add to selection
-		 */
-		virtual void AddFolderToSelection(const FName& FolderName) override;
-
-		/**
-		 * Remove a folder from the selection of the scene outliner
-		 * @param FolderName The name of the folder to remove from the selection
-		 */
-		virtual void RemoveFolderFromSelection(const FName& FolderName) override;
-
-		/** Deselect all selected items */
-		virtual void ClearSelection() override;
+		/** Called when a payload is dropped onto a target */
+		void OnDropPayload(ITreeItem& DropTarget, const FDragDropPayload& Payload, const FDragValidationInfo& ValidationInfo) const;
 
+		/** Called when a payload is dragged over an item */
+		FReply OnDragOverItem(const FDragDropEvent& Event, const ITreeItem& Item) const;
 	private:
 		/** Methods that implement structural modification logic for the tree */
 
@@ -377,10 +421,13 @@ namespace SceneOutliner
 		/** Repopulates the entire tree */
 		void RepopulateEntireTree();
 
-		/** Tells the scene outliner that there was a change in the level actor list. */
-		void OnLevelActorListChanged();
+		/** Adds a single new item to the pending map and creates an add operation for it */
+		void AddPendingItem(FTreeItemPtr Item);
+
+		/** Adds a new item and all of its children to the pending items. */
+		void AddPendingItemAndChildren(FTreeItemPtr Item);
 
-		/** Attempts to add an item to the tree. Will add any parents if required. */
+		/** Attempts to add a pending item to the current tree. Will add any parents if required. */
 		bool AddItemToTree(FTreeItemRef InItem);
 
 		/** Add an item to the tree, even if it doesn't match the filter terms. Used to add parent's that would otherwise be filtered out */
@@ -398,42 +445,31 @@ namespace SceneOutliner
 		/** Called when a child has been moved in the tree hierarchy */
 		void OnItemMoved(const FTreeItemRef& Item);
 
-		void ConstructSubComponentItems(FTreeItemRef Item)
+	public:
+		// Test the filters using stack-allocated data to prevent unnecessary heap allocations
+		template <typename TreeItemType, typename TreeItemData>
+		FTreeItemPtr CreateItemFor(const TreeItemData& Data, bool bForce = false)
 		{
-			for (FTreeItemRef SubItem : Item->GetSubComponentItems())
+			const TreeItemType Temporary(Data);
+			bool bPassesFilters = Filters->PassesAllFilters(Temporary);
+			if (bPassesFilters)
 			{
-				PendingOperations.Emplace(FPendingTreeOperation::Added, SubItem);
-				PendingTreeItemMap.Add(SubItem->GetID(), SubItem);
+				check(Mode);
+				Mode->OnItemPassesFilters(Temporary);
 			}
-		}
-
-		/** Visitor that is used to validate if the item should added to the tree */
-		struct FValidateItemBeforeAddingToTree : TTreeItemGetter<bool>
-		{
-			/** Override to extract the data from specific tree item types */
-			virtual bool Get(const FActorTreeItem& ActorItem) const { return ActorItem.Actor.IsValid(); }
-			virtual bool Get(const FWorldTreeItem& WorldItem) const { return true; }
-			virtual bool Get(const FFolderTreeItem& FolderItem) const { return true; }
-			virtual bool Get(const FComponentTreeItem& ComponentFunction) const { return ComponentFunction.Component.IsValid(); }
-			virtual bool Get(const FSubComponentTreeItem& CustomFunction) const { return CustomFunction.ParentComponent.IsValid(); }
-		};
 
-		void RegisterDefaultContextMenu();
+			bPassesFilters &= SearchBoxFilter->PassesFilter(Temporary);
 
-		/** Visitor that is used to set up type-specific data after tree items are added to the tree */
-		struct FOnItemAddedToTree : IMutableTreeItemVisitor
-		{
-			SSceneOutliner& Outliner;
-			FOnItemAddedToTree(SSceneOutliner& InOutliner) : Outliner(InOutliner) {}
-
-			virtual void Visit(FActorTreeItem& Actor) const override;
-			virtual void Visit(FFolderTreeItem& Folder) const override;
-		};
-
-		/** Friendship required so the visitor can access our guts */
-		friend FOnItemAddedToTree;
+			if (bForce || bPassesFilters)
+			{
+				FTreeItemPtr Result = MakeShareable(new TreeItemType(Data));
+				Result->Flags.bIsFilteredOut = !bPassesFilters;
+				Result->Flags.bInteractive = Filters->GetInteractiveState(*Result);
+				return Result;
+			}
 
-	public:
+			return nullptr;
+		}
 
 		/** Instruct the outliner to perform an action on the specified item when it is created */
 		void OnItemAdded(const FTreeItemID& ItemID, uint8 ActionMask);
@@ -444,6 +480,26 @@ namespace SceneOutliner
 			return Columns;
 		}
 
+		bool PassesFilters(const ITreeItem& Item) const
+		{
+			return Filters->PassesAllFilters(Item);
+		}
+
+		/** @return	Returns true if the text filter is currently active */
+		bool IsTextFilterActive() const;
+
+		bool PassesTextFilter(const FTreeItemPtr& Item) const
+		{
+			return SearchBoxFilter->PassesFilter(*Item);
+		}
+
+		bool HasSelectorFocus(FTreeItemPtr Item) const
+		{
+			return OutlinerTreeView->Private_HasSelectorFocus(Item);
+		}
+
+		/** Handler for when a property changes on any item. Called by the mode */
+		void OnItemLabelChanged(FTreeItemPtr ChangedItem);
 	private:
 
 		/** Map of columns that are shown on this outliner. */
@@ -463,37 +519,22 @@ namespace SceneOutliner
 		/** Miscellaneous helper functions */
 
 		/** Scroll the specified item into view */
-		void ScrollItemIntoView(FTreeItemPtr Item);
-
-	private:
+		void ScrollItemIntoView(const FTreeItemPtr& Item);
 
-		/** Synchronize the current actor selection in the world, to the tree */
-		void SynchronizeActorSelection();
+		void SetItemExpansion(const FTreeItemPtr& Item, bool bIsExpanded);
+		
+		bool IsItemExpanded(const FTreeItemPtr& Item) const;
 
-		/** Component has has an selection change that we need to Synchronize with */
-		void OnComponentSelectionChanged(UActorComponent* Component);
-
-		/** Component has has an selection change that we need to Synchronize with */
-		void OnComponentsUpdated();
-
-		/** Check that we are reflecting a valid world */
-		bool CheckWorld() const { return SharedData->RepresentingWorld != nullptr; }
+	private:
 
 		/** Check whether we should be showing folders or not in this scene outliner */
 		bool ShouldShowFolders() const;
 
 		/** Get an array of selected folders */
-		TArray<FFolderTreeItem*> GetSelectedFolders() const;
+		void GetSelectedFolders(TArray<FFolderTreeItem*>& OutFolders) const;
 
 		/** Get an array of selected folder names */
 		TArray<FName> GetSelectedFolderNames() const;
-
-		/** Checks to see if the actor is valid for displaying in the outliner */
-		bool IsActorDisplayable( const AActor* Actor ) const;
-
-		/** @return	Returns true if the filter is currently active */
-		bool IsFilterActive() const;
-
 	private:
 		/** Tree view event bindings */
 
@@ -518,94 +559,13 @@ namespace SceneOutliner
 	private:
 		/** Level, editor and other global event hooks required to keep the outliner up to date */
 
-		/** Called by USelection::SelectionChangedEvent delegate when the level's selection changes */
-		void OnLevelSelectionChanged(UObject* Obj);
-
-		/** Called by the engine when a level is added to the world. */
-		void OnLevelAdded(ULevel* InLevel, UWorld* InWorld);
-
-		/** Called by the engine when a level is removed from the world. */
-		void OnLevelRemoved(ULevel* InLevel, UWorld* InWorld);
-
-		/** Called by the engine when an actor is added to the world. */
-		void OnLevelActorsAdded(AActor* InActor);
-
-		/** Called by the engine when an actor is remove from the world. */
-		void OnLevelActorsRemoved(AActor* InActor);
-
-		/** Called by the engine when an actor is attached in the world. */
-		void OnLevelActorsAttached(AActor* InActor, const AActor* InParent);
-
-		/** Called by the engine when an actor is dettached in the world. */
-		void OnLevelActorsDetached(AActor* InActor, const AActor* InParent);
-
-		/** Called by the engine when an actor is being requested to be renamed */
-		void OnLevelActorsRequestRename(const AActor* InActor);
-
-		/** Called by the engine when an actor's folder is changed */
-		void OnLevelActorFolderChanged(const AActor* InActor, FName OldPath);
-
-		/** Handler for when a property changes on any object */
-		void OnActorLabelChanged(AActor* ChangedActor);
+		void OnHierarchyChangedEvent(FHierarchyChangedData Event);
 
 		/** Handler for when an asset is reloaded */
 		void OnAssetReloaded(const EPackageReloadPhase InPackageReloadPhase, FPackageReloadedEvent* InPackageReloadedEvent);
 
-		/** Called when the map has changed*/
-		void OnMapChange(uint32 MapFlags);
-
-		/** Called when the current level has changed */
-		void OnNewCurrentLevel();
-
-		/** Called when a folder is to be created */
-		void OnBroadcastFolderCreate(UWorld& InWorld, FName NewPath);
-
-		/** Called when a folder is to be moved */
-		void OnBroadcastFolderMove(UWorld& InWorld, FName OldPath, FName NewPath);
-
-		/** Called when a folder is to be deleted */
-		void OnBroadcastFolderDelete(UWorld& InWorld, FName Path);
-
-
-		/**
-		 * All those function bellow are some callback to some editor signals/commands.
-		 * They are only bind to the editors delegates when the scene outliner is in Actor Browsing mode
-		 */
-
-		/** Called by engine when edit cut actors begins */
-		void OnEditCutActorsBegin();
-
-		/** Called by engine when edit cut actors ends */
-		void OnEditCutActorsEnd();
-
-		/** Called by engine when edit copy actors begins */
-		void OnEditCopyActorsBegin();
-
-		/** Called by engine when edit copy actors ends */
-		void OnEditCopyActorsEnd();
-
-		/** Called by engine when edit paste actors begins */
-		void OnEditPasteActorsBegin();
-
-		/** Called by engine when edit paste actors ends */
-		void OnEditPasteActorsEnd();
-
-		/** Called by engine when edit duplicate actors begins */
-		void OnDuplicateActorsBegin();
-
-		/** Called by engine when edit duplicate actors ends */
-		void OnDuplicateActorsEnd();
-
-		/** Called by engine when edit delete actors begins */
-		void OnDeleteActorsBegin();
-
-		/** Called by engine when edit delete actors ends */
-		void OnDeleteActorsEnd();
-
-		// End of the editor callback
-
-
-		/** Copy specified folders to clipboard, keeping current clipboard contents if they differ from previous clipboard contents (meaning actors were copied) */
+	public:
+		/** Copy specified folders to clipboard, keeping current clipboard contents if they differ from previous clipboard contents (meaning items were copied) */
 		void CopyFoldersToClipboard(const TArray<FName>& InFolders, const FString& InPrevClipboardContents);
 
 		/** Called by copy and duplicate */
@@ -614,9 +574,6 @@ namespace SceneOutliner
 		/** Called by copy and duplicate */
 		void CopyFoldersEnd();
 
-		/** Called by paste and duplicate */
-		void PasteFoldersBegin(TArray<FFolderTreeItem*> InFolders);
-
 		/** Called by paste and duplicate */
 		void PasteFoldersBegin(TArray<FName> InFolders);
 
@@ -626,7 +583,7 @@ namespace SceneOutliner
 		/** Called by cut and delete */
 		void DeleteFoldersBegin();
 
-		/** Called by cute and delete */
+		/** Called by cut and delete */
 		void DeleteFoldersEnd();
 
 		/** Get an array of folders to paste */
@@ -686,107 +643,40 @@ namespace SceneOutliner
 		ESelectionMode::Type GetSelectionMode() const;
 
 		/** @return the content for the view button */
-		TSharedRef<SWidget> GetViewButtonContent(bool bWorldPickerOnly, bool bShouldDisplayChooseWorld);
-
-		/** Build the content for the world picker submenu */
-		void BuildWorldPickerContent(FMenuBuilder& MenuBuilder);
+		TSharedRef<SWidget> GetViewButtonContent(bool bShowFilters);
 
 		/** @return the foreground color for the view button */
 		FSlateColor GetViewButtonForegroundColor() const;
 
-		/** @return the foreground color for the world picker button */
-		FSlateColor GetWorldPickerForegroundColor() const;
-
-	private:
+	public:
 
 		/** Open a context menu for this scene outliner */
 		TSharedPtr<SWidget> OnOpenContextMenu();
 
-		/** Build a context menu for right-clicking an item in the tree */
-		TSharedPtr<SWidget> BuildDefaultContextMenu();
 		void FillFoldersSubMenu(UToolMenu* Menu) const;
 		void AddMoveToFolderOutliner(UToolMenu* Menu) const;
 		void FillSelectionSubMenu(UToolMenu* Menun) const;
 		TSharedRef<TSet<FName>> GatherInvalidMoveToDestinations() const;
 
-	private:
-
 		/** Called to select descendants of the currently selected folders */
 		void SelectFoldersDescendants(bool bSelectImmediateChildrenOnly = false);
 
-		/** Move the selected items to the specified parent */
-		void MoveSelectionTo(FTreeItemRef NewParent);
-
 		/** Moves the current selection to the specified folder path */
 		void MoveSelectionTo(FName NewParent);
 
-		/** Called when the user has clicked the button to add a new folder */
-		FReply OnCreateFolderClicked();
-
 		/** Create a new folder under the specified parent name (NAME_None for root) */
 		void CreateFolder();
 
 	private:
-		/** FILTERS */
-
-		/** Synchronize the build in filter */
-		void OnSharedSettingChanged();
-
-		/** @return whether we are displaying only selected Actors */
-		virtual bool IsShowingOnlySelected() const override;
-		/** Toggles whether we are displaying only selected Actors */
-		void ToggleShowOnlySelected();
-		/** Enables/Disables whether the SelectedActorFilter is applied */
-		void ApplyShowOnlySelectedFilter(bool bShowOnlySelected);
-
-		/** @return whether we are hiding temporary Actors */
-		virtual bool IsHidingTemporaryActors() const override;
-		/** Toggles whether we are hiding temporary Actors */
-		void ToggleHideTemporaryActors();
-		/** Enables/Disables whether the HideTemporaryActorsFilter is applied */
-		void ApplyHideTemporaryActorsFilter(bool bHideTemporaryActors);
-
-		/** @return whether we are showing only Actors that are in the Current Level */
-		virtual bool IsShowingOnlyCurrentLevel() const override;
-		/** Toggles whether we are hiding Actors that aren't in the current level */
-		void ToggleShowOnlyCurrentLevel();
-		/** Enables/Disables whether the ShowOnlyActorsInCurrentLevelFilter is applied */
-		void ApplyShowOnlyCurrentLevelFilter(bool bShowOnlyActorsInCurrentLevel);
-
-		/** @return whether we are hiding Folders with hidden actors */
-		bool IsHidingFoldersContainingOnlyHiddenActors() const;
-		/** Toggles whether we are hiding Folders with hidden actors */
-		void ToggleHideFoldersContainingOnlyHiddenActors();
-
-		/** @return whether we are showing the components of the Actors */
-		virtual bool IsShowingActorComponents() const override;
-		/** Toggles whether we are showing the components of the Actors */
-		void ToggleShowActorComponents();
-		/** Enables/Disables whether the HideTemporaryActorsFilter is applied */
-		void ApplyShowActorComponentsFilter(bool bShowActorComponents);
-
-		/** When applied, only selected Actors are displayed */
-		TSharedPtr< FOutlinerFilter > SelectedActorFilter;
-
-		/** When applied, temporary and run-time actors are hidden */
-		TSharedPtr< FOutlinerFilter > HideTemporaryActorsFilter;
-
-		/** When applied, only Actors that are in the current level are displayed */
-		TSharedPtr< FOutlinerFilter > ShowOnlyActorsInCurrentLevelFilter;
-
-		/** When applied, Actor components are displayed */
-		TSharedPtr< FOutlinerFilter > ShowActorComponentsFilter;
+		/** Called when the user has clicked the button to add a new folder */
+		FReply OnCreateFolderClicked();
 
 	private:
 
 		/** Context menu opening delegate provided by the client */
 		FOnContextMenuOpening OnContextMenuOpening;
 
-		/** Callback that's fired when an item is selected while in 'picking' mode */
-		FOnSceneOutlinerItemPicked OnItemPicked;
-
-		/** Shared data required by the tree and its items */
-		TSharedRef<FSharedOutlinerData> SharedData;
+		TSharedPtr<FSharedOutlinerData> SharedData;
 
 		/** List of pending operations to be applied to the tree */
 		TArray<FPendingTreeOperation> PendingOperations;
@@ -809,9 +699,6 @@ namespace SceneOutliner
 		/** Root level tree items */
 		TArray<FTreeItemPtr> RootTreeItems;
 
-		/** A set of all actors that pass the non-text filters in the representing world */
-		TSet<TWeakObjectPtr<AActor>> ApplicableActors;
-
 		/** The button that displays view options */
 		TSharedPtr<SComboButton> ViewOptionsComboButton;
 
@@ -825,29 +712,21 @@ namespace SceneOutliner
 
 		/** Updates the expansion state of parent items after a repopulate, according to the previous state */
 		void SetParentsExpansionState(const FParentsExpansionState& ExpansionStateInfo) const;
-
-		/** Pair of functions to Hide Folders in Outliner when HideFoldersContainingHiddenActors filter is active*/
-		void HideFoldersContainingOnlyHiddenActors();
-		bool HideFoldersContainingOnlyHiddenActors(FTreeItemPtr Parent, bool bIsRoot = false);
-
 	private:
 
-		/** Number of actors that passed the search filter */
-		int32 FilteredActorCount;
-
 		/** True if the outliner needs to be repopulated at the next appropriate opportunity, usually because our
-		    actor set has changed in some way. */
+		    item set has changed in some way. */
 		uint8 bNeedsRefresh : 1;
 
 		/** true if the Scene Outliner should do a full refresh. */
 		uint8 bFullRefresh : 1;
 
+		/** true if the Scene Outliner should refresh selection */
+		uint8 bSelectionDirty : 1;
+
 		/** True if the Scene Outliner is currently responding to a level visibility change */
 		uint8 bDisableIntermediateSorting : 1;
 
-		/** true when the actor selection state in the world does not match the selection state of the tree */
-		uint8 bActorSelectionDirty : 1;
-
 		uint8 bNeedsColumRefresh : 1;
 
 		/** Reentrancy guard */
@@ -859,7 +738,7 @@ namespace SceneOutliner
 		/** The header row of the scene outliner */
 		TSharedPtr< SHeaderRow > HeaderRowWidget;
 
-		/** A collection of filters used to filter the displayed actors and folders in the scene outliner */
+		/** A collection of filters used to filter the displayed items and folders in the scene outliner */
 		TSharedPtr< FOutlinerFilters > Filters;
 
 		/** The TextFilter attached to the SearchBox widget of the Scene Outliner */
@@ -873,34 +752,15 @@ namespace SceneOutliner
 
 		TMap<FName, const FSlateBrush*> CachedIcons;
 
-
-		/** Specific for the custom mode */
-
-		/** The current selection mode of this scene outliner */
-		ESelectionMode::Type SelectionMode;
-
-		/** A optional visitor that can be use to validate if the scene outliner should let a user rename that item */
-		TUniquePtr<TTreeItemGetter<bool>> CanRenameItemVisitor;
-
-		/** A optional visitor to select new item added to tree */
-		TUniquePtr<TTreeItemGetter<bool>> ShouldSelectNewItemVisitor;
-
-		TUniqueFunction<FReply(const SceneOutliner::ITreeItem&)> OnItemDragDetected;
-		TUniqueFunction<FReply(const FDragDropEvent&, const SceneOutliner::ITreeItem&)>  OnDragOverItem;
-		TUniqueFunction<FReply(const FDragDropEvent&, const SceneOutliner::ITreeItem&)> OnDropOnItem;
-		TUniqueFunction<void(const FDragDropEvent&, const SceneOutliner::ITreeItem&)> OnDragEnterItem;
-		TUniqueFunction<void(const FDragDropEvent&, const SceneOutliner::ITreeItem&)> OnDragLeaveItem;
+		/** Maintain a count of the number of folders active in the outliner */
+		uint32 FolderCount = 0;
 
 		FTreeItemPtrEvent OnDoubleClickOnTreeEvent;
 
 		FOnItemSelectionChanged OnItemSelectionChanged;
-
-		/** Settings specifics to this scene outliner if it doesn't use the shared settings */
-		USceneOutlinerSettings* SceneOutlinerSettings = nullptr;
-
 	private:
 
-		virtual void AddReferencedObjects(FReferenceCollector& Collector) override;
+		virtual void AddReferencedObjects(FReferenceCollector& Collector) override {};
 
 		/** Functions relating to sorting */
 
@@ -922,11 +782,7 @@ namespace SceneOutliner
 		/** Sort the specified array of items based on the current sort column */
 		void SortItems(TArray<FTreeItemPtr>& Items) const;
 
-		/** Select the world we want to view */
-		void OnSelectWorld(TWeakObjectPtr<UWorld> InWorld);
-
-		/** Display a checkbox next to the world we are viewing */
-		bool IsWorldChecked(TWeakObjectPtr<UWorld> InWorld);
+		virtual uint32 GetTypeSortPriority(const ITreeItem& Item) const override;
 
 		/** Handler for recursively expanding/collapsing items */
 		void SetItemExpansionRecursive(FTreeItemPtr Model, bool bInExpansionState);
diff --git a/Engine/Source/Editor/SceneOutliner/Public/SceneOutliner.h b/Engine/Source/Editor/SceneOutliner/Public/SceneOutliner.h
--- a/Engine/Source/Editor/SceneOutliner/Public/SceneOutliner.h
+++ b/Engine/Source/Editor/SceneOutliner/Public/SceneOutliner.h
@@ -14,7 +14,6 @@ MONOLITHIC_HEADER_BOILERPLATE()
 
 #include "SceneOutlinerModule.h"
 
-#include "SceneOutlinerVisitorTypes.h"
 #include "ITreeItem.h"
 
 #include "ActorTreeItem.h"
diff --git a/Engine/Source/Editor/SceneOutliner/Public/SceneOutlinerDelegates.h b/Engine/Source/Editor/SceneOutliner/Public/SceneOutlinerDelegates.h
--- a/Engine/Source/Editor/SceneOutliner/Public/SceneOutlinerDelegates.h
+++ b/Engine/Source/Editor/SceneOutliner/Public/SceneOutlinerDelegates.h
@@ -3,7 +3,6 @@
 #pragma once
 
 #include "UnrealEdMisc.h"
-#include "SubComponentTreeItem.h"
 
 namespace SceneOutliner
 {
@@ -22,10 +21,6 @@ namespace SceneOutliner
 		/**	Broadcasts whenever the current selection changes */
 		FSimpleMulticastDelegate SelectionChanged;
 
-		/** Broadcasts whenever a SubComponentTreeItem selection changes */
-		DECLARE_MULTICAST_DELEGATE_OneParam(FOnSubComponentSelectionChanged, TArray<FSubComponentTreeItem*>&);
-		FOnSubComponentSelectionChanged OnSubComponentSelectionChanged;
-
 		/** Broadcasts whenever a Component Selection Changes */
 		DECLARE_MULTICAST_DELEGATE_OneParam(FOnComponentSelectionChanged, class UActorComponent*);
 		FOnComponentSelectionChanged OnComponentSelectionChanged;
diff --git a/Engine/Source/Editor/SceneOutliner/Public/SceneOutlinerDragDrop.h b/Engine/Source/Editor/SceneOutliner/Public/SceneOutlinerDragDrop.h
--- a/Engine/Source/Editor/SceneOutliner/Public/SceneOutlinerDragDrop.h
+++ b/Engine/Source/Editor/SceneOutliner/Public/SceneOutlinerDragDrop.h
@@ -6,78 +6,144 @@
 #include "SceneOutlinerFwd.h"
 #include "Layout/Visibility.h"
 #include "Input/DragAndDrop.h"
-#include "DragAndDrop/DecoratedDragDropOp.h"
-#include "DragAndDrop/ActorDragDropGraphEdOp.h"
+#include "DragAndDrop/CompositeDragDropOp.h"
 
 namespace SceneOutliner
 {
-	/** Consilidated drag/drop information parsed for the scene outliner */
+	/** Enum to describe the compatibility of a drag drop operation */
+	enum DropCompatibility
+	{
+		Drop_Compatible,
+		Drop_Incompatible,
+		Drop_MultipleSelection_Incompatible,
+		Drop_CompatibleAttach,
+		Drop_IncompatibleGeneric,
+		Drop_CompatibleGeneric,
+		Drop_CompatibleMultipleAttach,
+		Drop_IncompatibleMultipleAttach,
+		Drop_CompatibleDetach,
+		Drop_CompatibleMultipleDetach
+	};
+
+	/** Consilidated drag/drop with parsing functions for the scene outliner */
 	struct FDragDropPayload
 	{
 		/** Default constructor, resulting in unset contents */
-		FDragDropPayload();
+		FDragDropPayload() {}
 
 		/** Populate this payload from an array of tree items */
-		template<typename TItem>
-		FDragDropPayload(const TArray<TItem>& InDraggedItems)
+		template<typename TreeType>
+		FDragDropPayload(const TArray<TreeType>& InDraggedItems)
 		{
 			for (const auto& Item : InDraggedItems)
 			{
-				Item->PopulateDragDropPayload(*this);
+				DraggedItems.Add(Item);
 			}
 		}
 
-		/** Optional array of dragged folders */
-		TOptional<FFolderPaths> Folders;
+		/** Returns true if the payload has an item of a specified type */
+		template <typename TreeType>
+		bool Has() const
+		{
+			for (const TWeakPtr<ITreeItem>& Item : DraggedItems)
+			{
+				if (const auto ItemPtr = Item.Pin())
+				{
+					if (ItemPtr->IsA<TreeType>())
+					{
+						return true;
+					}
+				}
+			}
+			return false;
+		}
 
-		/** Optional array of dragged actors */
-		TOptional<FActorArray> Actors;
+		/** Return an array of all tree items in the payload which are of a specified type */
+		template <typename TreeType>
+		TArray<TreeType*> Get() const
+		{
+			TArray<TreeType*> Result;
+			for (const TWeakPtr<ITreeItem>& Item : DraggedItems)
+			{
+				if (const auto ItemPtr = Item.Pin())
+				{
+					if (TreeType* CastedItem = ItemPtr->CastTo<TreeType>())
+					{
+						Result.Add(CastedItem);
+					}
+				}
+			}
+			return Result;
+		}
 
-		/** Optional array of dragged sub-component Items */
-		TOptional<FSubComponentItemArray> SubComponents;
+		/** Apply a function to each item in the payload */
+		template <typename TreeType>
+		void ForEachItem(TFunctionRef<void(TreeType&)> Func) const
+		{
+			for (const TWeakPtr<ITreeItem>& Item : DraggedItems)
+			{
+				if (const auto ItemPtr = Item.Pin())
+				{
+					if (TreeType* CastedItem = ItemPtr->CastTo<TreeType>())
+					{
+						Func(*CastedItem);
+					}
+				}
+			}
+		}
+		
+		/** Use a selector to retrieve an array of a specific data type from the items in the payload */
+		template <typename DataType>
+		TArray<DataType> GetData(TFunctionRef<bool(TWeakPtr<ITreeItem>, DataType&)> Selector) const
+		{
+			TArray<DataType> Result;
+			for (TWeakPtr<ITreeItem>& Item : DraggedItems)
+			{
+				DataType Data;
+				if (Selector(Item, Data))
+				{
+					Result.Add(Data);
+				}
+			}
+			return Result;
+		}
 
-		/**
-		 *	Parse a drag operation into our list of actors and folders
-		 *	@return true if the operation is viable for the sceneoutliner to process, false otherwise
-		 */
-		bool ParseDrag(const FDragDropOperation& Operation);
+		/** List of all dragged items */
+		mutable TArray<TWeakPtr<ITreeItem>> DraggedItems;
 	};
-	
-	/** Construct a new Drag and drop operation for a scene outliner */
-	TSharedPtr<FDragDropOperation> CreateDragDropOperation(const TArray<FTreeItemPtr>& InTreeItems);
 
 	/** Struct used for validation of a drag/drop operation in the scene outliner */
 	struct FDragValidationInfo
 	{
 		/** The tooltip type to display on the operation */
-		FActorDragDropGraphEdOp::ToolTipTextType TooltipType;
+		DropCompatibility CompatibilityType;
 
 		/** The tooltip text to display on the operation */
 		FText ValidationText;
 
 		/** Construct this validation information out of a tootip type and some text */
-		FDragValidationInfo(const FActorDragDropGraphEdOp::ToolTipTextType InTooltipType, const FText InValidationText)
-			: TooltipType(InTooltipType)
+		FDragValidationInfo(const DropCompatibility InCompatibilityType, const FText InValidationText)
+			: CompatibilityType(InCompatibilityType)
 			, ValidationText(InValidationText)
 		{}
 
 		/** Return a generic invalid result */
 		static FDragValidationInfo Invalid()
 		{
-			return FDragValidationInfo(FActorDragDropGraphEdOp::ToolTip_IncompatibleGeneric, FText());
+			return FDragValidationInfo(DropCompatibility::Drop_IncompatibleGeneric, FText());
 		}
 		
 		/** @return true if this operation is valid, false otheriwse */ 
 		bool IsValid() const
 		{
-			switch(TooltipType)
+			switch(CompatibilityType)
 			{
-			case FActorDragDropGraphEdOp::ToolTip_Compatible:
-			case FActorDragDropGraphEdOp::ToolTip_CompatibleAttach:
-			case FActorDragDropGraphEdOp::ToolTip_CompatibleGeneric:
-			case FActorDragDropGraphEdOp::ToolTip_CompatibleMultipleAttach:
-			case FActorDragDropGraphEdOp::ToolTip_CompatibleDetach:
-			case FActorDragDropGraphEdOp::ToolTip_CompatibleMultipleDetach:
+			case DropCompatibility::Drop_Compatible:
+			case DropCompatibility::Drop_CompatibleAttach:
+			case DropCompatibility::Drop_CompatibleGeneric:
+			case DropCompatibility::Drop_CompatibleMultipleAttach:
+			case DropCompatibility::Drop_CompatibleDetach:
+			case DropCompatibility::Drop_CompatibleMultipleDetach:
 				return true;
 			default:
 				return false;
@@ -96,35 +162,15 @@ namespace SceneOutliner
 		void Init(FFolderPaths InFolders);
 	};
 
-	/* A drag/drop operation when dragging sub-component items in the scene outliner */
-	struct FSubComponentDragDropOp: public FDecoratedDragDropOp
-	{
-		DRAG_DROP_OPERATOR_TYPE(FSubComponentDragDropOp, FDecoratedDragDropOp)
-
-			/** Actor that we are dragging */
-			FSubComponentItemArray Items;
-
-		void Init(const FSubComponentItemArray& InItems);
-	};
-
 	/** A drag/drop operation that was started from the scene outliner */
-	struct FSceneOutlinerDragDropOp : public FDragDropOperation
+	struct FSceneOutlinerDragDropOp : public FCompositeDragDropOp
 	{
-		DRAG_DROP_OPERATOR_TYPE(FSceneOutlinerDragDropOp, FDragDropOperation);
-
-		FSceneOutlinerDragDropOp(const FDragDropPayload& DraggedObjects);
+		DRAG_DROP_OPERATOR_TYPE(FSceneOutlinerDragDropOp, FDecoratedDragDropOp);
+		
+		FSceneOutlinerDragDropOp();
 
 		using FDragDropOperation::Construct;
 
-		/** Actor drag operation */
-		TSharedPtr<FActorDragDropOp>	ActorOp;
-
-		/** Folder drag operation */
-		TSharedPtr<FFolderDragDropOp>	FolderOp;
-
-		/** Sub-Component drag operation */
-		TSharedPtr<FSubComponentDragDropOp>	SubComponentOp;
-
 		void ResetTooltip()
 		{
 			OverrideText = FText();
diff --git a/Engine/Source/Editor/SceneOutliner/Public/SceneOutlinerFilters.h b/Engine/Source/Editor/SceneOutliner/Public/SceneOutlinerFilters.h
--- a/Engine/Source/Editor/SceneOutliner/Public/SceneOutlinerFilters.h
+++ b/Engine/Source/Editor/SceneOutliner/Public/SceneOutlinerFilters.h
@@ -5,18 +5,15 @@
 #include "CoreMinimal.h"
 #include "Misc/IFilter.h"
 #include "Misc/FilterCollection.h"
-
-#include "FolderTreeItem.h"
-#include "ActorTreeItem.h"
-#include "WorldTreeItem.h"
-#include "ComponentTreeItem.h"
-#include "SubComponentTreeItem.h"
+#include "SceneOutlinerFwd.h"
 
 class FMenuBuilder;
 
 namespace SceneOutliner
 {
-
+	/**
+	 * Contains information used to create a filter which will be displayed as user toggleable filter
+	 */
 	class FOutlinerFilterInfo
 	{ 
 	public:
@@ -33,10 +30,14 @@ namespace SceneOutliner
 		/** Add menu for this filter */
 		void AddMenu(FMenuBuilder& InMenuBuilder);
 
+		bool IsFilterActive() const;
+
+		DECLARE_EVENT_OneParam(FOutlinerFilterInfo, FOnToggle, bool);
+		FOnToggle& OnToggle() { return OnToggleEvent; }
+
 	private:
 		void ApplyFilter(bool bActive);
 		void ToggleFilterActive();
-		bool IsFilterActive() const;
 
 		TWeakPtr<FOutlinerFilters> Filters;
 
@@ -46,26 +47,22 @@ namespace SceneOutliner
 		FText FilterTooltip;
 		bool bActive;
 
+		FOnToggle OnToggleEvent;
+
 		FCreateSceneOutlinerFilter Factory;
 	};
 
 	/** Enum to specify how items that are not explicitly handled by this filter should be managed */
 	enum class EDefaultFilterBehaviour : uint8 { Pass, Fail };
 
-	/** Enum that defines how a tree item should be dealt with in the case where it appears in the tree, but doesn't match the filter (eg if it has a matching child) */
-	enum class EFailedFilterState : uint8 { Interactive, NonInteractive };
-
 	/** A filter that can be applied to any type in the tree */
-	class FOutlinerFilter : public ITreeItemVisitor, public IFilter<const ITreeItem&>
+	class FOutlinerFilter : public IFilter<const ITreeItem&>
 	{
 	public:
 
 		/** Event that is fired if this filter changes */
 		DECLARE_DERIVED_EVENT(FOutlinerFilter, IFilter<const ITreeItem&>::FChangedEvent, FChangedEvent);
 		virtual FChangedEvent& OnChanged() override { return ChangedEvent; }
-		
-		/** Enum that defines how a tree item should be dealt with in the case where it appears in the tree, but doesn't match the filter (eg if it has a matching child) */
-		EFailedFilterState FailedItemState;
 
 	protected:
 
@@ -76,135 +73,105 @@ namespace SceneOutliner
 		const EDefaultFilterBehaviour DefaultBehaviour;
 
 		/** Constructor to specify the default result of a filter */
-		FOutlinerFilter(EDefaultFilterBehaviour InDefaultBehaviour, EFailedFilterState InFailedFilterState = EFailedFilterState::NonInteractive)
-			: FailedItemState(InFailedFilterState), DefaultBehaviour(InDefaultBehaviour)
+		FOutlinerFilter(EDefaultFilterBehaviour InDefaultBehaviour)
+			: DefaultBehaviour(InDefaultBehaviour)
 		{}
 
-		/** Overridden in derived types to filter actors */
-		virtual bool PassesFilter(const AActor* Actor) const { return DefaultBehaviour == EDefaultFilterBehaviour::Pass; }
-
-		/** Overridden in derived types to filter actor components */
-		virtual bool PassesFilter(const UActorComponent* ActorComponent) const { return DefaultBehaviour == EDefaultFilterBehaviour::Pass; }
-
-		/** Overridden in derived types to filter worlds */
-		virtual bool PassesFilter(const UWorld* World) const { return DefaultBehaviour == EDefaultFilterBehaviour::Pass; }
-
-		/** Overridden in derived types to filter folders */
-		virtual bool PassesFilter(FName Folder) const { return DefaultBehaviour == EDefaultFilterBehaviour::Pass; }
-
 	private:
-
 		/** Transient result from the filter operation. Only valid until the next invocation of the filter. */
 		mutable bool bTransientFilterResult;
 
-		virtual void Visit(const FActorTreeItem& ActorItem) const override
+		/** Check whether the specified item passes our filter */
+		virtual bool PassesFilter(const ITreeItem& InItem) const override
 		{
-			if (const AActor* Actor = ActorItem.Actor.Get())
-			{
-				bTransientFilterResult = PassesFilter(Actor);
-			}
-			else
-			{
-				bTransientFilterResult = false;
-			}
+			bTransientFilterResult = (DefaultBehaviour == EDefaultFilterBehaviour::Pass);
+			return bTransientFilterResult;
+		}
+	public:
+		/** 
+		 * Check if an item should be interactive according to this filter.
+		 * Default behavior just checks if it passes the filter or not.
+		 */
+		virtual bool GetInteractiveState(const ITreeItem& InItem) const
+		{
+			return PassesFilter(InItem);
 		}
+	};
+
+	/** Outliner filter which will be applied on items which match the specified type */
+	template <class T>
+	struct TOutlinerFilter : public FOutlinerFilter
+	{
+		TOutlinerFilter(EDefaultFilterBehaviour InDefaultBehaviour)
+			: FOutlinerFilter(InDefaultBehaviour)
+		{}
 
-		virtual void Visit(const FComponentTreeItem& ComponentItem) const override
+		virtual bool PassesFilter(const ITreeItem& InItem) const override
 		{
-			if (const UActorComponent* ActorComponent = ComponentItem.Component.Get())
-			{
-				bTransientFilterResult = PassesFilter(ActorComponent);
-			}
-			else
+			if (const T* CastedItem = InItem.CastTo<T>())
 			{
-				bTransientFilterResult = false;
+				return PassesFilterImpl(*CastedItem);
 			}
+			return DefaultBehaviour == EDefaultFilterBehaviour::Pass;
 		}
 
-		virtual void Visit(const FWorldTreeItem& WorldItem) const override
+		virtual bool PassesFilterImpl(const T& InItem) const
 		{
-			if (const UWorld* World = WorldItem.World.Get())
-			{
-				bTransientFilterResult = PassesFilter(World);
-			}
-			else
-			{
-				bTransientFilterResult = false;
-			}
+			return DefaultBehaviour == EDefaultFilterBehaviour::Pass;
 		}
 
-		virtual void Visit(const FFolderTreeItem& FolderItem) const override
+		virtual bool GetInteractiveState(const ITreeItem& InItem) const
 		{
-			bTransientFilterResult = PassesFilter(FolderItem.Path);
+			if (const T* CastedItem = InItem.CastTo<T>())
+			{
+				return GetInteractiveStateImpl(*CastedItem);
+			}
+			return DefaultBehaviour == EDefaultFilterBehaviour::Pass;
 		}
 
-		/** Check whether the specified item passes our filter */
-		virtual bool PassesFilter( const ITreeItem& InItem ) const override
+		// If not overriden will just default to testing against PassesFilter
+		virtual bool GetInteractiveStateImpl(const T& InItem) const
 		{
-			InItem.Visit(*this);
-			return bTransientFilterResult;
+			return PassesFilterImpl(InItem);
 		}
 	};
 
-	DECLARE_DELEGATE_RetVal_OneParam( bool, FActorFilterPredicate, const AActor* );
-	DECLARE_DELEGATE_RetVal_OneParam( bool, FComponentFilterPredicate, const UActorComponent* );
-	DECLARE_DELEGATE_RetVal_OneParam( bool, FWorldFilterPredicate, const UWorld* );
-	DECLARE_DELEGATE_RetVal_OneParam( bool, FFolderFilterPredicate, FName );
 
 	/** Predicate based filter for the outliner */
-	struct FOutlinerPredicateFilter : public FOutlinerFilter
+	template <class T>
+	struct TOutlinerPredicateFilter : public TOutlinerFilter<T>
 	{
-		/** Predicate used to filter actors */
-		mutable FActorFilterPredicate		ActorPred;
-		/** Predicate used to filter actor components */
-		mutable FComponentFilterPredicate	ComponentPred;
-		/** Predicate used to filter worlds */
-		mutable FWorldFilterPredicate		WorldPred;
-		/** Predicate used to filter Folders */
-		mutable FFolderFilterPredicate		FolderPred;
-
-		FOutlinerPredicateFilter(FActorFilterPredicate InActorPred, EDefaultFilterBehaviour InDefaultBehaviour, EFailedFilterState InFailedFilterState = EFailedFilterState::NonInteractive)
-			: FOutlinerFilter(InDefaultBehaviour, InFailedFilterState)
-			, ActorPred(InActorPred)
-		{}
-
-		FOutlinerPredicateFilter(FComponentFilterPredicate InComponentPred, EDefaultFilterBehaviour InDefaultBehaviour, EFailedFilterState InFailedFilterState = EFailedFilterState::NonInteractive)
-			: FOutlinerFilter(InDefaultBehaviour, InFailedFilterState)
-			, ComponentPred(InComponentPred)
-		{}
+		using TFilterPredicate = typename T::FFilterPredicate;
+		using TInteractivePredicate = typename T::FInteractivePredicate;
 
-		FOutlinerPredicateFilter(FWorldFilterPredicate InWorldPred, EDefaultFilterBehaviour InDefaultBehaviour, EFailedFilterState InFailedFilterState = EFailedFilterState::NonInteractive)
-			: FOutlinerFilter(InDefaultBehaviour, InFailedFilterState)
-			, WorldPred(InWorldPred)
-		{}
+		/** Predicate used to filter tree items */
+		mutable TFilterPredicate FilterPred;
+		mutable TInteractivePredicate InteractivePred;
 
-		FOutlinerPredicateFilter(FFolderFilterPredicate InFolderPred, EDefaultFilterBehaviour InDefaultBehaviour, EFailedFilterState InFailedFilterState = EFailedFilterState::NonInteractive)
-			: FOutlinerFilter(InDefaultBehaviour, InFailedFilterState)
-			, FolderPred(InFolderPred)
+		TOutlinerPredicateFilter(TFilterPredicate InFilterPred, EDefaultFilterBehaviour InDefaultBehaviour, TInteractivePredicate InInteractivePredicate = TInteractivePredicate())
+			: TOutlinerFilter<T>(InDefaultBehaviour)
+			, FilterPred(InFilterPred)
+			, InteractivePred(InInteractivePredicate)
 		{}
 
-		virtual bool PassesFilter(const AActor* Actor) const override
-		{
-			return ActorPred.IsBound() ? ActorPred.Execute(Actor) : DefaultBehaviour == EDefaultFilterBehaviour::Pass;
-		}
-
-		virtual bool PassesFilter(const UActorComponent* Component) const override
+		virtual bool PassesFilterImpl(const T& InItem) const override
 		{
-			return ComponentPred.IsBound() ? ComponentPred.Execute(Component) : DefaultBehaviour == EDefaultFilterBehaviour::Pass;
+			return InItem.Filter(FilterPred);
 		}
 
-		virtual bool PassesFilter(const UWorld* World) const override
+		virtual bool GetInteractiveStateImpl(const T& InItem) const override 
 		{
-			return WorldPred.IsBound() ? WorldPred.Execute(World) : DefaultBehaviour == EDefaultFilterBehaviour::Pass;
-		}
+			if (InteractivePred.IsBound())
+			{
+				return InItem.GetInteractiveState(InteractivePred);
+			}
 
-		virtual bool PassesFilter(FName Folder) const override
-		{
-			return FolderPred.IsBound() ? FolderPred.Execute(Folder) : DefaultBehaviour == EDefaultFilterBehaviour::Pass;
+			// If not interactive state impl is provided, default to interactive if filter passes
+			return PassesFilterImpl(InItem);
 		}
 	};
 
-	/** Scene outliner filter class. This class abstracts the filtering of both actors and folders and allows for filtering on both types */
+	/** Scene outliner filters class. This class wraps a collection of filters and allows items of any type to be tested against the entire set. */
 	struct FOutlinerFilters : public TFilterCollection<const ITreeItem&>
 	{
 		/** Overridden to ensure we only ever have FOutlinerFilters added */
@@ -214,35 +181,24 @@ namespace SceneOutliner
 		}
 
 		/** Test whether this tree item passes all filters, and set its interactive state according to the filter it failed (if applicable) */
-		bool TestAndSetInteractiveState(ITreeItem& InItem) const
+		bool GetInteractiveState(const ITreeItem& InItem) const
 		{
-			bool bPassed = true;
-
-			// Default to interactive
-			InItem.Flags.bInteractive = true;
-			
 			for (const auto& Filter : ChildFilters)
 			{
-				if (!Filter->PassesFilter(InItem))
+				if (!StaticCastSharedPtr<FOutlinerFilter>(Filter)->GetInteractiveState(InItem))
 				{
-					bPassed = false;
-					InItem.Flags.bInteractive = StaticCastSharedPtr<FOutlinerFilter>(Filter)->FailedItemState == EFailedFilterState::Interactive;
-					// If this has failed, but is still interactive, we carry on to see if any others fail *and* set to non-interactive
-					if (!InItem.Flags.bInteractive)
-					{
-						return false;
-					}
+					return false;
 				}
 			}
 
-			return bPassed;
+			return true;
 		}
 
 		/** Add a filter predicate to this filter collection */
-		template<typename T>
-		void AddFilterPredicate(T Predicate, EDefaultFilterBehaviour InDefaultBehaviour = EDefaultFilterBehaviour::Fail, EFailedFilterState InFailedFilterState = EFailedFilterState::NonInteractive)
+		template <typename T>
+		void AddFilterPredicate(typename T::FFilterPredicate InFilterPred, EDefaultFilterBehaviour InDefaultBehaviour = EDefaultFilterBehaviour::Pass, typename T::FInteractivePredicate InInteractivePred = T::FInteractivePredicate())
 		{
-			Add(MakeShareable(new FOutlinerPredicateFilter(Predicate, InDefaultBehaviour, InFailedFilterState)));
+			Add(MakeShareable(new TOutlinerPredicateFilter<T>(InFilterPred, InDefaultBehaviour, InInteractivePred)));
 		}
 	};
 }
diff --git a/Engine/Source/Editor/SceneOutliner/Public/SceneOutlinerFwd.h b/Engine/Source/Editor/SceneOutliner/Public/SceneOutlinerFwd.h
--- a/Engine/Source/Editor/SceneOutliner/Public/SceneOutlinerFwd.h
+++ b/Engine/Source/Editor/SceneOutliner/Public/SceneOutlinerFwd.h
@@ -4,7 +4,6 @@
 
 #include "CoreMinimal.h"
 #include "Misc/CoreDelegates.h"
-#include "GameFramework/Actor.h"
 
 class ISceneOutliner;
 class ISceneOutlinerColumn;
@@ -21,16 +20,15 @@ namespace SceneOutliner
 	struct FWorldTreeItem;
 	struct FFolderTreeItem;
 	struct FComponentTreeItem;
-	struct FSubComponentTreeItem;
-
-	struct ITreeItemVisitor;
-	struct IMutableTreeItemVisitor;
 	
 	typedef TSharedPtr<ITreeItem> FTreeItemPtr;
 	typedef TSharedRef<ITreeItem> FTreeItemRef;
 
 	typedef TMap<FTreeItemID, FTreeItemPtr> FTreeItemMap;
 
+	class ISceneOutlinerHierarchy;
+	class ISceneOutlinerMode;
+
 	class SSceneOutliner;
 
 	class FOutlinerFilter;
@@ -39,23 +37,13 @@ namespace SceneOutliner
 	struct FDragDropPayload;
 	struct FDragValidationInfo;
 
-	/** Typedef to define an array of actors, used during dragging */
-	typedef TArray<TWeakObjectPtr<AActor>> FActorArray;
-
 	/** Typedef to define an array of folder names, used during dragging */
 	typedef TArray<FName> FFolderPaths;
-
-	/** Typedef to define an array of actors, used during dragging */
-	typedef TArray<TWeakPtr<const FSubComponentTreeItem>> FSubComponentItemArray;
-
 }
 
-/** Delegate used with the Scene Outliner in 'actor picking' mode.  You'll bind a delegate when the
-    outliner widget is created, which will be fired off when an actor is selected in the list */
-DECLARE_DELEGATE_OneParam( FOnActorPicked, AActor* );
 DECLARE_DELEGATE_OneParam( FOnSceneOutlinerItemPicked, TSharedRef<SceneOutliner::ITreeItem> );
 
-DECLARE_DELEGATE_OneParam( FCustomSceneOutlinerDeleteDelegate, const TArray< TWeakObjectPtr< AActor > >&  )
+DECLARE_DELEGATE_OneParam( FCustomSceneOutlinerDeleteDelegate, const TArray<TWeakPtr<SceneOutliner::ITreeItem>>&  )
 
 /** A delegate used to factory a new column type */
 DECLARE_DELEGATE_RetVal_OneParam( TSharedRef< ISceneOutlinerColumn >, FCreateSceneOutlinerColumn, ISceneOutliner& );
diff --git a/Engine/Source/Editor/SceneOutliner/Public/SceneOutlinerModule.h b/Engine/Source/Editor/SceneOutliner/Public/SceneOutlinerModule.h
--- a/Engine/Source/Editor/SceneOutliner/Public/SceneOutlinerModule.h
+++ b/Engine/Source/Editor/SceneOutliner/Public/SceneOutlinerModule.h
@@ -9,7 +9,13 @@
 
 class ICustomSceneOutliner;
 class ISceneOutlinerColumn;
-class FOutlinerFilter;
+
+/** Delegate used with the Scene Outliner in 'actor picking' mode.  You'll bind a delegate when the
+	outliner widget is created, which will be fired off when an actor is selected in the list */
+DECLARE_DELEGATE_OneParam(FOnActorPicked, AActor*);
+/** Delegate used with the Scene Outliner in 'component picking' mode.  You'll bind a delegate when the
+	outliner widget is created, which will be fired off when an actor is selected in the list */
+DECLARE_DELEGATE_OneParam(FOnComponentPicked, UActorComponent*);
 
 /**
  * Implements the Scene Outliner module.
@@ -23,37 +29,41 @@ public:
 	 * Creates a scene outliner widget
 	 *
 	 * @param	InitOptions						Programmer-driven configuration for this widget instance
-	 * @param	MakeContentMenuWidgetDelegate	Optional delegate to execute to build a context menu when right clicking on actors
-	 * @param	OnActorPickedDelegate			Optional callback when an actor is selected in 'actor picking' mode
+	 * @param	OutlinerModeFactory				Factory delegate used to create the outliner mode
 	 *
 	 * @return	New scene outliner widget
 	 */
-	virtual TSharedRef< ISceneOutliner > CreateSceneOutliner(
-		const SceneOutliner::FInitializationOptions& InitOptions,
-		const FOnActorPicked& OnActorPickedDelegate ) const;
+	virtual TSharedRef<ISceneOutliner> CreateSceneOutliner(
+		const SceneOutliner::FInitializationOptions& InitOptions) const;
 
 	/**
 	 * Creates a scene outliner widget
 	 *
 	 * @param	InitOptions						Programmer-driven configuration for this widget instance
-	 * @param	MakeContentMenuWidgetDelegate	Optional delegate to execute to build a context menu when right clicking on actors
-	 * @param	OnItemPickedDelegate			Optional callback when an item is selected in 'picking' mode
+	 * @param	OnActorPickedDelegate			Optional callback when an actor is selected in 'actor picking' mode
 	 *
 	 * @return	New scene outliner widget
 	 */
-	virtual TSharedRef< ISceneOutliner > CreateSceneOutliner(
-		const SceneOutliner::FInitializationOptions& InitOptions,
-		const FOnSceneOutlinerItemPicked& OnItemPickedDelegate ) const;
 
-	/**
-	 * Creates a custom scene outliner widget
-	 *
-	 * @param	InitOptions						Programmer-driven configuration for this widget instance. Know that the mode will always be converted to custom.
-	 *
-	 * @return	New custom scene outliner widget
-	 */
-	virtual TSharedRef< ICustomSceneOutliner > CreateCustomSceneOutliner(
-		SceneOutliner::FInitializationOptions& InitOptions ) const;
+
+	/* Some common scene outliners */
+
+	/** Creates an actor picker widget. Calls the OnActorPickedDelegate when an item is selected. */
+	virtual TSharedRef<ISceneOutliner> CreateActorPicker(
+		const SceneOutliner::FInitializationOptions& InInitOptions,
+		const FOnActorPicked& OnActorPickedDelegate,
+		TWeakObjectPtr<UWorld> SpecifiedWorld = nullptr) const;
+
+	/** Creates a component picker widget. Calls the OnComponentPickedDelegate when an item is selected. */
+	virtual TSharedRef<ISceneOutliner> CreateComponentPicker(
+		const SceneOutliner::FInitializationOptions& InInitOptions,
+		const FOnComponentPicked& OnComponentPickedDelegate,
+		TWeakObjectPtr<UWorld> SpecifiedWorld = nullptr) const;
+
+	/** Creates an actor browser widget (also known as a World Outliner). */
+	virtual TSharedRef<ISceneOutliner> CreateActorBrowser(
+		const SceneOutliner::FInitializationOptions& InInitOptions,
+		TWeakObjectPtr<UWorld> SpecifiedWorld = nullptr) const;
 
 public:
 	/** Register a new type of column available to all scene outliners */
@@ -73,7 +83,7 @@ public:
 
 	/** Register a new type of default column available to all scene outliners */
 	template< typename T >
-	void RegisterDefaultColumnType(SceneOutliner::FDefaultColumnInfo InDefaultColumnInfo)
+	void RegisterDefaultColumnType(SceneOutliner::FColumnInfo InColumnInfo)
 	{
 		auto ID = T::GetID();
 		if ( !ColumnMap.Contains( ID ) )
@@ -83,7 +93,7 @@ public:
 			};
 
 			ColumnMap.Add( ID, FCreateSceneOutlinerColumn::CreateStatic( CreateColumn ) );
-			DefaultColumnMap.Add( ID, InDefaultColumnInfo );
+			DefaultColumnMap.Add( ID, InColumnInfo);
 		}
 	}
 
@@ -107,11 +117,7 @@ public:
 	}
 
 	/** Map of column type name -> default column info */
-	TMap< FName, SceneOutliner::FDefaultColumnInfo> DefaultColumnMap;
-
-
-	/** Additional outliner filters */
-	TMap< FName, SceneOutliner::FOutlinerFilterInfo > OutlinerFilterInfoMap;
+	TMap< FName, SceneOutliner::FColumnInfo> DefaultColumnMap;
 
 private:
 
diff --git a/Engine/Source/Editor/SceneOutliner/Public/SceneOutlinerPublicTypes.h b/Engine/Source/Editor/SceneOutliner/Public/SceneOutlinerPublicTypes.h
--- a/Engine/Source/Editor/SceneOutliner/Public/SceneOutlinerPublicTypes.h
+++ b/Engine/Source/Editor/SceneOutliner/Public/SceneOutlinerPublicTypes.h
@@ -12,28 +12,13 @@
 class FExtender;
 struct FToolMenuContext;
 
-enum class ESceneOutlinerMode : uint8
-{
-	/** Allows all actors to be browsed and selected; syncs selection with the editor; drag and drop attachment, stc. */
-	ActorBrowsing,
-
-	/** Sets the outliner to operate as an actor 'picker'. */
-	ActorPicker,
-
-	/** Sets the outliner to operate as a component 'picker'. */
-	ComponentPicker,
-
-	/** 
-	 * Set the outliner to be a custom displayer of the scene hierarchy
-	 * This mode allow a user to manually change the default behavior of the scene outliner
-	 */
-	Custom,
-};
-
 DECLARE_DELEGATE_TwoParams(FSceneOutlinerModifyContextMenu, FName& /* MenuName */, FToolMenuContext& /* MenuContext */);
 
 namespace SceneOutliner
 {
+	/** A delegate used as a factory to defer mode creation in the outliner */
+	DECLARE_DELEGATE_RetVal_OneParam(ISceneOutlinerMode*, FCreateOutlinerMode, SSceneOutliner*);
+
 	/** Container for built in column types. Function-static so they are available without linking */
 	struct FBuiltInColumnTypes
 	{
@@ -88,30 +73,12 @@ namespace SceneOutliner
 		FCreateSceneOutlinerColumn	Factory;
 	};
 
-	/** Default column information for the scene outliner */
-	struct FDefaultColumnInfo
-	{
-		FDefaultColumnInfo(const FColumnInfo& InColumnInfo, TOptional<ESceneOutlinerMode> InValidMode = TOptional<ESceneOutlinerMode>())
-			: ColumnInfo(InColumnInfo), ValidMode(InValidMode)
-		{}
-
-		FColumnInfo ColumnInfo;
-
-		// The valid mode for this column. If not set, this column will be valid for all.
-		TOptional<ESceneOutlinerMode> ValidMode;
-	};
-
-	struct FSharedDataBase
+	/** Settings for the scene outliner which can be quieried publicly */
+	struct FSharedOutlinerData
 	{
-		/** Mode to operate in */
-		ESceneOutlinerMode Mode;
-
 		/**	Invoked whenever the user attempts to delete an actor from within a Scene Outliner in the actor browsing mode */
 		FCustomSceneOutlinerDeleteDelegate CustomDelete;
 
-		/** Override default context menu handling */
-		FOnContextMenuOpening ContextMenuOverride;
-
 		/** Modify context menu before display */
 		FSceneOutlinerModifyContextMenu ModifyContextMenu;
 
@@ -130,9 +97,8 @@ namespace SceneOutliner
 	public:
 
 		/** Constructor */
-		FSharedDataBase()
-			: Mode( ESceneOutlinerMode::ActorPicker )
-			, bShowParentTree( true )
+		FSharedOutlinerData()
+			: bShowParentTree( true )
 			, bOnlyShowFolders( false )
 			, bShowTransient( false )
 		{}
@@ -146,7 +112,7 @@ namespace SceneOutliner
 	 * is used to modify the outliner's behavior in various ways, such as filtering in or out specific classes
 	 * of actors.
 	 */
-	struct FInitializationOptions : FSharedDataBase
+	struct FInitializationOptions : FSharedOutlinerData
 	{
 		/** True if we should draw the header row above the tree view */
 		bool bShowHeaderRow : 1;
@@ -161,20 +127,9 @@ namespace SceneOutliner
 		bool bShowCreateNewFolder : 1;
 
 		/** Optional collection of filters to use when filtering in the Scene Outliner */
-		TSharedPtr<FOutlinerFilters> Filters;
-
-		/**
-		 * (Old comment) : Broadcasts whenever the Scene Outliners selection changes
-		 *
-		 * This variable will be depecrated as this give the false intuition that it will broadcast when the created scene outliner selection has changed.
-		 * In truth this is only broadcasted by the scene outliners when they are in the actor browsing mode.
-		 * Also it's by any scene outliner, not necessary the one instanced by this initialization options.
-		 */
-		UE_DEPRECATED(4.24, "Use FSceneOutlinerDelegates::Get().SelectionChanged instead.")
-		FSimpleMulticastDelegate::FDelegate OnSelectionChanged;
-		
-		/** If not null, it will force the Scene Outliner to only display this world and it will remove the chose world from the UI */
-		UWorld* SpecifiedWorldToDisplay;
+		TSharedPtr<FOutlinerFilters> Filters;		
+
+		FCreateOutlinerMode ModeFactory;
 
 	public:
 
@@ -185,26 +140,7 @@ namespace SceneOutliner
 			, bFocusSearchBoxWhenOpened( false )
 			, bShowCreateNewFolder( true )
 			, Filters( new FOutlinerFilters )
-			, SpecifiedWorldToDisplay( nullptr )
-		{}
-	};
-
-	/** Outliner data that is shared between a scene outliner and its items */
-	struct FSharedOutlinerData : FSharedDataBase, TSharedFromThis<FSharedOutlinerData>
-	{
-		/** Whether the scene outliner is currently displaying PlayWorld actors */
-		bool bRepresentingPlayWorld;
-
-		/** The world that we are representing */
-		UWorld* RepresentingWorld;
-
-		/** The world the user has chosen to display */
-		TWeakObjectPtr<UWorld> UserChosenWorld;
-
-		FSharedOutlinerData()
-			: bRepresentingPlayWorld(false)
-			, RepresentingWorld(nullptr)
-			, UserChosenWorld(nullptr)
+			, ModeFactory()
 		{}
 	};
 
@@ -215,4 +151,16 @@ namespace SceneOutliner
 		static FMargin	IconPadding() { return FMargin(0.f, 0.f, 6.f, 0.f); };
 	};
 
+	/** A struct which gets, and caches the visibility of a tree item */
+	struct SCENEOUTLINER_API FGetVisibilityCache
+	{
+		/** Map of tree item to visibility */
+		mutable TMap<const ITreeItem*, bool> VisibilityInfo;
+
+		/** Get an item's visibility based on its children */
+		bool RecurseChildren(const ITreeItem& Item) const;
+
+		bool GetVisibility(const ITreeItem& Item) const;
+	};
+
 }	// namespace SceneOutliner
diff --git a/Engine/Source/Editor/SceneOutliner/Public/SceneOutlinerStandaloneTypes.h b/Engine/Source/Editor/SceneOutliner/Public/SceneOutlinerStandaloneTypes.h
--- a/Engine/Source/Editor/SceneOutliner/Public/SceneOutlinerStandaloneTypes.h
+++ b/Engine/Source/Editor/SceneOutliner/Public/SceneOutlinerStandaloneTypes.h
@@ -6,15 +6,17 @@
 #include "Misc/Paths.h"
 #include "UObject/ObjectKey.h"
 #include "Templates/MaxSizeof.h"
-#include "FTreeItemUniqueID.h"
+#include "SceneOutlinerFwd.h"
 
 namespace SceneOutliner
 {
+	using FTreeItemUniqueID = uint32;
+
 	/** Variant type that defines an identifier for a tree item. Assumes 'trivial relocatability' as with many unreal containers. */
 	struct FTreeItemID
 	{
 	public:
-		enum class EType : uint8 { Object, Folder, GCBone, Null };
+		enum class EType : uint8 { Object, Folder, UniqueID, Null };
 
 		/** Default constructed null item ID */
 		FTreeItemID() : Type(EType::Null), CachedHash(0) {}
@@ -39,12 +41,9 @@ namespace SceneOutliner
 			CachedHash = CalculateTypeHash();
 		}
 
-		/** ID representing a sub component tree item */
-		FTreeItemID(FTreeItemID::EType InType, const FTreeItemUniqueID& CustomID) : Type(InType)
+		/** ID representing a generic tree item */
+		FTreeItemID(const FTreeItemUniqueID& CustomID) : Type(EType::UniqueID)
 		{
-			check(InType != EType::Object);
-			check(InType != EType::Folder);
-
 			new (Data) FTreeItemUniqueID(CustomID);
 			CachedHash = CalculateTypeHash();
 		}
@@ -61,7 +60,7 @@ namespace SceneOutliner
 			{
 				case EType::Object:			new (Data) FObjectKey(Other.GetAsObjectKey());		break;
 				case EType::Folder:			new (Data) FName(Other.GetAsFolderRef());			break;
-				case EType::GCBone:			new (Data) FTreeItemUniqueID(Other.GetAsHash());	break;
+				case EType::UniqueID:		new (Data) FTreeItemUniqueID(Other.GetAsHash());	break;
 				default:																		break;
 			}
 
@@ -86,7 +85,7 @@ namespace SceneOutliner
 			{
 				case EType::Object:			GetAsObjectKey().~FObjectKey();							break;
 				case EType::Folder:			GetAsFolderRef().~FName();								break;
-				case EType::GCBone:			/* NOP */												break;
+				case EType::UniqueID:		/* NOP */												break;
 				default:																			break;
 			}
 		}
@@ -107,7 +106,7 @@ namespace SceneOutliner
 			{
 				case EType::Object:			Hash = GetTypeHash(GetAsObjectKey());				break;
 				case EType::Folder:			Hash = GetTypeHash(GetAsFolderRef());				break;
-				case EType::GCBone:			Hash = GetAsHash();									break;
+				case EType::UniqueID:		Hash = GetAsHash();									break;
 				default:																		break;
 			}
 
@@ -132,7 +131,7 @@ namespace SceneOutliner
 			{
 				case EType::Object:			return GetAsObjectKey() == Other.GetAsObjectKey();
 				case EType::Folder:			return GetAsFolderRef() == Other.GetAsFolderRef();
-				case EType::GCBone:			return GetAsHash() == Other.GetAsHash();
+				case EType::UniqueID:		return GetAsHash() == Other.GetAsHash();
 				case EType::Null:			return true;
 				default: check(false);		return false;
 			}
@@ -144,10 +143,64 @@ namespace SceneOutliner
 		static const uint32 MaxSize = TMaxSizeof<FObjectKey, FName, FTreeItemUniqueID>::Value;
 		mutable uint8 Data[MaxSize];
 	};
+	
+	struct FTreeItemType
+	{
+	public:
+		explicit FTreeItemType(const FTreeItemType* Parent = nullptr) : ID(++NextUniqueID), ParentType(Parent) {}
+		FTreeItemType(const FTreeItemType& Src) : ID(Src.ID), ParentType(Src.ParentType) {}
+
+		bool operator==(const FTreeItemType& Other) const
+		{
+			return ID == Other.ID || (ParentType != nullptr && *ParentType == Other);
+		}
+
+		bool IsA(const FTreeItemType& Other) const
+		{
+			return (ID == Other.ID) || (ParentType && ParentType->IsA(Other));
+		}
+
+	private:
+		static uint32 NextUniqueID;
+		uint32 ID;
+		const FTreeItemType* ParentType;
+	};
+
+	struct FCommonLabelData
+	{
+		TWeakPtr<ISceneOutliner> WeakSceneOutliner;
+		static const FLinearColor DarkColor;
+
+		TOptional<FLinearColor> GetForegroundColor(const ITreeItem& TreeItem) const;
+
+		bool CanExecuteRenameRequest(const ITreeItem& Item) const;
+	};
 
-	struct ETreeItemSortOrder
+	/**
+	 * Contains hierarchy change data.
+	 * When an item is added, it will contain a pointer to the new item itself.
+	 * When an item is removed or moved, it will contain the unique ItemID to that item.
+	 * In the case that a folder is being moved, it will also contain the new path to that folder.
+	 */
+	struct FHierarchyChangedData
 	{
-		enum Type { World = 0, Folder = 10, Actor = 20 };
+		enum
+		{
+			Added,
+			Removed,
+			Moved,
+			FolderMoved,
+			FullRefresh
+		} Type;
+
+		// This event may pass one of two kinds of data, depending on the type of event
+		FTreeItemPtr Item;
+
+		FTreeItemID ItemID;
+		// Used for FolderMoved events
+		FName NewPath;
+		/** Actions to apply to items */
+		uint8 ItemActions = 0;
 	};
 
 	/** Folder-specific helper functions */
@@ -161,7 +214,5 @@ namespace SceneOutliner
 		return FName(*FPaths::GetPath(Path.ToString()));
 	}
 
-	/** Move the specified folder path to the specified new parent */
-	FName MoveFolderTo(FName Path, FName NewParent, UWorld& World);
-
+	bool PathIsChildOf(const FName& PotentialChild, const FName& Parent);
 }	// namespace SceneOutliner
