SceneOutliner refactor

diff --git a/Engine/Source/Editor/SceneOutliner/Public/SortHelper.h b/Engine/Source/Editor/SceneOutliner/Public/SortHelper.h
--- a/Engine/Source/Editor/SceneOutliner/Public/SortHelper.h
+++ b/Engine/Source/Editor/SceneOutliner/Public/SortHelper.h
@@ -15,8 +15,8 @@ namespace SceneOutliner
 	 *
 	 * Or:
 	 * 		FSortHelper<FString, FString>()
-	 *			.Primary(FGetPrimaryStringVisitor())
-	 *			.Secondary(FGetSecondaryStringVisitor())
+	 *			.Primary(FGetPrimaryString())
+	 *			.Secondary(FGetSecondaryString())
 	 *			.SortItems(Array);
 	 */
 	template<typename PrimaryKeyType, typename SecondaryKeyType = int32>
@@ -27,7 +27,6 @@ namespace SceneOutliner
 
 		FSortHelper()
 			: PrimarySortMode(EColumnSortMode::None), SecondarySortMode(EColumnSortMode::None)
-			, PrimaryVisitor(nullptr), SecondaryVisitor(nullptr)
 		{}
 
 		/** Sort primarily by the specified function and mode. Beware the function is a reference, so must be valid for the lifetime of this instance. */
@@ -37,13 +36,6 @@ namespace SceneOutliner
 			PrimaryFunction = MoveTemp(Function);
 			return *this;
 		}
-		/** Sort primarily using the specified 'getter' visitor and mode */
-		FSortHelper& Primary(const TTreeItemGetter<PrimaryKeyType>& Visitor, EColumnSortMode::Type SortMode)
-		{
-			PrimarySortMode = SortMode;
-			PrimaryVisitor = &Visitor;
-			return *this;
-		}
 
 		/** Sort secondarily by the specified function and mode. Beware the function is a reference, so must be valid for the lifetime of this instance. */
 		FSortHelper& Secondary(FSecondaryFunction&& Function, EColumnSortMode::Type SortMode)
@@ -52,13 +44,6 @@ namespace SceneOutliner
 			SecondaryFunction = MoveTemp(Function);
 			return *this;
 		}
-		/** Sort secondarily using the specified 'getter' visitor and mode */
-		FSortHelper& Secondary(const TTreeItemGetter<SecondaryKeyType>& Visitor, EColumnSortMode::Type SortMode)
-		{
-			PrimarySortMode = SortMode;
-			SecondaryVisitor = &Visitor;
-			return *this;
-		}
 
 		/** Sort the specified array using the current sort settings */
 		void Sort(TArray<FTreeItemPtr>& Array)
@@ -69,29 +54,12 @@ namespace SceneOutliner
 			{
 				const auto& Element = Array[Index];
 
-				PrimaryKeyType PrimaryKey;
-				if (PrimaryVisitor)
-				{
-					Element->Visit(*PrimaryVisitor);
-					PrimaryKey = MoveTemp(PrimaryVisitor->Data);
-				}
-				else if (PrimaryFunction)
-				{
-					PrimaryKey = PrimaryFunction.GetValue()(*Element);
-				}
+				PrimaryKeyType PrimaryKey = PrimaryFunction(*Element);
 
 				SecondaryKeyType SecondaryKey;
 				if (SecondarySortMode != EColumnSortMode::None)
 				{
-					if (SecondaryVisitor)
-					{
-						Element->Visit(*SecondaryVisitor);
-						SecondaryKey = MoveTemp(SecondaryVisitor->Data);
-					}
-					else if (SecondaryFunction)
-					{
-						SecondaryKey = SecondaryFunction.GetValue()(*Element);
-					}
+					SecondaryKey = SecondaryFunction(*Element);
 				}
 
 				SortData.Emplace(Index, MoveTemp(PrimaryKey), MoveTemp(SecondaryKey));
@@ -126,11 +94,8 @@ namespace SceneOutliner
 		EColumnSortMode::Type 	PrimarySortMode;
 		EColumnSortMode::Type 	SecondarySortMode;
 
-		TOptional<FPrimaryFunction>		PrimaryFunction;
-		TOptional<FSecondaryFunction>	SecondaryFunction;
-
-		const TTreeItemGetter<PrimaryKeyType>*		PrimaryVisitor;
-		const TTreeItemGetter<SecondaryKeyType>*	SecondaryVisitor;
+		FPrimaryFunction	PrimaryFunction;
+		FSecondaryFunction	SecondaryFunction;
 
 		/** Aggregated data from the sort methods. We extract the sort data from all elements first, then sort based on the extracted data. */
 		struct FSortPayload
diff --git a/Engine/Source/Editor/SceneOutliner/Public/WorldTreeItem.h b/Engine/Source/Editor/SceneOutliner/Public/WorldTreeItem.h
--- a/Engine/Source/Editor/SceneOutliner/Public/WorldTreeItem.h
+++ b/Engine/Source/Editor/SceneOutliner/Public/WorldTreeItem.h
@@ -11,68 +11,52 @@ class UToolMenu;
 
 namespace SceneOutliner
 {
+	/** Get a description of a world to display in the scene outliner */
+	FText GetWorldDescription(UWorld* World);
+
 	/** A tree item that represents an entire world */
-	struct FWorldTreeItem : ITreeItem
+	struct SCENEOUTLINER_API FWorldTreeItem : ITreeItem
 	{
+	public:
+		DECLARE_DELEGATE_RetVal_OneParam(bool, FFilterPredicate, const UWorld*);
+		DECLARE_DELEGATE_RetVal_OneParam(bool, FInteractivePredicate, const UWorld*);
+
+		bool Filter(FFilterPredicate Pred) const
+		{
+			return Pred.Execute(World.Get());
+		}
+
+		bool GetInteractiveState(FInteractivePredicate Pred) const
+		{
+			return Pred.Execute(World.Get());
+		}
+
 		/** The world this tree item is associated with. */
 		mutable TWeakObjectPtr<UWorld> World;
 		
 		/** Constant identifier for this tree item */
 		const FObjectKey ID;
 
+		/** Static type identifier for this tree item class */
+		static const FTreeItemType Type;
+
 		/** Construct this item from a world */
 		FWorldTreeItem(UWorld* InWorld);
+		FWorldTreeItem(TWeakObjectPtr<UWorld> InWorld);
 
-		/** Get this item's parent item. Always returns nullptr. */
-		virtual FTreeItemPtr FindParent(const FTreeItemMap& ExistingItems) const override;
-
-		/** Create this item's parent. Always returns nullptr. */
-		virtual FTreeItemPtr CreateParent() const override;
-
- 		/** Visit this tree item */
- 		virtual void Visit(const ITreeItemVisitor& Visitor) const override;
- 		virtual void Visit(const IMutableTreeItemVisitor& Visitor) override;
-
-	public:
-
-		/** Get the ID that represents this tree item. Used to reference this item in a map */
+		/* Begin ITreeItem Implementation */
+		virtual bool IsValid() const override { return World.IsValid(); }
 		virtual FTreeItemID GetID() const override;
-
-		/** Get the raw string to display for this tree item - used for sorting */
 		virtual FString GetDisplayString() const override;
-
-		/** Get the sort priority given to this item's type */
-		virtual int32 GetTypeSortPriority() const override;
-
-		/** Check whether it should be possible to interact with this tree item */
 		virtual bool CanInteract() const override;
-
-		/** Generate a context menu for this item. Only called if *only* this item is selected. */
 		virtual void GenerateContextMenu(UToolMenu* Menu, SSceneOutliner& Outliner) override;
-
-		/** Populate the specified drag/drop payload with any relevant information for this type */
-		virtual void PopulateDragDropPayload(FDragDropPayload& Payload) const override
-		{
-		}
-
-		/** Called to test whether the specified payload can be dropped onto this tree item */
-		virtual FDragValidationInfo ValidateDrop(FDragDropPayload& DraggedObjects, UWorld& World) const override;
-
-		/** Called to drop the specified objects on this item. Only called if ValidateDrop() allows. */
-		virtual void OnDrop(FDragDropPayload& DraggedObjects, UWorld& World, const FDragValidationInfo& ValidationInfo, TSharedRef<SWidget> DroppedOnWidget) override;
-
-		/** Get just the name of the world, for tooltip use */
-		FString GetWorldName() const;
-
-	private:
-
-		/** Create a new folder at the root of this world */
-		void CreateFolder(TWeakPtr<SSceneOutliner> WeakOutliner);
-		
-	public:
+		virtual TSharedRef<SWidget> GenerateLabelWidget(ISceneOutliner& Outliner, const STableRow<FTreeItemPtr>& InRow) override;
+		/* End ITreeItem Implementation */
 
 		/** Open the world settings for the contained world */
 		void OpenWorldSettings() const;
+		/** Get just the name of the world, for tooltip use */
+		FString GetWorldName() const;
 	};
 
 }		// namespace SceneOutliner
diff --git a/Engine/Source/Editor/Sequencer/Private/LevelEditorSequencerIntegration.cpp b/Engine/Source/Editor/Sequencer/Private/LevelEditorSequencerIntegration.cpp
--- a/Engine/Source/Editor/Sequencer/Private/LevelEditorSequencerIntegration.cpp
+++ b/Engine/Source/Editor/Sequencer/Private/LevelEditorSequencerIntegration.cpp
@@ -997,7 +997,7 @@ void FLevelEditorSequencerIntegration::AttachOutlinerColumn()
 	SceneOutliner::FColumnInfo ColumnInfo(SceneOutliner::EColumnVisibility::Visible, 15, 
 		FCreateSceneOutlinerColumn::CreateRaw( this, &FLevelEditorSequencerIntegration::CreateSequencerInfoColumn));
 
-	SceneOutlinerModule.RegisterDefaultColumnType< Sequencer::FSequencerInfoColumn >(SceneOutliner::FDefaultColumnInfo(ColumnInfo));
+	SceneOutlinerModule.RegisterDefaultColumnType< Sequencer::FSequencerInfoColumn >(ColumnInfo);
 
 	AcquiredResources.Add([=]{ this->DetachOutlinerColumn(); });
 }
diff --git a/Engine/Source/Editor/Sequencer/Private/SSequencerPlayRateCombo.cpp b/Engine/Source/Editor/Sequencer/Private/SSequencerPlayRateCombo.cpp
--- a/Engine/Source/Editor/Sequencer/Private/SSequencerPlayRateCombo.cpp
+++ b/Engine/Source/Editor/Sequencer/Private/SSequencerPlayRateCombo.cpp
@@ -22,6 +22,7 @@
 #include "ContentBrowserModule.h"
 #include "Misc/Timecode.h"
 #include "Styling/ToolBarStyle.h"
+#include "ActorTreeItem.h"
 
 #define LOCTEXT_NAMESPACE "SSequencerPlayRateCombo"
 
@@ -395,7 +396,6 @@ void SSequencerPlayRateCombo::PopulateCustomClockSourceMenu(FMenuBuilder& MenuBu
 
 		// Set up a menu entry to assign an actor to the object binding node
 		FInitializationOptions InitOptions;
-		InitOptions.Mode = ESceneOutlinerMode::ActorPicker;
 
 		// We hide the header row to keep the UI compact.
 		InitOptions.bShowHeaderRow = false;
@@ -406,7 +406,7 @@ void SSequencerPlayRateCombo::PopulateCustomClockSourceMenu(FMenuBuilder& MenuBu
 		InitOptions.ColumnMap.Add(FBuiltInColumnTypes::Label(), FColumnInfo(EColumnVisibility::Visible, 0));
 
 		// Only display actors that are not possessed already
-		InitOptions.Filters->AddFilterPredicate( FActorFilterPredicate::CreateLambda( IsActorValid ) );
+		InitOptions.Filters->AddFilterPredicate<FActorTreeItem>(FActorTreeItem::FFilterPredicate::CreateLambda( IsActorValid ) );
 
 		// actor selector to allow the user to choose an actor
 		FSceneOutlinerModule& SceneOutlinerModule = FModuleManager::LoadModuleChecked<FSceneOutlinerModule>("SceneOutliner");
@@ -415,7 +415,7 @@ void SSequencerPlayRateCombo::PopulateCustomClockSourceMenu(FMenuBuilder& MenuBu
 			.MaxDesiredHeight(400.0f)
 			.WidthOverride(300.0f)
 			[
-				SceneOutlinerModule.CreateSceneOutliner(InitOptions, FOnActorPicked::CreateLambda([this](AActor* In){ this->SetCustomClockSource(In); }))
+				SceneOutlinerModule.CreateActorPicker(InitOptions, FOnActorPicked::CreateLambda([this](AActor* In){ this->SetCustomClockSource(In); }))
 			],
 			FText(),
 			true /*bNoIndent*/
diff --git a/Engine/Source/Editor/Sequencer/Private/Sequencer.cpp b/Engine/Source/Editor/Sequencer/Private/Sequencer.cpp
--- a/Engine/Source/Editor/Sequencer/Private/Sequencer.cpp
+++ b/Engine/Source/Editor/Sequencer/Private/Sequencer.cpp
@@ -139,6 +139,7 @@
 #include "Components/SkeletalMeshComponent.h"
 #include "SequencerCustomizationManager.h"
 #include "SSequencerGroupManager.h"
+#include "ActorTreeItem.h"
 
 #define LOCTEXT_NAMESPACE "Sequencer"
 
@@ -7066,8 +7067,6 @@ void FSequencer::AssignActor(FMenuBuilder& MenuBuilder, FGuid InObjectBinding)
 	// Set up a menu entry to assign an actor to the object binding node
 	FInitializationOptions InitOptions;
 	{
-		InitOptions.Mode = ESceneOutlinerMode::ActorPicker;
-
 		// We hide the header row to keep the UI compact.
 		InitOptions.bShowHeaderRow = false;
 		InitOptions.bShowSearchBox = true;
@@ -7077,7 +7076,7 @@ void FSequencer::AssignActor(FMenuBuilder& MenuBuilder, FGuid InObjectBinding)
 		InitOptions.ColumnMap.Add(FBuiltInColumnTypes::Label(), FColumnInfo(EColumnVisibility::Visible, 0));
 
 		// Only display actors that are not possessed already
-		InitOptions.Filters->AddFilterPredicate( FActorFilterPredicate::CreateLambda( IsActorValidForAssignment ) );
+		InitOptions.Filters->AddFilterPredicate<FActorTreeItem>(FActorTreeItem::FFilterPredicate::CreateLambda( IsActorValidForAssignment ) );
 	}
 
 	// actor selector to allow the user to choose an actor
@@ -7087,7 +7086,7 @@ void FSequencer::AssignActor(FMenuBuilder& MenuBuilder, FGuid InObjectBinding)
 		.MaxDesiredHeight(400.0f)
 		.WidthOverride(300.0f)
 		[
-			SceneOutlinerModule.CreateSceneOutliner(
+			SceneOutlinerModule.CreateActorPicker(
 				InitOptions,
 				FOnActorPicked::CreateLambda([=](AActor* Actor){
 					// Create a new binding for this actor
diff --git a/Engine/Source/Editor/Sequencer/Private/SequencerInfoColumn.cpp b/Engine/Source/Editor/Sequencer/Private/SequencerInfoColumn.cpp
--- a/Engine/Source/Editor/Sequencer/Private/SequencerInfoColumn.cpp
+++ b/Engine/Source/Editor/Sequencer/Private/SequencerInfoColumn.cpp
@@ -15,33 +15,38 @@
 #include "EditorClassUtils.h"
 #include "SortHelper.h"
 
+#include "ActorTreeItem.h"
+
 #define LOCTEXT_NAMESPACE "SequencerInfoColumn"
 
 namespace Sequencer
 {
 
-struct FGetInfo : SceneOutliner::TTreeItemGetter<FString>
+/** Functor which retrieves actor string information for sorting */
+struct FGetActorInfo
 {
-	FGetInfo(const FSequencerInfoColumn& InColumn)
+	FGetActorInfo(const FSequencerInfoColumn& InColumn)
 		: WeakColumn(StaticCastSharedRef<const FSequencerInfoColumn>(InColumn.AsShared()))
 	{}
 
-	virtual FString Get(const SceneOutliner::FActorTreeItem& ActorItem) const override
+	FString operator()(const SceneOutliner::ITreeItem& Item) const
 	{
 		if (!WeakColumn.IsValid())
 		{
 			return FString();
 		}
-
-		AActor* Actor = ActorItem.Actor.Get();
-		if (!Actor)
+		if (const SceneOutliner::FActorTreeItem* ActorItem = Item.CastTo<SceneOutliner::FActorTreeItem>())
 		{
-			return FString();
-		}
+			AActor* Actor = ActorItem->Actor.Get();
+			if (Actor)
+			{
+				const FSequencerInfoColumn& Column = *WeakColumn.Pin();
 
-		const FSequencerInfoColumn& Column = *WeakColumn.Pin();
-
-		return Column.GetTextForActor(Actor);
+				return Column.GetTextForActor(Actor);
+			}
+		}
+		
+		return FString();
 	}
 
 
@@ -49,7 +54,6 @@ struct FGetInfo : SceneOutliner::TTreeItemGetter<FString>
 	TWeakPtr< const FSequencerInfoColumn > WeakColumn;
 };
 
-
 FSequencerInfoColumn::FSequencerInfoColumn(ISceneOutliner& InSceneOutliner, FSequencer& InSequencer, const FLevelEditorSequencerBindingData& InBindingData)
 	: WeakSceneOutliner(StaticCastSharedRef<ISceneOutliner>(InSceneOutliner.AsShared())) 
 	, WeakSequencer(StaticCastSharedRef<FSequencer>(InSequencer.AsShared()))
@@ -120,7 +124,7 @@ void FSequencerInfoColumn::SortItems(TArray<SceneOutliner::FTreeItemPtr>& OutIte
 	if (WeakBindingData.IsValid())
 	{
 		SceneOutliner::FSortHelper<FString>()
-			.Primary(FGetInfo(*this), SortMode)
+			.Primary(FGetActorInfo(*this), SortMode)
 			.Sort(OutItems);
 	}
 }
@@ -138,7 +142,7 @@ FString FSequencerInfoColumn::GetTextForActor(AActor* InActor) const
 FText FSequencerInfoColumn::GetTextForItem( TWeakPtr<SceneOutliner::ITreeItem> TreeItem ) const
 {
 	auto Item = TreeItem.Pin();
-	return Item.IsValid() && WeakBindingData.IsValid() ? FText::FromString(Item->Get(FGetInfo(*this))) : FText::GetEmpty();
+	return Item.IsValid() && WeakBindingData.IsValid() ? FText::FromString(FGetActorInfo(*this)(*Item)) : FText::GetEmpty();
 }
 
 }
