Build Fixing (a little bit in the dark)

diff --git a/Engine/Plugins/Enterprise/DataprepEditor/Source/DataprepEditor/Private/PreviewSystem/DataprepPreviewSceneOutlinerColumn.cpp b/Engine/Plugins/Enterprise/DataprepEditor/Source/DataprepEditor/Private/PreviewSystem/DataprepPreviewSceneOutlinerColumn.cpp
--- a/Engine/Plugins/Enterprise/DataprepEditor/Source/DataprepEditor/Private/PreviewSystem/DataprepPreviewSceneOutlinerColumn.cpp
+++ b/Engine/Plugins/Enterprise/DataprepEditor/Source/DataprepEditor/Private/PreviewSystem/DataprepPreviewSceneOutlinerColumn.cpp
@@ -12,6 +12,8 @@
 #include "UObject/WeakObjectPtrTemplates.h"
 #include "Widgets/SNullWidget.h"
 #include "Widgets/SWidget.h"
+#include "GameFramework/Actor.h"
+#include "Components/ActorComponent.h"
 
 #define LOCTEXT_NAMESPACE "DataprepPreviewOutlinerColumn"
 
diff --git a/Engine/Source/Editor/SceneOutliner/Private/ActorBrowsingMode.cpp b/Engine/Source/Editor/SceneOutliner/Private/ActorBrowsingMode.cpp
--- a/Engine/Source/Editor/SceneOutliner/Private/ActorBrowsingMode.cpp
+++ b/Engine/Source/Editor/SceneOutliner/Private/ActorBrowsingMode.cpp
@@ -27,8 +27,6 @@
 #include "ActorBrowsingModeSettings.h"
 #include "ScopedTransaction.h"
 
-DEFINE_LOG_CATEGORY_STATIC(LogSceneOutliner, Log, All);
-
 #define LOCTEXT_NAMESPACE "SceneOutliner_ActorBrowsingMode"
 
 UActorBrowsingModeSettings::UActorBrowsingModeSettings(const FObjectInitializer& ObjectInitializer)
@@ -479,7 +477,7 @@ namespace SceneOutliner
 			const bool bSelectEvenIfHidden = true;	// @todo outliner: Is this actually OK?
 			for (auto* Actor : SelectedActors)
 			{
-				UE_LOG(LogSceneOutliner, Verbose, TEXT("Clicking on Actor (world outliner): %s (%s)"), *Actor->GetClass()->GetName(), *Actor->GetActorLabel());
+				//UE_LOG(LogSceneOutliner, Verbose, TEXT("Clicking on Actor (world outliner): %s (%s)"), *Actor->GetClass()->GetName(), *Actor->GetActorLabel());
 				GEditor->SelectActor(Actor, bShouldSelect, bNotifyAfterSelect, bSelectEvenIfHidden);
 			}
 
diff --git a/Engine/Source/Editor/SceneOutliner/Public/SSceneOutliner.h b/Engine/Source/Editor/SceneOutliner/Public/SSceneOutliner.h
--- a/Engine/Source/Editor/SceneOutliner/Public/SSceneOutliner.h
+++ b/Engine/Source/Editor/SceneOutliner/Public/SSceneOutliner.h
@@ -27,6 +27,7 @@
 #include "SceneOutlinerStandaloneTypes.h"
 
 #include "ISceneOutlinerHierarchy.h"
+#include "ISceneOutlinerMode.h"
 #include "SceneOutlinerDragDrop.h"
 
 class FMenuBuilder;
diff --git a/Engine/Source/Editor/SceneOutliner/Public/SceneOutlinerFilters.h b/Engine/Source/Editor/SceneOutliner/Public/SceneOutlinerFilters.h
--- a/Engine/Source/Editor/SceneOutliner/Public/SceneOutlinerFilters.h
+++ b/Engine/Source/Editor/SceneOutliner/Public/SceneOutlinerFilters.h
@@ -6,6 +6,7 @@
 #include "Misc/IFilter.h"
 #include "Misc/FilterCollection.h"
 #include "SceneOutlinerFwd.h"
+#include "ITreeItem.h"
 
 class FMenuBuilder;
 
