Fix Win64 Build

diff --git a/Engine/Source/Editor/SceneOutliner/Public/SSceneOutliner.h b/Engine/Source/Editor/SceneOutliner/Public/SSceneOutliner.h
--- a/Engine/Source/Editor/SceneOutliner/Public/SSceneOutliner.h
+++ b/Engine/Source/Editor/SceneOutliner/Public/SSceneOutliner.h
@@ -27,7 +27,6 @@
 #include "SceneOutlinerStandaloneTypes.h"
 
 #include "ISceneOutlinerHierarchy.h"
-#include "ISceneOutlinerMode.h"
 #include "SceneOutlinerDragDrop.h"
 
 class FMenuBuilder;
@@ -455,8 +454,8 @@ namespace SceneOutliner
 			bool bPassesFilters = Filters->PassesAllFilters(Temporary);
 			if (bPassesFilters)
 			{
-				check(Mode);
-				Mode->OnItemPassesFilters(Temporary);
+				check(this->Mode);
+				this->Mode->OnItemPassesFilters(Temporary);
 			}
 
 			bPassesFilters &= SearchBoxFilter->PassesFilter(Temporary);
diff --git a/Engine/Source/Editor/SceneOutliner/Public/SceneOutlinerDragDrop.h b/Engine/Source/Editor/SceneOutliner/Public/SceneOutlinerDragDrop.h
--- a/Engine/Source/Editor/SceneOutliner/Public/SceneOutlinerDragDrop.h
+++ b/Engine/Source/Editor/SceneOutliner/Public/SceneOutlinerDragDrop.h
@@ -7,6 +7,7 @@
 #include "Layout/Visibility.h"
 #include "Input/DragAndDrop.h"
 #include "DragAndDrop/CompositeDragDropOp.h"
+#include "ITreeItem.h"
 
 namespace SceneOutliner
 {
