Fix Linux/Mac Build

Remove ISceneOutlnerMode/SSceneOutliner Circular Dep
Add Missing typename (linux)
Fix Parenthesis assignements
Fix Bad Constructor param forwarding

diff --git a/Engine/Source/Editor/SceneOutliner/Private/ActorHierarchy.cpp b/Engine/Source/Editor/SceneOutliner/Private/ActorHierarchy.cpp
--- a/Engine/Source/Editor/SceneOutliner/Private/ActorHierarchy.cpp
+++ b/Engine/Source/Editor/SceneOutliner/Private/ActorHierarchy.cpp
@@ -157,22 +157,21 @@ namespace SceneOutliner
 		{
 			if (ActorTreeItem->Actor.IsValid())
 			{
-				const FTreeItemPtr* ParentItem;
 				if (const AActor* ParentActor = ActorTreeItem->Actor->GetAttachParentActor())
 				{
-					if (ParentItem = Items.Find(ParentActor))
+					if (const FTreeItemPtr* ParentItem = Items.Find(ParentActor))
 					{
 						return *ParentItem;
 					}
 				}
 				else if (Mode->ShouldShowFolders() && !ActorTreeItem->Actor->GetFolderPath().IsNone())
 				{
-					if (ParentItem = Items.Find(ActorTreeItem->Actor->GetFolderPath()))
+					if (const FTreeItemPtr* ParentItem = Items.Find(ActorTreeItem->Actor->GetFolderPath()))
 					{
 						return *ParentItem;
 					}
 				}
-				else if (ParentItem = Items.Find(ActorTreeItem->Actor->GetWorld()))
+				else if (const FTreeItemPtr* ParentItem = Items.Find(ActorTreeItem->Actor->GetWorld()))
 				{
 					return *ParentItem;
 				}
diff --git a/Engine/Source/Editor/SceneOutliner/Private/ActorModeInteractive.cpp b/Engine/Source/Editor/SceneOutliner/Private/ActorModeInteractive.cpp
--- a/Engine/Source/Editor/SceneOutliner/Private/ActorModeInteractive.cpp
+++ b/Engine/Source/Editor/SceneOutliner/Private/ActorModeInteractive.cpp
@@ -8,7 +8,7 @@
 namespace SceneOutliner
 {
 	FActorModeInteractive::FActorModeInteractive(SSceneOutliner* InSceneOutliner, bool bInHideComponents, TWeakObjectPtr<UWorld> InSpecifiedWorldToDisplay)
-		: FActorMode(InSceneOutliner, bHideComponents, SpecifiedWorldToDisplay)
+		: FActorMode(InSceneOutliner, bInHideComponents, InSpecifiedWorldToDisplay)
 	{
 		USelection::SelectionChangedEvent.AddRaw(this, &FActorModeInteractive::OnLevelSelectionChanged);
 		USelection::SelectObjectEvent.AddRaw(this, &FActorModeInteractive::OnLevelSelectionChanged);
diff --git a/Engine/Source/Editor/SceneOutliner/Public/ISceneOutlinerMode.h b/Engine/Source/Editor/SceneOutliner/Public/ISceneOutlinerMode.h
--- a/Engine/Source/Editor/SceneOutliner/Public/ISceneOutlinerMode.h
+++ b/Engine/Source/Editor/SceneOutliner/Public/ISceneOutlinerMode.h
@@ -31,7 +31,8 @@ namespace SceneOutliner
 		template <typename TreeItemType, typename TreeItemData>
 		FTreeItemPtr CreateItemFor(const TreeItemData& Data, bool bForce = false)
 		{
-			return SceneOutliner->CreateItemFor<TreeItemType>(Data, bForce);
+			return SceneOutliner->CreateItemFor<TreeItemType>(Data, 
+				[this](const ITreeItem& Item) { OnItemPassesFilters(Item); }, bForce);
 		}
 
 		/** Fill the view button menu with content */
diff --git a/Engine/Source/Editor/SceneOutliner/Public/SSceneOutliner.h b/Engine/Source/Editor/SceneOutliner/Public/SSceneOutliner.h
--- a/Engine/Source/Editor/SceneOutliner/Public/SSceneOutliner.h
+++ b/Engine/Source/Editor/SceneOutliner/Public/SSceneOutliner.h
@@ -448,14 +448,13 @@ namespace SceneOutliner
 	public:
 		// Test the filters using stack-allocated data to prevent unnecessary heap allocations
 		template <typename TreeItemType, typename TreeItemData>
-		FTreeItemPtr CreateItemFor(const TreeItemData& Data, bool bForce = false)
+		FTreeItemPtr CreateItemFor(const TreeItemData& Data, TFunctionRef<void(const TreeItemType&)> OnItemPassesFilters, bool bForce = false)
 		{
 			const TreeItemType Temporary(Data);
 			bool bPassesFilters = Filters->PassesAllFilters(Temporary);
 			if (bPassesFilters)
 			{
-				check(this->Mode);
-				this->Mode->OnItemPassesFilters(Temporary);
+				OnItemPassesFilters(Temporary);
 			}
 
 			bPassesFilters &= SearchBoxFilter->PassesFilter(Temporary);
diff --git a/Engine/Source/Editor/SceneOutliner/Public/SceneOutlinerFilters.h b/Engine/Source/Editor/SceneOutliner/Public/SceneOutlinerFilters.h
--- a/Engine/Source/Editor/SceneOutliner/Public/SceneOutlinerFilters.h
+++ b/Engine/Source/Editor/SceneOutliner/Public/SceneOutlinerFilters.h
@@ -197,7 +197,7 @@ namespace SceneOutliner
 
 		/** Add a filter predicate to this filter collection */
 		template <typename T>
-		void AddFilterPredicate(typename T::FFilterPredicate InFilterPred, EDefaultFilterBehaviour InDefaultBehaviour = EDefaultFilterBehaviour::Pass, typename T::FInteractivePredicate InInteractivePred = T::FInteractivePredicate())
+		void AddFilterPredicate(typename T::FFilterPredicate InFilterPred, EDefaultFilterBehaviour InDefaultBehaviour = EDefaultFilterBehaviour::Pass, typename T::FInteractivePredicate InInteractivePred = typename T::FInteractivePredicate())
 		{
 			Add(MakeShareable(new TOutlinerPredicateFilter<T>(InFilterPred, InDefaultBehaviour, InInteractivePred)));
 		}
