Removed SceneOutliner namespace

diff --git a/Engine/Source/Editor/SceneOutliner/Private/ActorBrowsingMode.cpp b/Engine/Source/Editor/SceneOutliner/Private/ActorBrowsingMode.cpp
--- a/Engine/Source/Editor/SceneOutliner/Private/ActorBrowsingMode.cpp
+++ b/Engine/Source/Editor/SceneOutliner/Private/ActorBrowsingMode.cpp
@@ -27,516 +27,459 @@
 #include "ActorBrowsingModeSettings.h"
 #include "ScopedTransaction.h"
 
+//#todo_Outliner: remove
+DEFINE_LOG_CATEGORY_STATIC(LogSceneOutliner, Log, All);
+
 #define LOCTEXT_NAMESPACE "SceneOutliner_ActorBrowsingMode"
 
 UActorBrowsingModeSettings::UActorBrowsingModeSettings(const FObjectInitializer& ObjectInitializer)
-	: Super(ObjectInitializer)
+: Super(ObjectInitializer)
 { }
 
-namespace SceneOutliner
+using FActorFilter = TSceneOutlinerPredicateFilter<FActorTreeItem>;
+
+FActorBrowsingMode::FActorBrowsingMode(SSceneOutliner* InSceneOutliner, TWeakObjectPtr<UWorld> InSpecifiedWorldToDisplay)
+	: FActorModeInteractive(InSceneOutliner, true, InSpecifiedWorldToDisplay)
+	, FilteredActorCount(0)
 {
-	using FActorFilter = TOutlinerPredicateFilter<FActorTreeItem>;
+	// Capture selection changes of bones from mesh selection in fracture tools
+	FSceneOutlinerDelegates::Get().OnComponentsUpdated.AddRaw(this, &FActorBrowsingMode::OnComponentsUpdated);
+
+	GEngine->OnLevelActorDeleted().AddRaw(this, &FActorBrowsingMode::OnLevelActorDeleted);
+
+	FEditorDelegates::OnEditCutActorsBegin.AddRaw(this, &FActorBrowsingMode::OnEditCutActorsBegin);
+	FEditorDelegates::OnEditCutActorsEnd.AddRaw(this, &FActorBrowsingMode::OnEditCutActorsEnd);
+	FEditorDelegates::OnEditCopyActorsBegin.AddRaw(this, &FActorBrowsingMode::OnEditCopyActorsBegin);
+	FEditorDelegates::OnEditCopyActorsEnd.AddRaw(this, &FActorBrowsingMode::OnEditCopyActorsEnd);
+	FEditorDelegates::OnEditPasteActorsBegin.AddRaw(this, &FActorBrowsingMode::OnEditPasteActorsBegin);
+	FEditorDelegates::OnEditPasteActorsEnd.AddRaw(this, &FActorBrowsingMode::OnEditPasteActorsEnd);
+	FEditorDelegates::OnDuplicateActorsBegin.AddRaw(this, &FActorBrowsingMode::OnDuplicateActorsBegin);
+	FEditorDelegates::OnDuplicateActorsEnd.AddRaw(this, &FActorBrowsingMode::OnDuplicateActorsEnd);
+	FEditorDelegates::OnDeleteActorsBegin.AddRaw(this, &FActorBrowsingMode::OnDeleteActorsBegin);
+	FEditorDelegates::OnDeleteActorsEnd.AddRaw(this, &FActorBrowsingMode::OnDeleteActorsEnd);
+
+	UActorBrowsingModeSettings* SharedSettings = GetMutableDefault<UActorBrowsingModeSettings>();
+	// Get the OutlinerModule to register FilterInfos with the FilterInfoMap
+	FSceneOutlinerFilterInfo ShowOnlySelectedActorsInfo(LOCTEXT("ToggleShowOnlySelected", "Only Selected"), LOCTEXT("ToggleShowOnlySelectedToolTip", "When enabled, only displays actors that are currently selected."), SharedSettings->bShowOnlySelectedActors, FCreateSceneOutlinerFilter::CreateStatic(&FActorBrowsingMode::CreateShowOnlySelectedActorsFilter));
+	ShowOnlySelectedActorsInfo.OnToggle().AddLambda([this](bool bIsActive)
+		{
+			UActorBrowsingModeSettings* Settings = GetMutableDefault<UActorBrowsingModeSettings>();
+			Settings->bShowOnlySelectedActors = bIsActive;
+			Settings->PostEditChange();
+		});
+	FilterInfoMap.Add(TEXT("ShowOnlySelectedActors"), ShowOnlySelectedActorsInfo);
+
+	FSceneOutlinerFilterInfo HideTemporaryActorsInfo(LOCTEXT("ToggleHideTemporaryActors", "Hide Temporary Actors"), LOCTEXT("ToggleHideTemporaryActorsToolTip", "When enabled, hides temporary/run-time Actors."), SharedSettings->bHideTemporaryActors, FCreateSceneOutlinerFilter::CreateStatic(&FActorBrowsingMode::CreateHideTemporaryActorsFilter));
+	HideTemporaryActorsInfo.OnToggle().AddLambda([this](bool bIsActive)
+		{
+			UActorBrowsingModeSettings* Settings = GetMutableDefault<UActorBrowsingModeSettings>();
+			Settings->bHideTemporaryActors = bIsActive;
+			Settings->PostEditChange();
+		});
+	FilterInfoMap.Add(TEXT("HideTemporaryActors"), HideTemporaryActorsInfo);
+
+	FSceneOutlinerFilterInfo OnlyCurrentLevelInfo(LOCTEXT("ToggleShowOnlyCurrentLevel", "Only in Current Level"), LOCTEXT("ToggleShowOnlyCurrentLevelToolTip", "When enabled, only shows Actors that are in the Current Level."), SharedSettings->bShowOnlyActorsInCurrentLevel, FCreateSceneOutlinerFilter::CreateStatic(&FActorBrowsingMode::CreateIsInCurrentLevelFilter));
+	OnlyCurrentLevelInfo.OnToggle().AddLambda([this](bool bIsActive)
+		{
+			UActorBrowsingModeSettings* Settings = GetMutableDefault<UActorBrowsingModeSettings>();
+			Settings->bShowOnlyActorsInCurrentLevel = bIsActive;
+			Settings->PostEditChange();
+		});
+	FilterInfoMap.Add(TEXT("ShowOnlyCurrentLevel"), OnlyCurrentLevelInfo);
+
+	bHideComponents = SharedSettings->bHideActorComponents;
+	FSceneOutlinerFilterInfo HideComponentsInfo(LOCTEXT("ToggleHideActorComponents", "Hide Actor Components"), LOCTEXT("ToggleHideActorComponentsToolTip", "When Enabled, hides components belonging to actors."), SharedSettings->bHideActorComponents, FCreateSceneOutlinerFilter::CreateStatic(&FActorBrowsingMode::CreateHideComponentsFilter));
+	HideComponentsInfo.OnToggle().AddLambda([this](bool bIsActive)
+		{
+			UActorBrowsingModeSettings* Settings = GetMutableDefault<UActorBrowsingModeSettings>();
+			Settings->bHideActorComponents = bHideComponents = bIsActive;
+			Settings->PostEditChange();
 
-	FActorBrowsingMode::FActorBrowsingMode(SSceneOutliner* InSceneOutliner, TWeakObjectPtr<UWorld> InSpecifiedWorldToDisplay)
-		: FActorModeInteractive(InSceneOutliner, true, InSpecifiedWorldToDisplay)
-		, FilteredActorCount(0)
-	{
-		// Capture selection changes of bones from mesh selection in fracture tools
-		FSceneOutlinerDelegates::Get().OnComponentsUpdated.AddRaw(this, &FActorBrowsingMode::OnComponentsUpdated);
-
-		GEngine->OnLevelActorDeleted().AddRaw(this, &FActorBrowsingMode::OnLevelActorDeleted);
-
-		FEditorDelegates::OnEditCutActorsBegin.AddRaw(this, &FActorBrowsingMode::OnEditCutActorsBegin);
-		FEditorDelegates::OnEditCutActorsEnd.AddRaw(this, &FActorBrowsingMode::OnEditCutActorsEnd);
-		FEditorDelegates::OnEditCopyActorsBegin.AddRaw(this, &FActorBrowsingMode::OnEditCopyActorsBegin);
-		FEditorDelegates::OnEditCopyActorsEnd.AddRaw(this, &FActorBrowsingMode::OnEditCopyActorsEnd);
-		FEditorDelegates::OnEditPasteActorsBegin.AddRaw(this, &FActorBrowsingMode::OnEditPasteActorsBegin);
-		FEditorDelegates::OnEditPasteActorsEnd.AddRaw(this, &FActorBrowsingMode::OnEditPasteActorsEnd);
-		FEditorDelegates::OnDuplicateActorsBegin.AddRaw(this, &FActorBrowsingMode::OnDuplicateActorsBegin);
-		FEditorDelegates::OnDuplicateActorsEnd.AddRaw(this, &FActorBrowsingMode::OnDuplicateActorsEnd);
-		FEditorDelegates::OnDeleteActorsBegin.AddRaw(this, &FActorBrowsingMode::OnDeleteActorsBegin);
-		FEditorDelegates::OnDeleteActorsEnd.AddRaw(this, &FActorBrowsingMode::OnDeleteActorsEnd);
-
-		UActorBrowsingModeSettings* SharedSettings = GetMutableDefault<UActorBrowsingModeSettings>();
-		// Get the OutlinerModule to register FilterInfos with the FilterInfoMap
-		FOutlinerFilterInfo ShowOnlySelectedActorsInfo(LOCTEXT("ToggleShowOnlySelected", "Only Selected"), LOCTEXT("ToggleShowOnlySelectedToolTip", "When enabled, only displays actors that are currently selected."), SharedSettings->bShowOnlySelectedActors, FCreateSceneOutlinerFilter::CreateStatic(&FActorBrowsingMode::CreateShowOnlySelectedActorsFilter));
-		ShowOnlySelectedActorsInfo.OnToggle().AddLambda([this](bool bIsActive)
-			{
-				UActorBrowsingModeSettings* Settings = GetMutableDefault<UActorBrowsingModeSettings>();
-				Settings->bShowOnlySelectedActors = bIsActive;
-				Settings->PostEditChange();
-			});
-		FilterInfoMap.Add(TEXT("ShowOnlySelectedActors"), ShowOnlySelectedActorsInfo);
-
-		FOutlinerFilterInfo HideTemporaryActorsInfo(LOCTEXT("ToggleHideTemporaryActors", "Hide Temporary Actors"), LOCTEXT("ToggleHideTemporaryActorsToolTip", "When enabled, hides temporary/run-time Actors."), SharedSettings->bHideTemporaryActors, FCreateSceneOutlinerFilter::CreateStatic(&FActorBrowsingMode::CreateHideTemporaryActorsFilter));
-		HideTemporaryActorsInfo.OnToggle().AddLambda([this](bool bIsActive)
+			if (auto ActorHierarchy = StaticCast<FActorHierarchy*>(Hierarchy.Get()))
 			{
-				UActorBrowsingModeSettings* Settings = GetMutableDefault<UActorBrowsingModeSettings>();
-				Settings->bHideTemporaryActors = bIsActive;
-				Settings->PostEditChange();
-			});
-		FilterInfoMap.Add(TEXT("HideTemporaryActors"), HideTemporaryActorsInfo);
-
-		FOutlinerFilterInfo OnlyCurrentLevelInfo(LOCTEXT("ToggleShowOnlyCurrentLevel", "Only in Current Level"), LOCTEXT("ToggleShowOnlyCurrentLevelToolTip", "When enabled, only shows Actors that are in the Current Level."), SharedSettings->bShowOnlyActorsInCurrentLevel, FCreateSceneOutlinerFilter::CreateStatic(&FActorBrowsingMode::CreateIsInCurrentLevelFilter));
-		OnlyCurrentLevelInfo.OnToggle().AddLambda([this](bool bIsActive)
-			{
-				UActorBrowsingModeSettings* Settings = GetMutableDefault<UActorBrowsingModeSettings>();
-				Settings->bShowOnlyActorsInCurrentLevel = bIsActive;
-				Settings->PostEditChange();
-			});
-		FilterInfoMap.Add(TEXT("ShowOnlyCurrentLevel"), OnlyCurrentLevelInfo);
-
-		bHideComponents = SharedSettings->bHideActorComponents;
-		FOutlinerFilterInfo HideComponentsInfo(LOCTEXT("ToggleHideActorComponents", "Hide Actor Components"), LOCTEXT("ToggleHideActorComponentsToolTip", "When Enabled, hides components belonging to actors."), SharedSettings->bHideActorComponents, FCreateSceneOutlinerFilter::CreateStatic(&FActorBrowsingMode::CreateHideComponentsFilter));
-		HideComponentsInfo.OnToggle().AddLambda([this](bool bIsActive)
-			{
-				UActorBrowsingModeSettings* Settings = GetMutableDefault<UActorBrowsingModeSettings>();
-				Settings->bHideActorComponents = bHideComponents = bIsActive;
-				Settings->PostEditChange();
-
-				if (auto ActorHierarchy = StaticCast<FActorHierarchy*>(Hierarchy.Get()))
-				{
-					ActorHierarchy->SetShowingComponents(!bIsActive);
-				}
-			});
+				ActorHierarchy->SetShowingComponents(!bIsActive);
+			}
+		});
 
-		FilterInfoMap.Add(TEXT("HideComponentsFilter"), HideComponentsInfo);
+	FilterInfoMap.Add(TEXT("HideComponentsFilter"), HideComponentsInfo);
 
-		Rebuild();
-	}
+	Rebuild();
+}
 
-	FActorBrowsingMode::~FActorBrowsingMode()
-	{
-		FSceneOutlinerDelegates::Get().OnComponentsUpdated.RemoveAll(this);
+FActorBrowsingMode::~FActorBrowsingMode()
+{
+	FSceneOutlinerDelegates::Get().OnComponentsUpdated.RemoveAll(this);
 
-		GEngine->OnLevelActorDeleted().RemoveAll(this);
+	GEngine->OnLevelActorDeleted().RemoveAll(this);
 		
-		FEditorDelegates::OnEditCutActorsBegin.RemoveAll(this);
-		FEditorDelegates::OnEditCutActorsEnd.RemoveAll(this);
-		FEditorDelegates::OnEditCopyActorsBegin.RemoveAll(this);
-		FEditorDelegates::OnEditCopyActorsEnd.RemoveAll(this);
-		FEditorDelegates::OnEditPasteActorsBegin.RemoveAll(this);
-		FEditorDelegates::OnEditPasteActorsEnd.RemoveAll(this);
-		FEditorDelegates::OnDuplicateActorsBegin.RemoveAll(this);
-		FEditorDelegates::OnDuplicateActorsEnd.RemoveAll(this);
-		FEditorDelegates::OnDeleteActorsBegin.RemoveAll(this);
-		FEditorDelegates::OnDeleteActorsEnd.RemoveAll(this);
-	}
+	FEditorDelegates::OnEditCutActorsBegin.RemoveAll(this);
+	FEditorDelegates::OnEditCutActorsEnd.RemoveAll(this);
+	FEditorDelegates::OnEditCopyActorsBegin.RemoveAll(this);
+	FEditorDelegates::OnEditCopyActorsEnd.RemoveAll(this);
+	FEditorDelegates::OnEditPasteActorsBegin.RemoveAll(this);
+	FEditorDelegates::OnEditPasteActorsEnd.RemoveAll(this);
+	FEditorDelegates::OnDuplicateActorsBegin.RemoveAll(this);
+	FEditorDelegates::OnDuplicateActorsEnd.RemoveAll(this);
+	FEditorDelegates::OnDeleteActorsBegin.RemoveAll(this);
+	FEditorDelegates::OnDeleteActorsEnd.RemoveAll(this);
+}
 
-	void FActorBrowsingMode::Rebuild()
-	{
-		FActorMode::Rebuild();
+void FActorBrowsingMode::Rebuild()
+{
+	FActorMode::Rebuild();
 
-		FilteredActorCount = 0;
-		ApplicableActors.Empty();
-	}
+	FilteredActorCount = 0;
+	ApplicableActors.Empty();
+}
 
-	FText FActorBrowsingMode::GetStatusText() const 
+FText FActorBrowsingMode::GetStatusText() const 
+{
+	if (!RepresentingWorld.IsValid())
 	{
-		if (!RepresentingWorld.IsValid())
-		{
-			return FText();
-		}
+		return FText();
+	}
 
-		const int32 TotalActorCount = ApplicableActors.Num();
-		const int32 SelectedActorCount = SceneOutliner->GetSelection().Num<FActorTreeItem>();
+	const int32 TotalActorCount = ApplicableActors.Num();
+	const int32 SelectedActorCount = SceneOutliner->GetSelection().Num<FActorTreeItem>();
 
-		if (!SceneOutliner->IsTextFilterActive())
-		{
-			if (SelectedActorCount == 0) //-V547
-			{
-				return FText::Format(LOCTEXT("ShowingAllActorsFmt", "{0} actors"), FText::AsNumber(FilteredActorCount));
-			}
-			else
-			{
-				return FText::Format(LOCTEXT("ShowingAllActorsSelectedFmt", "{0} actors ({1} selected)"), FText::AsNumber(FilteredActorCount), FText::AsNumber(SelectedActorCount));
-			}
-		}
-		else if (SceneOutliner->IsTextFilterActive() && FilteredActorCount == 0)
-		{
-			return FText::Format(LOCTEXT("ShowingNoActorsFmt", "No matching actors ({0} total)"), FText::AsNumber(TotalActorCount));
-		}
-		else if (SelectedActorCount != 0) //-V547
+	if (!SceneOutliner->IsTextFilterActive())
+	{
+		if (SelectedActorCount == 0) //-V547
 		{
-			return FText::Format(LOCTEXT("ShowingOnlySomeActorsSelectedFmt", "Showing {0} of {1} actors ({2} selected)"), FText::AsNumber(FilteredActorCount), FText::AsNumber(TotalActorCount), FText::AsNumber(SelectedActorCount));
+			return FText::Format(LOCTEXT("ShowingAllActorsFmt", "{0} actors"), FText::AsNumber(FilteredActorCount));
 		}
 		else
 		{
-			return FText::Format(LOCTEXT("ShowingOnlySomeActorsFmt", "Showing {0} of {1} actors"), FText::AsNumber(FilteredActorCount), FText::AsNumber(TotalActorCount));
+			return FText::Format(LOCTEXT("ShowingAllActorsSelectedFmt", "{0} actors ({1} selected)"), FText::AsNumber(FilteredActorCount), FText::AsNumber(SelectedActorCount));
 		}
 	}
-
-	FSlateColor FActorBrowsingMode::GetStatusTextColor() const
+	else if (SceneOutliner->IsTextFilterActive() && FilteredActorCount == 0)
 	{
-		if (!SceneOutliner->IsTextFilterActive())
-		{
-			return FSlateColor::UseForeground();
-		}
-		else if (FilteredActorCount == 0)
-		{
-			return FAppStyle::Get().GetSlateColor("Colors.AccentRed");
-		}
-		else
-		{
-			return FAppStyle::Get().GetSlateColor("Colors.AccentGreen");
-		}
+		return FText::Format(LOCTEXT("ShowingNoActorsFmt", "No matching actors ({0} total)"), FText::AsNumber(TotalActorCount));
 	}
-
-	void FActorBrowsingMode::CreateViewContent(FMenuBuilder& MenuBuilder)
+	else if (SelectedActorCount != 0) //-V547
 	{
-		MenuBuilder.BeginSection("AssetThumbnails", LOCTEXT("ShowWorldHeading", "World"));
-		{
-			MenuBuilder.AddSubMenu(
-				LOCTEXT("ChooseWorldSubMenu", "Choose World"),
-				LOCTEXT("ChooseWorldSubMenuToolTip", "Choose the world to display in the outliner."),
-				FNewMenuDelegate::CreateRaw(this, &FActorMode::BuildWorldPickerMenu)
-			);
-		}
-		MenuBuilder.EndSection();
+		return FText::Format(LOCTEXT("ShowingOnlySomeActorsSelectedFmt", "Showing {0} of {1} actors ({2} selected)"), FText::AsNumber(FilteredActorCount), FText::AsNumber(TotalActorCount), FText::AsNumber(SelectedActorCount));
 	}
-
-	TSharedRef<FOutlinerFilter> FActorBrowsingMode::CreateShowOnlySelectedActorsFilter()
+	else
 	{
-		auto IsActorSelected = [](const AActor* InActor)
-		{
-			return InActor && InActor->IsSelected();
-		};
-		return MakeShareable(new FActorFilter(FActorTreeItem::FFilterPredicate::CreateStatic(IsActorSelected), EDefaultFilterBehaviour::Fail, FActorTreeItem::FFilterPredicate::CreateStatic(IsActorSelected)));
+		return FText::Format(LOCTEXT("ShowingOnlySomeActorsFmt", "Showing {0} of {1} actors"), FText::AsNumber(FilteredActorCount), FText::AsNumber(TotalActorCount));
 	}
+}
 
-	TSharedRef<FOutlinerFilter> FActorBrowsingMode::CreateHideTemporaryActorsFilter()
+FSlateColor FActorBrowsingMode::GetStatusTextColor() const
+{
+	if (!SceneOutliner->IsTextFilterActive())
 	{
-		return MakeShareable(new FActorFilter(FActorTreeItem::FFilterPredicate::CreateStatic([](const AActor* InActor)
-			{
-				return ((InActor->GetWorld() && InActor->GetWorld()->WorldType != EWorldType::PIE) || GEditor->ObjectsThatExistInEditorWorld.Get(InActor)) && !InActor->HasAnyFlags(EObjectFlags::RF_Transient);
-			}), EDefaultFilterBehaviour::Pass));
+		return FSlateColor::UseForeground();
 	}
-
-	TSharedRef<FOutlinerFilter> FActorBrowsingMode::CreateIsInCurrentLevelFilter()
+	else if (FilteredActorCount == 0)
 	{
-		return MakeShareable(new FActorFilter(FActorTreeItem::FFilterPredicate::CreateStatic([](const AActor* InActor)
-			{
-				if (InActor->GetWorld())
-				{
-					return InActor->GetLevel() == InActor->GetWorld()->GetCurrentLevel();
-				}
+		return FAppStyle::Get().GetSlateColor("Colors.AccentRed");
+	}
+	else
+	{
+		return FAppStyle::Get().GetSlateColor("Colors.AccentGreen");
+	}
+}
 
-				return false;
-			}), EDefaultFilterBehaviour::Pass));
+void FActorBrowsingMode::CreateViewContent(FMenuBuilder& MenuBuilder)
+{
+	MenuBuilder.BeginSection("AssetThumbnails", LOCTEXT("ShowWorldHeading", "World"));
+	{
+		MenuBuilder.AddSubMenu(
+			LOCTEXT("ChooseWorldSubMenu", "Choose World"),
+			LOCTEXT("ChooseWorldSubMenuToolTip", "Choose the world to display in the outliner."),
+			FNewMenuDelegate::CreateRaw(this, &FActorMode::BuildWorldPickerMenu)
+		);
 	}
+	MenuBuilder.EndSection();
+}
 
-	TSharedRef<FOutlinerFilter> FActorBrowsingMode::CreateHideComponentsFilter()
+TSharedRef<FSceneOutlinerFilter> FActorBrowsingMode::CreateShowOnlySelectedActorsFilter()
+{
+	auto IsActorSelected = [](const AActor* InActor)
 	{
-		return MakeShared<TOutlinerPredicateFilter<FComponentTreeItem>>(TOutlinerPredicateFilter<FComponentTreeItem>(FComponentTreeItem::FFilterPredicate::CreateStatic([](const UActorComponent* Component)
+		return InActor && InActor->IsSelected();
+	};
+	return MakeShareable(new FActorFilter(FActorTreeItem::FFilterPredicate::CreateStatic(IsActorSelected), FSceneOutlinerFilter::EDefaultBehaviour::Fail, FActorTreeItem::FFilterPredicate::CreateStatic(IsActorSelected)));
+}
+
+TSharedRef<FSceneOutlinerFilter> FActorBrowsingMode::CreateHideTemporaryActorsFilter()
+{
+	return MakeShareable(new FActorFilter(FActorTreeItem::FFilterPredicate::CreateStatic([](const AActor* InActor)
+		{
+			return ((InActor->GetWorld() && InActor->GetWorld()->WorldType != EWorldType::PIE) || GEditor->ObjectsThatExistInEditorWorld.Get(InActor)) && !InActor->HasAnyFlags(EObjectFlags::RF_Transient);
+		}), FSceneOutlinerFilter::EDefaultBehaviour::Pass));
+}
+
+TSharedRef<FSceneOutlinerFilter> FActorBrowsingMode::CreateIsInCurrentLevelFilter()
+{
+	return MakeShareable(new FActorFilter(FActorTreeItem::FFilterPredicate::CreateStatic([](const AActor* InActor)
+		{
+			if (InActor->GetWorld())
 			{
-				// Filter out all components
-				return false;
-			}), EDefaultFilterBehaviour::Pass));
-	}
+				return InActor->GetLevel() == InActor->GetWorld()->GetCurrentLevel();
+			}
+
+			return false;
+		}), FSceneOutlinerFilter::EDefaultBehaviour::Pass));
+}
+
+TSharedRef<FSceneOutlinerFilter> FActorBrowsingMode::CreateHideComponentsFilter()
+{
+	return MakeShared<TSceneOutlinerPredicateFilter<FComponentTreeItem>>(TSceneOutlinerPredicateFilter<FComponentTreeItem>(FComponentTreeItem::FFilterPredicate::CreateStatic([](const UActorComponent* Component)
+		{
+			// Filter out all components
+			return false;
+		}), FSceneOutlinerFilter::EDefaultBehaviour::Pass));
+}
+
+static const FName DefaultContextBaseMenuName("SceneOutliner.DefaultContextMenuBase");
+static const FName DefaultContextMenuName("SceneOutliner.DefaultContextMenu");
 
-	static const FName DefaultContextBaseMenuName("SceneOutliner.DefaultContextMenuBase");
-	static const FName DefaultContextMenuName("SceneOutliner.DefaultContextMenu");
+void FActorBrowsingMode::RegisterContextMenu()
+{
+	UToolMenus* ToolMenus = UToolMenus::Get();
 
-	void FActorBrowsingMode::RegisterContextMenu()
+	if (!ToolMenus->IsMenuRegistered(DefaultContextBaseMenuName))
 	{
-		UToolMenus* ToolMenus = UToolMenus::Get();
+		UToolMenu* Menu = ToolMenus->RegisterMenu(DefaultContextBaseMenuName);
 
-		if (!ToolMenus->IsMenuRegistered(DefaultContextBaseMenuName))
-		{
-			UToolMenu* Menu = ToolMenus->RegisterMenu(DefaultContextBaseMenuName);
+		Menu->AddDynamicSection("DynamicSection1", FNewToolMenuDelegate::CreateLambda([](UToolMenu* InMenu)
+			{
+				USceneOutlinerMenuContext* Context = InMenu->FindContext<USceneOutlinerMenuContext>();
+				if (!Context || !Context->SceneOutliner.IsValid())
+				{
+					return;
+				}
 
-			Menu->AddDynamicSection("DynamicSection1", FNewToolMenuDelegate::CreateLambda([](UToolMenu* InMenu)
+				SSceneOutliner* SceneOutliner = Context->SceneOutliner.Pin().Get();
+				if (Context->bShowParentTree)
 				{
-					USceneOutlinerMenuContext* Context = InMenu->FindContext<USceneOutlinerMenuContext>();
-					if (!Context || !Context->SceneOutliner.IsValid())
+					if (Context->NumSelectedItems == 0)
 					{
-						return;
+						InMenu->FindOrAddSection("Section").AddMenuEntry(
+							"CreateFolder",
+							LOCTEXT("CreateFolder", "Create Folder"),
+							FText(),
+							FSlateIcon(FEditorStyle::GetStyleSetName(), "SceneOutliner.NewFolderIcon"),
+							FUIAction(FExecuteAction::CreateSP(SceneOutliner, &SSceneOutliner::CreateFolder)));
 					}
-
-					SSceneOutliner* SceneOutliner = Context->SceneOutliner.Pin().Get();
-					if (Context->bShowParentTree)
+					else
 					{
-						if (Context->NumSelectedItems == 0)
+						if (Context->NumSelectedItems == 1)
 						{
-							InMenu->FindOrAddSection("Section").AddMenuEntry(
-								"CreateFolder",
-								LOCTEXT("CreateFolder", "Create Folder"),
-								FText(),
-								FSlateIcon(FEditorStyle::GetStyleSetName(), "SceneOutliner.NewFolderIcon"),
-								FUIAction(FExecuteAction::CreateSP(SceneOutliner, &SSceneOutliner::CreateFolder)));
+							SceneOutliner->GetTree().GetSelectedItems()[0]->GenerateContextMenu(InMenu, *SceneOutliner);
 						}
-						else
+
+						// If we've only got folders selected, show the selection and edit sub menus
+						if (Context->NumSelectedItems > 0 && Context->NumSelectedFolders == Context->NumSelectedItems)
 						{
-							if (Context->NumSelectedItems == 1)
-							{
-								SceneOutliner->GetTree().GetSelectedItems()[0]->GenerateContextMenu(InMenu, *SceneOutliner);
-							}
-
-							// If we've only got folders selected, show the selection and edit sub menus
-							if (Context->NumSelectedItems > 0 && Context->NumSelectedFolders == Context->NumSelectedItems)
-							{
-								InMenu->FindOrAddSection("Section").AddSubMenu(
-									"SelectSubMenu",
-									LOCTEXT("SelectSubmenu", "Select"),
-									LOCTEXT("SelectSubmenu_Tooltip", "Select the contents of the current selection"),
-									FNewToolMenuDelegate::CreateSP(SceneOutliner, &SSceneOutliner::FillSelectionSubMenu));
-							}
+							InMenu->FindOrAddSection("Section").AddSubMenu(
+								"SelectSubMenu",
+								LOCTEXT("SelectSubmenu", "Select"),
+								LOCTEXT("SelectSubmenu_Tooltip", "Select the contents of the current selection"),
+								FNewToolMenuDelegate::CreateSP(SceneOutliner, &SSceneOutliner::FillSelectionSubMenu));
 						}
 					}
-				}));
+				}
+			}));
 
-			Menu->AddDynamicSection("DynamicMainSection", FNewToolMenuDelegate::CreateLambda([](UToolMenu* InMenu)
-				{
-					// We always create a section here, even if there is no parent so that clients can still extend the menu
-					FToolMenuSection& Section = InMenu->AddSection("MainSection");
+		Menu->AddDynamicSection("DynamicMainSection", FNewToolMenuDelegate::CreateLambda([](UToolMenu* InMenu)
+			{
+				// We always create a section here, even if there is no parent so that clients can still extend the menu
+				FToolMenuSection& Section = InMenu->AddSection("MainSection");
 
-					if (USceneOutlinerMenuContext* Context = InMenu->FindContext<USceneOutlinerMenuContext>())
+				if (USceneOutlinerMenuContext* Context = InMenu->FindContext<USceneOutlinerMenuContext>())
+				{
+					// Don't add any of these menu items if we're not showing the parent tree
+					// Can't move worlds or level blueprints
+					if (Context->bShowParentTree && Context->NumSelectedItems > 0 && Context->NumWorldsSelected == 0 && Context->SceneOutliner.IsValid())
 					{
-						// Don't add any of these menu items if we're not showing the parent tree
-						// Can't move worlds or level blueprints
-						if (Context->bShowParentTree && Context->NumSelectedItems > 0 && Context->NumWorldsSelected == 0 && Context->SceneOutliner.IsValid())
-						{
-							Section.AddSubMenu(
-								"MoveActorsTo",
-								LOCTEXT("MoveActorsTo", "Move To"),
-								LOCTEXT("MoveActorsTo_Tooltip", "Move selection to another folder"),
-								FNewToolMenuDelegate::CreateSP(Context->SceneOutliner.Pin().Get(), &SSceneOutliner::FillFoldersSubMenu));
-						}
+						Section.AddSubMenu(
+							"MoveActorsTo",
+							LOCTEXT("MoveActorsTo", "Move To"),
+							LOCTEXT("MoveActorsTo_Tooltip", "Move selection to another folder"),
+							FNewToolMenuDelegate::CreateSP(Context->SceneOutliner.Pin().Get(), &SSceneOutliner::FillFoldersSubMenu));
 					}
-				}));
-		}
-
-		if (!ToolMenus->IsMenuRegistered(DefaultContextMenuName))
-		{
-			ToolMenus->RegisterMenu(DefaultContextMenuName, DefaultContextBaseMenuName);
-		}
+				}
+			}));
 	}
 
-	TSharedPtr<SWidget> FActorBrowsingMode::BuildContextMenu()
+	if (!ToolMenus->IsMenuRegistered(DefaultContextMenuName))
 	{
-		RegisterContextMenu();
+		ToolMenus->RegisterMenu(DefaultContextMenuName, DefaultContextBaseMenuName);
+	}
+}
 
-		FItemSelection ItemSelection(SceneOutliner->GetSelection());
+TSharedPtr<SWidget> FActorBrowsingMode::BuildContextMenu()
+{
+	RegisterContextMenu();
+
+	FSceneOutlinerItemSelection ItemSelection(SceneOutliner->GetSelection());
 
-		USceneOutlinerMenuContext* ContextObject = NewObject<USceneOutlinerMenuContext>();
-		ContextObject->SceneOutliner = StaticCastSharedRef<SSceneOutliner>(SceneOutliner->AsShared());
-		ContextObject->bShowParentTree = SceneOutliner->GetSharedData().bShowParentTree;
-		ContextObject->NumSelectedItems = ItemSelection.Num();
-		ContextObject->NumSelectedFolders = ItemSelection.Num<FFolderTreeItem>();
-		ContextObject->NumWorldsSelected = ItemSelection.Num<FWorldTreeItem>();
-		FToolMenuContext Context(ContextObject);
+	USceneOutlinerMenuContext* ContextObject = NewObject<USceneOutlinerMenuContext>();
+	ContextObject->SceneOutliner = StaticCastSharedRef<SSceneOutliner>(SceneOutliner->AsShared());
+	ContextObject->bShowParentTree = SceneOutliner->GetSharedData().bShowParentTree;
+	ContextObject->NumSelectedItems = ItemSelection.Num();
+	ContextObject->NumSelectedFolders = ItemSelection.Num<FFolderTreeItem>();
+	ContextObject->NumWorldsSelected = ItemSelection.Num<FWorldTreeItem>();
+	FToolMenuContext Context(ContextObject);
 
-		FName MenuName = DefaultContextMenuName;
-		SceneOutliner->GetSharedData().ModifyContextMenu.ExecuteIfBound(MenuName, Context);
+	FName MenuName = DefaultContextMenuName;
+	SceneOutliner->GetSharedData().ModifyContextMenu.ExecuteIfBound(MenuName, Context);
 
-		// Build up the menu for a selection
-		UToolMenus* ToolMenus = UToolMenus::Get();
-		UToolMenu* Menu = ToolMenus->GenerateMenu(MenuName, Context);
+	// Build up the menu for a selection
+	UToolMenus* ToolMenus = UToolMenus::Get();
+	UToolMenu* Menu = ToolMenus->GenerateMenu(MenuName, Context);
 
-		for (const FToolMenuSection& Section : Menu->Sections)
+	for (const FToolMenuSection& Section : Menu->Sections)
+	{
+		if (Section.Blocks.Num() > 0)
 		{
-			if (Section.Blocks.Num() > 0)
-			{
-				return ToolMenus->GenerateWidget(Menu);
-			}
+			return ToolMenus->GenerateWidget(Menu);
 		}
-
-		return nullptr;
 	}
 
-	TSharedPtr<SWidget> FActorBrowsingMode::CreateContextMenu()
-	{
-		TArray<AActor*> SelectedActors;
-		GEditor->GetSelectedActors()->GetSelectedObjects<AActor>(SelectedActors);
+	return nullptr;
+}
 
-		// Make sure that no components are selected
-		if (GEditor->GetSelectedComponentCount() > 0)
-		{
-			// We want to be able to undo to regain the previous component selection
-			const FScopedTransaction Transaction(NSLOCTEXT("UnrealEd", "ClickingOnActorsContextMenu", "Clicking on Actors (context menu)"));
-			USelection* ComponentSelection = GEditor->GetSelectedComponents();
-			ComponentSelection->Modify(false);
-			ComponentSelection->DeselectAll();
-
-			GUnrealEd->UpdatePivotLocationForSelection();
-			GEditor->RedrawLevelEditingViewports(false);
-		}
+TSharedPtr<SWidget> FActorBrowsingMode::CreateContextMenu()
+{
+	TArray<AActor*> SelectedActors;
+	GEditor->GetSelectedActors()->GetSelectedObjects<AActor>(SelectedActors);
 
-		return BuildContextMenu();
+	// Make sure that no components are selected
+	if (GEditor->GetSelectedComponentCount() > 0)
+	{
+		// We want to be able to undo to regain the previous component selection
+		const FScopedTransaction Transaction(NSLOCTEXT("UnrealEd", "ClickingOnActorsContextMenu", "Clicking on Actors (context menu)"));
+		USelection* ComponentSelection = GEditor->GetSelectedComponents();
+		ComponentSelection->Modify(false);
+		ComponentSelection->DeselectAll();
+
+		GUnrealEd->UpdatePivotLocationForSelection();
+		GEditor->RedrawLevelEditingViewports(false);
 	}
 
-	void FActorBrowsingMode::OnItemAdded(FTreeItemPtr Item)
+	return BuildContextMenu();
+}
+
+void FActorBrowsingMode::OnItemAdded(FSceneOutlinerTreeItemPtr Item)
+{
+	if (const FActorTreeItem* ActorItem = Item->CastTo<FActorTreeItem>())
 	{
-		if (const FActorTreeItem* ActorItem = Item->CastTo<FActorTreeItem>())
+		if (!Item->Flags.bIsFilteredOut)
 		{
-			if (!Item->Flags.bIsFilteredOut)
-			{
-				++FilteredActorCount;
+			++FilteredActorCount;
 
-				// Synchronize selection
-				if (GEditor->GetSelectedActors()->IsSelected(ActorItem->Actor.Get()))
-				{
-					SceneOutliner->SetItemSelection(Item, true);
-				}
-			}
-		}
-		else if (FActorFolderTreeItem* FolderItem = Item->CastTo<FActorFolderTreeItem>())
-		{
-			if (FolderItem->World.IsValid())
+			// Synchronize selection
+			if (GEditor->GetSelectedActors()->IsSelected(ActorItem->Actor.Get()))
 			{
-				if (FActorFolderProps* Props = FActorFolders::Get().GetFolderProperties(*FolderItem->World, FolderItem->Path))
-				{
-					FolderItem->Flags.bIsExpanded = Props->bIsExpanded;
-				}
+				SceneOutliner->SetItemSelection(Item, true);
 			}
 		}
 	}
-
-	void FActorBrowsingMode::OnItemRemoved(FTreeItemPtr Item)
+	else if (FActorFolderTreeItem* FolderItem = Item->CastTo<FActorFolderTreeItem>())
 	{
-		if (const FActorTreeItem* ActorItem = Item->CastTo<FActorTreeItem>())
+		if (FolderItem->World.IsValid())
 		{
-			if (!ActorItem->Flags.bIsFilteredOut)
+			if (FActorFolderProps* Props = FActorFolders::Get().GetFolderProperties(*FolderItem->World, FolderItem->Path))
 			{
-				--FilteredActorCount;
+				FolderItem->Flags.bIsExpanded = Props->bIsExpanded;
 			}
 		}
 	}
+}
 
-	void FActorBrowsingMode::OnComponentsUpdated()
-	{
-		SceneOutliner->FullRefresh();
-	}
-
-	void FActorBrowsingMode::OnLevelActorDeleted(AActor* Actor)
-	{
-		ApplicableActors.Remove(Actor);
-	}
-
-	void FActorBrowsingMode::OnItemSelectionChanged(FTreeItemPtr TreeItem, ESelectInfo::Type SelectionType, const FItemSelection& Selection)
+void FActorBrowsingMode::OnItemRemoved(FSceneOutlinerTreeItemPtr Item)
+{
+	if (const FActorTreeItem* ActorItem = Item->CastTo<FActorTreeItem>())
 	{
-		TArray<AActor*> SelectedActors = Selection.GetData<AActor*>(FActorSelector());
-
-		bool bChanged = false;
-		bool bAnyInPIE = false;
-		for (auto* Actor : SelectedActors)
+		if (!ActorItem->Flags.bIsFilteredOut)
 		{
-			if (!bAnyInPIE && Actor && Actor->GetOutermost()->HasAnyPackageFlags(PKG_PlayInEditor))
-			{
-				bAnyInPIE = true;
-			}
-			if (!GEditor->GetSelectedActors()->IsSelected(Actor))
-			{
-				bChanged = true;
-				break;
-			}
+			--FilteredActorCount;
 		}
+	}
+}
 
-		for (FSelectionIterator SelectionIt(*GEditor->GetSelectedActors()); SelectionIt && !bChanged; ++SelectionIt)
-		{
-			const AActor* Actor = CastChecked< AActor >(*SelectionIt);
-			if (!bAnyInPIE && Actor->GetOutermost()->HasAnyPackageFlags(PKG_PlayInEditor))
-			{
-				bAnyInPIE = true;
-			}
-			if (!SelectedActors.Contains(Actor))
-			{
-				// Actor has been deselected
-				bChanged = true;
+void FActorBrowsingMode::OnComponentsUpdated()
+{
+	SceneOutliner->FullRefresh();
+}
 
-				// If actor was a group actor, remove its members from the ActorsToSelect list
-				const AGroupActor* DeselectedGroupActor = Cast<AGroupActor>(Actor);
-				if (DeselectedGroupActor)
-				{
-					TArray<AActor*> GroupActors;
-					DeselectedGroupActor->GetGroupActors(GroupActors);
+void FActorBrowsingMode::OnLevelActorDeleted(AActor* Actor)
+{
+	ApplicableActors.Remove(Actor);
+}
 
-					for (auto* GroupActor : GroupActors)
-					{
-						SelectedActors.Remove(GroupActor);
-					}
+void FActorBrowsingMode::OnItemSelectionChanged(FSceneOutlinerTreeItemPtr TreeItem, ESelectInfo::Type SelectionType, const FSceneOutlinerItemSelection& Selection)
+{
+	TArray<AActor*> SelectedActors = Selection.GetData<AActor*>(SceneOutliner::FActorSelector());
 
-				}
-			}
+	bool bChanged = false;
+	bool bAnyInPIE = false;
+	for (auto* Actor : SelectedActors)
+	{
+		if (!bAnyInPIE && Actor && Actor->GetOutermost()->HasAnyPackageFlags(PKG_PlayInEditor))
+		{
+			bAnyInPIE = true;
 		}
-
-		// If there's a discrepancy, update the selected actors to reflect this list.
-		if (bChanged)
+		if (!GEditor->GetSelectedActors()->IsSelected(Actor))
 		{
-			const FScopedTransaction Transaction(NSLOCTEXT("UnrealEd", "ClickingOnActors", "Clicking on Actors"), !bAnyInPIE);
-			GEditor->GetSelectedActors()->Modify();
-
-			// Clear the selection.
-			GEditor->SelectNone(false, true, true);
-
-			// We'll batch selection changes instead by using BeginBatchSelectOperation()
-			GEditor->GetSelectedActors()->BeginBatchSelectOperation();
-
-			const bool bShouldSelect = true;
-			const bool bNotifyAfterSelect = false;
-			const bool bSelectEvenIfHidden = true;	// @todo outliner: Is this actually OK?
-			for (auto* Actor : SelectedActors)
-			{
-				//UE_LOG(LogSceneOutliner, Verbose, TEXT("Clicking on Actor (world outliner): %s (%s)"), *Actor->GetClass()->GetName(), *Actor->GetActorLabel());
-				GEditor->SelectActor(Actor, bShouldSelect, bNotifyAfterSelect, bSelectEvenIfHidden);
-			}
-
-			// Commit selection changes
-			GEditor->GetSelectedActors()->EndBatchSelectOperation(/*bNotify*/false);
-
-			// Fire selection changed event
-			GEditor->NoteSelectionChange();
+			bChanged = true;
+			break;
 		}
-
-		SceneOutliner->RefreshSelection();
 	}
 
-	void FActorBrowsingMode::OnItemDoubleClick(FTreeItemPtr Item)
+	for (FSelectionIterator SelectionIt(*GEditor->GetSelectedActors()); SelectionIt && !bChanged; ++SelectionIt)
 	{
-		if (Item->CanInteract())
+		const AActor* Actor = CastChecked< AActor >(*SelectionIt);
+		if (!bAnyInPIE && Actor->GetOutermost()->HasAnyPackageFlags(PKG_PlayInEditor))
 		{
-			if (Item->IsA<FActorTreeItem>())
-			{
-				FItemSelection Selection(SceneOutliner->GetSelection());
-				if (Selection.Has<FActorTreeItem>())
-				{
-					const bool bActiveViewportOnly = false;
-					GEditor->MoveViewportCamerasToActor(Selection.GetData<AActor*>(FActorSelector()), bActiveViewportOnly);
-				}
-			}
+			bAnyInPIE = true;
 		}
-		else
+		if (!SelectedActors.Contains(Actor))
 		{
-			if (const FActorTreeItem* ActorItem = Item->CastTo<FActorTreeItem>())
+			// Actor has been deselected
+			bChanged = true;
+
+			// If actor was a group actor, remove its members from the ActorsToSelect list
+			const AGroupActor* DeselectedGroupActor = Cast<AGroupActor>(Actor);
+			if (DeselectedGroupActor)
 			{
-				if (AActor* Actor = ActorItem->Actor.Get())
+				TArray<AActor*> GroupActors;
+				DeselectedGroupActor->GetGroupActors(GroupActors);
+
+				for (auto* GroupActor : GroupActors)
 				{
-					const bool bActiveViewportOnly = false;
-					GEditor->MoveViewportCamerasToActor(*Actor, bActiveViewportOnly);
+					SelectedActors.Remove(GroupActor);
 				}
-			}
-			else if (Item->IsA<FFolderTreeItem>())
-			{
-				SceneOutliner->SetItemExpansion(Item, !SceneOutliner->IsItemExpanded(Item));
+
 			}
 		}
 	}
 
-	void FActorBrowsingMode::OnFilterTextCommited(FItemSelection& Selection, ETextCommit::Type CommitType)
+	// If there's a discrepancy, update the selected actors to reflect this list.
+	if (bChanged)
 	{
-		// Start batching selection changes
+		const FScopedTransaction Transaction(NSLOCTEXT("UnrealEd", "ClickingOnActors", "Clicking on Actors"), !bAnyInPIE);
+		GEditor->GetSelectedActors()->Modify();
+
+		// Clear the selection.
+		GEditor->SelectNone(false, true, true);
+
+		// We'll batch selection changes instead by using BeginBatchSelectOperation()
 		GEditor->GetSelectedActors()->BeginBatchSelectOperation();
 
-		// Select actors (and only the actors) that match the filter text
-		const bool bNoteSelectionChange = false;
-		const bool bDeselectBSPSurfs = false;
-		const bool WarnAboutManyActors = true;
-		GEditor->SelectNone(bNoteSelectionChange, bDeselectBSPSurfs, WarnAboutManyActors);
-		for (AActor* Actor : Selection.GetData<AActor*>(FActorSelector()))
+		const bool bShouldSelect = true;
+		const bool bNotifyAfterSelect = false;
+		const bool bSelectEvenIfHidden = true;	// @todo outliner: Is this actually OK?
+		for (auto* Actor : SelectedActors)
 		{
-			const bool bShouldSelect = true;
-			const bool bSelectEvenIfHidden = false;
-			GEditor->SelectActor(Actor, bShouldSelect, bNoteSelectionChange, bSelectEvenIfHidden);
+			UE_LOG(LogSceneOutliner, Verbose, TEXT("Clicking on Actor (world outliner): %s (%s)"), *Actor->GetClass()->GetName(), *Actor->GetActorLabel());
+			GEditor->SelectActor(Actor, bShouldSelect, bNotifyAfterSelect, bSelectEvenIfHidden);
 		}
 
 		// Commit selection changes
@@ -544,685 +487,742 @@ namespace SceneOutliner
 
 		// Fire selection changed event
 		GEditor->NoteSelectionChange();
-
-		// Set keyboard focus to the SceneOutliner, so the user can perform keyboard commands that interact
-		// with selected actors (such as Delete, to delete selected actors.)
-		SceneOutliner->SetKeyboardFocus();
 	}
 
-	void FActorBrowsingMode::OnItemPassesFilters(const ITreeItem& Item)
+	SceneOutliner->RefreshSelection();
+}
+
+void FActorBrowsingMode::OnItemDoubleClick(FSceneOutlinerTreeItemPtr Item)
+{
+	if (Item->CanInteract())
 	{
-		if (const FActorTreeItem* ActorItem = Item.CastTo<FActorTreeItem>())
+		if (Item->IsA<FActorTreeItem>())
 		{
-			ApplicableActors.Add(ActorItem->Actor);
+			FSceneOutlinerItemSelection Selection(SceneOutliner->GetSelection());
+			if (Selection.Has<FActorTreeItem>())
+			{
+				const bool bActiveViewportOnly = false;
+				GEditor->MoveViewportCamerasToActor(Selection.GetData<AActor*>(SceneOutliner::FActorSelector()), bActiveViewportOnly);
+			}
 		}
 	}
-
-	FReply FActorBrowsingMode::OnKeyDown(const FKeyEvent& InKeyEvent)
+	else
 	{
-		const FItemSelection& Selection = SceneOutliner->GetSelection();
-
-
-		// Rename key: Rename selected actors (not rebindable, because it doesn't make much sense to bind.)
-		if (InKeyEvent.GetKey() == EKeys::F2)
+		if (const FActorTreeItem* ActorItem = Item->CastTo<FActorTreeItem>())
 		{
-			if (Selection.Num() == 1)
+			if (AActor* Actor = ActorItem->Actor.Get())
 			{
-				FTreeItemPtr ItemToRename = Selection.SelectedItems[0].Pin();
-
-				if (ItemToRename.IsValid() && CanRenameItem(*ItemToRename) && ItemToRename->CanInteract())
-				{
-					SceneOutliner->SetPendingRenameItem(ItemToRename);
-					SceneOutliner->ScrollItemIntoView(ItemToRename);
-				}
-
-				return FReply::Handled();
+				const bool bActiveViewportOnly = false;
+				GEditor->MoveViewportCamerasToActor(*Actor, bActiveViewportOnly);
 			}
 		}
-
-		// F5 forces a full refresh
-		else if (InKeyEvent.GetKey() == EKeys::F5)
-		{
-			SceneOutliner->FullRefresh();
-			return FReply::Handled();
-		}
-
-		// Delete key: Delete selected actors (not rebindable, because it doesn't make much sense to bind.)
-		// Use Delete and Backspace instead of Platform_Delete because the LevelEditor default Edit Delete is bound to both 
-		else if (InKeyEvent.GetKey() == EKeys::Delete || InKeyEvent.GetKey() == EKeys::BackSpace)
+		else if (Item->IsA<FFolderTreeItem>())
 		{
-			if (SceneOutliner->GetSharedData().CustomDelete.IsBound())
-			{
-				SceneOutliner->GetSharedData().CustomDelete.Execute(Selection.SelectedItems);
-			}
-			else
-			{
-				if (RepresentingWorld.IsValid())
-				{
-					GUnrealEd->Exec(RepresentingWorld.Get(), TEXT("DELETE"));
-				}
-			}
-			return FReply::Handled();
-			
+			SceneOutliner->SetItemExpansion(Item, !SceneOutliner->IsItemExpanded(Item));
 		}
-		return FReply::Unhandled();
 	}
+}
 
-	bool FActorBrowsingMode::CanRename() const
+void FActorBrowsingMode::OnFilterTextCommited(FSceneOutlinerItemSelection& Selection, ETextCommit::Type CommitType)
+{
+	// Start batching selection changes
+	GEditor->GetSelectedActors()->BeginBatchSelectOperation();
+
+	// Select actors (and only the actors) that match the filter text
+	const bool bNoteSelectionChange = false;
+	const bool bDeselectBSPSurfs = false;
+	const bool WarnAboutManyActors = true;
+	GEditor->SelectNone(bNoteSelectionChange, bDeselectBSPSurfs, WarnAboutManyActors);
+	for (AActor* Actor : Selection.GetData<AActor*>(SceneOutliner::FActorSelector()))
 	{
-		const FItemSelection ItemSelection = SceneOutliner->GetSelection();
-		const uint32 NumberOfFolders = ItemSelection.Num<FFolderTreeItem>();
-		return (NumberOfFolders == 1 && NumberOfFolders == ItemSelection.Num());
+		const bool bShouldSelect = true;
+		const bool bSelectEvenIfHidden = false;
+		GEditor->SelectActor(Actor, bShouldSelect, bNoteSelectionChange, bSelectEvenIfHidden);
 	}
 
-	bool FActorBrowsingMode::CanRenameItem(const ITreeItem& Item) const
-	{
-		// Can only rename actor and folder items when in actor browsing mode
-		return (Item.IsValid() && (Item.IsA<FActorTreeItem>() || Item.IsA<FFolderTreeItem>()));
-	}
+	// Commit selection changes
+	GEditor->GetSelectedActors()->EndBatchSelectOperation(/*bNotify*/false);
 
-	bool FActorBrowsingMode::CanCut() const
-	{
-		const FItemSelection ItemSelection = SceneOutliner->GetSelection();
-		const uint32 NumberOfFolders = ItemSelection.Num<FFolderTreeItem>();
-		return (NumberOfFolders > 0 && NumberOfFolders == ItemSelection.Num());
-	}
+	// Fire selection changed event
+	GEditor->NoteSelectionChange();
 
-	bool FActorBrowsingMode::CanCopy() const
-	{
-		const FItemSelection ItemSelection = SceneOutliner->GetSelection();
-		const uint32 NumberOfFolders = ItemSelection.Num<FFolderTreeItem>();
-		return (NumberOfFolders > 0 && NumberOfFolders == ItemSelection.Num());
-	}
+	// Set keyboard focus to the SceneOutliner, so the user can perform keyboard commands that interact
+	// with selected actors (such as Delete, to delete selected actors.)
+	SceneOutliner->SetKeyboardFocus();
+}
 
-	bool FActorBrowsingMode::CanPaste() const
+void FActorBrowsingMode::OnItemPassesFilters(const ISceneOutlinerTreeItem& Item)
+{
+	if (const FActorTreeItem* ActorItem = Item.CastTo<FActorTreeItem>())
 	{
-		return CanPasteFoldersOnlyFromClipboard();
+		ApplicableActors.Add(ActorItem->Actor);
 	}
+}
 
+FReply FActorBrowsingMode::OnKeyDown(const FKeyEvent& InKeyEvent)
+{
+	const FSceneOutlinerItemSelection& Selection = SceneOutliner->GetSelection();
 
-	bool FActorBrowsingMode::CanPasteFoldersOnlyFromClipboard() const
-	{
-		// Intentionally not checking if the level is locked/hidden here, as it's better feedback for the user if they attempt to paste
-		// and get the message explaining why it's failed, than just not having the option available to them.
-		FString PasteString;
-		FPlatformApplicationMisc::ClipboardPaste(PasteString);
-		return PasteString.StartsWith("BEGIN FOLDERLIST");
-	}
-
-	TSharedPtr<FDragDropOperation> FActorBrowsingMode::CreateDragDropOperation(const TArray<FTreeItemPtr>& InTreeItems) const
+	// Rename key: Rename selected actors (not rebindable, because it doesn't make much sense to bind.)
+	if (InKeyEvent.GetKey() == EKeys::F2)
 	{
-		FDragDropPayload DraggedObjects(InTreeItems);
+		if (Selection.Num() == 1)
+		{
+			FSceneOutlinerTreeItemPtr ItemToRename = Selection.SelectedItems[0].Pin();
 
-		TSharedPtr<FSceneOutlinerDragDropOp> OutlinerOp = MakeShareable(new FSceneOutlinerDragDropOp());
+			if (ItemToRename.IsValid() && CanRenameItem(*ItemToRename) && ItemToRename->CanInteract())
+			{
+				SceneOutliner->SetPendingRenameItem(ItemToRename);
+				SceneOutliner->ScrollItemIntoView(ItemToRename);
+			}
 
-		if (DraggedObjects.Has<FActorTreeItem>())
-		{
-			TSharedPtr<FActorDragDropOp> ActorOperation = MakeShareable(new FActorDragDropOp);
-			ActorOperation->Init(DraggedObjects.GetData<TWeakObjectPtr<AActor>>(FWeakActorSelector()));
-			OutlinerOp->AddSubOp(ActorOperation);
+			return FReply::Handled();
 		}
+	}
 
-		if (DraggedObjects.Has<FFolderTreeItem>())
-		{
-			TSharedPtr<FFolderDragDropOp> FolderOperation = MakeShareable(new FFolderDragDropOp);
-			FolderOperation->Init(DraggedObjects.GetData<FName>(FFolderPathSelector()));
-			OutlinerOp->AddSubOp(FolderOperation);
-		}
-		OutlinerOp->Construct();
-		return OutlinerOp;
+	// F5 forces a full refresh
+	else if (InKeyEvent.GetKey() == EKeys::F5)
+	{
+		SceneOutliner->FullRefresh();
+		return FReply::Handled();
 	}
-	
-	bool FActorBrowsingMode::ParseDragDrop(FDragDropPayload& OutPayload, const FDragDropOperation& Operation) const
+
+	// Delete key: Delete selected actors (not rebindable, because it doesn't make much sense to bind.)
+	// Use Delete and Backspace instead of Platform_Delete because the LevelEditor default Edit Delete is bound to both 
+	else if (InKeyEvent.GetKey() == EKeys::Delete || InKeyEvent.GetKey() == EKeys::BackSpace)
 	{
-		if (Operation.IsOfType<FSceneOutlinerDragDropOp>())
+		if (SceneOutliner->GetSharedData().CustomDelete.IsBound())
 		{
-			const auto& OutlinerOp = static_cast<const FSceneOutlinerDragDropOp&>(Operation);
-			if (const auto& FolderOp = OutlinerOp.GetSubOp<FFolderDragDropOp>())
-			{
-				for (const auto& Folder : FolderOp->Folders)
-				{
-					OutPayload.DraggedItems.Add(SceneOutliner->GetTreeItem(Folder));
-				}
-			}
-			if (const auto& ActorOp = OutlinerOp.GetSubOp<FActorDragDropOp>())
-			{
-				for (const auto& Actor : ActorOp->Actors)
-				{
-					OutPayload.DraggedItems.Add(SceneOutliner->GetTreeItem(Actor.Get()));
-				}
-			}
-			return true;
+			SceneOutliner->GetSharedData().CustomDelete.Execute(Selection.SelectedItems);
 		}
-		else if (Operation.IsOfType<FActorDragDropOp>())
+		else
 		{
-			for (const TWeakObjectPtr<AActor>& Actor : static_cast<const FActorDragDropOp&>(Operation).Actors)
+			if (RepresentingWorld.IsValid())
 			{
-				OutPayload.DraggedItems.Add(SceneOutliner->GetTreeItem(Actor.Get()));
+				GUnrealEd->Exec(RepresentingWorld.Get(), TEXT("DELETE"));
 			}
-			return true;
 		}
-		
-		return false;
+		return FReply::Handled();
+			
 	}
+	return FReply::Unhandled();
+}
 
-	FDragValidationInfo FActorBrowsingMode::ValidateDrop(const ITreeItem& DropTarget, const FDragDropPayload& Payload) const
-	{
-		if (const FActorTreeItem* ActorItem = DropTarget.CastTo<FActorTreeItem>())
-		{
-			if (Payload.Has<FFolderTreeItem>())
-			{
-				return FDragValidationInfo(DropCompatibility::Drop_IncompatibleGeneric, LOCTEXT("FoldersOnActorError", "Cannot attach folders to actors"));
-			}
+bool FActorBrowsingMode::CanRename() const
+{
+	const FSceneOutlinerItemSelection ItemSelection = SceneOutliner->GetSelection();
+	const uint32 NumberOfFolders = ItemSelection.Num<FFolderTreeItem>();
+	return (NumberOfFolders == 1 && NumberOfFolders == ItemSelection.Num());
+}
+
+bool FActorBrowsingMode::CanRenameItem(const ISceneOutlinerTreeItem& Item) const
+{
+	// Can only rename actor and folder items when in actor browsing mode
+	return (Item.IsValid() && (Item.IsA<FActorTreeItem>() || Item.IsA<FFolderTreeItem>()));
+}
+
+bool FActorBrowsingMode::CanCut() const
+{
+	const FSceneOutlinerItemSelection ItemSelection = SceneOutliner->GetSelection();
+	const uint32 NumberOfFolders = ItemSelection.Num<FFolderTreeItem>();
+	return (NumberOfFolders > 0 && NumberOfFolders == ItemSelection.Num());
+}
+
+bool FActorBrowsingMode::CanCopy() const
+{
+	const FSceneOutlinerItemSelection ItemSelection = SceneOutliner->GetSelection();
+	const uint32 NumberOfFolders = ItemSelection.Num<FFolderTreeItem>();
+	return (NumberOfFolders > 0 && NumberOfFolders == ItemSelection.Num());
+}
+
+bool FActorBrowsingMode::CanPaste() const
+{
+	return CanPasteFoldersOnlyFromClipboard();
+}
+
+
+bool FActorBrowsingMode::CanPasteFoldersOnlyFromClipboard() const
+{
+	// Intentionally not checking if the level is locked/hidden here, as it's better feedback for the user if they attempt to paste
+	// and get the message explaining why it's failed, than just not having the option available to them.
+	FString PasteString;
+	FPlatformApplicationMisc::ClipboardPaste(PasteString);
+	return PasteString.StartsWith("BEGIN FOLDERLIST");
+}
 
-			const AActor* ActorTarget = ActorItem->Actor.Get();
+TSharedPtr<FDragDropOperation> FActorBrowsingMode::CreateDragDropOperation(const TArray<FSceneOutlinerTreeItemPtr>& InTreeItems) const
+{
+	FSceneOutlinerDragDropPayload DraggedObjects(InTreeItems);
 
-			if (!ActorTarget || !Payload.Has<FActorTreeItem>())
+	TSharedPtr<FSceneOutlinerDragDropOp> OutlinerOp = MakeShareable(new FSceneOutlinerDragDropOp());
+
+	if (DraggedObjects.Has<FActorTreeItem>())
+	{
+		TSharedPtr<FActorDragDropOp> ActorOperation = MakeShareable(new FActorDragDropOp);
+		ActorOperation->Init(DraggedObjects.GetData<TWeakObjectPtr<AActor>>(SceneOutliner::FWeakActorSelector()));
+		OutlinerOp->AddSubOp(ActorOperation);
+	}
+
+	if (DraggedObjects.Has<FFolderTreeItem>())
+	{
+		TSharedPtr<FFolderDragDropOp> FolderOperation = MakeShareable(new FFolderDragDropOp);
+		FolderOperation->Init(DraggedObjects.GetData<FName>(SceneOutliner::FFolderPathSelector()));
+		OutlinerOp->AddSubOp(FolderOperation);
+	}
+	OutlinerOp->Construct();
+	return OutlinerOp;
+}
+	
+bool FActorBrowsingMode::ParseDragDrop(FSceneOutlinerDragDropPayload& OutPayload, const FDragDropOperation& Operation) const
+{
+	if (Operation.IsOfType<FSceneOutlinerDragDropOp>())
+	{
+		const auto& OutlinerOp = static_cast<const FSceneOutlinerDragDropOp&>(Operation);
+		if (const auto& FolderOp = OutlinerOp.GetSubOp<FFolderDragDropOp>())
+		{
+			for (const auto& Folder : FolderOp->Folders)
+			{
+				OutPayload.DraggedItems.Add(SceneOutliner->GetTreeItem(Folder));
+			}
+		}
+		if (const auto& ActorOp = OutlinerOp.GetSubOp<FActorDragDropOp>())
+		{
+			for (const auto& Actor : ActorOp->Actors)
 			{
-				return FDragValidationInfo(DropCompatibility::Drop_IncompatibleGeneric, FText());
+				OutPayload.DraggedItems.Add(SceneOutliner->GetTreeItem(Actor.Get()));
 			}
+		}
+		return true;
+	}
+	else if (Operation.IsOfType<FActorDragDropOp>())
+	{
+		for (const TWeakObjectPtr<AActor>& Actor : static_cast<const FActorDragDropOp&>(Operation).Actors)
+		{
+			OutPayload.DraggedItems.Add(SceneOutliner->GetTreeItem(Actor.Get()));
+		}
+		return true;
+	}
+		
+	return false;
+}
+
+FSceneOutlinerDragValidationInfo FActorBrowsingMode::ValidateDrop(const ISceneOutlinerTreeItem& DropTarget, const FSceneOutlinerDragDropPayload& Payload) const
+{
+	if (const FActorTreeItem* ActorItem = DropTarget.CastTo<FActorTreeItem>())
+	{
+		if (Payload.Has<FFolderTreeItem>())
+		{
+			return FSceneOutlinerDragValidationInfo(ESceneOutlinerDropCompatibility::IncompatibleGeneric, LOCTEXT("FoldersOnActorError", "Cannot attach folders to actors"));
+		}
+
+		const AActor* ActorTarget = ActorItem->Actor.Get();
+
+		if (!ActorTarget || !Payload.Has<FActorTreeItem>())
+		{
+			return FSceneOutlinerDragValidationInfo(ESceneOutlinerDropCompatibility::IncompatibleGeneric, FText());
+		}
 
-			FText AttachErrorMsg;
-			bool bCanAttach = true;
-			bool bDraggedOntoAttachmentParent = true;
+		FText AttachErrorMsg;
+		bool bCanAttach = true;
+		bool bDraggedOntoAttachmentParent = true;
 
-			const auto& DragActors = Payload.GetData<TWeakObjectPtr<AActor>>(FWeakActorSelector());
-			for (const auto& DragActorPtr : DragActors)
+		const auto& DragActors = Payload.GetData<TWeakObjectPtr<AActor>>(SceneOutliner::FWeakActorSelector());
+		for (const auto& DragActorPtr : DragActors)
+		{
+			AActor* DragActor = DragActorPtr.Get();
+			if (DragActor)
 			{
-				AActor* DragActor = DragActorPtr.Get();
-				if (DragActor)
+				if (bCanAttach)
 				{
-					if (bCanAttach)
+					if (DragActor->IsChildActor())
 					{
-						if (DragActor->IsChildActor())
-						{
-							AttachErrorMsg = FText::Format(LOCTEXT("Error_AttachChildActor", "Cannot move {0} as it is a child actor."), FText::FromString(DragActor->GetActorLabel()));
-							bCanAttach = bDraggedOntoAttachmentParent = false;
-							break;
-						}
-						if (!GEditor->CanParentActors(ActorTarget, DragActor, &AttachErrorMsg))
-						{
-							bCanAttach = false;
-						}
+						AttachErrorMsg = FText::Format(LOCTEXT("Error_AttachChildActor", "Cannot move {0} as it is a child actor."), FText::FromString(DragActor->GetActorLabel()));
+						bCanAttach = bDraggedOntoAttachmentParent = false;
+						break;
 					}
-
-					if (DragActor->GetAttachParentActor() != ActorTarget)
+					if (!GEditor->CanParentActors(ActorTarget, DragActor, &AttachErrorMsg))
 					{
-						bDraggedOntoAttachmentParent = false;
+						bCanAttach = false;
 					}
 				}
-			}
 
-			const FText ActorLabel = FText::FromString(ActorTarget->GetActorLabel());
-			if (bDraggedOntoAttachmentParent)
-			{
-				if (DragActors.Num() == 1)
-				{
-					return FDragValidationInfo(DropCompatibility::Drop_CompatibleDetach, ActorLabel);
-				}
-				else
+				if (DragActor->GetAttachParentActor() != ActorTarget)
 				{
-					return FDragValidationInfo(DropCompatibility::Drop_CompatibleMultipleDetach, ActorLabel);
+					bDraggedOntoAttachmentParent = false;
 				}
 			}
-			else if (bCanAttach)
+		}
+
+		const FText ActorLabel = FText::FromString(ActorTarget->GetActorLabel());
+		if (bDraggedOntoAttachmentParent)
+		{
+			if (DragActors.Num() == 1)
 			{
-				if (DragActors.Num() == 1)
-				{
-					return FDragValidationInfo(DropCompatibility::Drop_CompatibleAttach, ActorLabel);
-				}
-				else
-				{
-					return FDragValidationInfo(DropCompatibility::Drop_CompatibleMultipleAttach, ActorLabel);
-				}
+				return FSceneOutlinerDragValidationInfo(ESceneOutlinerDropCompatibility::CompatibleDetach, ActorLabel);
 			}
 			else
 			{
-				if (DragActors.Num() == 1)
-				{
-					return FDragValidationInfo(DropCompatibility::Drop_IncompatibleGeneric, AttachErrorMsg);
-				}
-				else
-				{
-					const FText ReasonText = FText::Format(LOCTEXT("DropOntoText", "{0}. {1}"), ActorLabel, AttachErrorMsg);
-					return FDragValidationInfo(DropCompatibility::Drop_IncompatibleMultipleAttach, ReasonText);
-				}
+				return FSceneOutlinerDragValidationInfo(ESceneOutlinerDropCompatibility::CompatibleMultipleDetach, ActorLabel);
 			}
 		}
-		else if (DropTarget.IsA<FFolderTreeItem>() || DropTarget.IsA<FWorldTreeItem>())
+		else if (bCanAttach)
 		{
-			const FFolderTreeItem* FolderItem = DropTarget.CastTo<FFolderTreeItem>();
-			// World items are treated as folders with path = none
-			const FName& DestinationPath = FolderItem ? FolderItem->Path : NAME_None;
-			if (Payload.Has<FFolderTreeItem>())
+			if (DragActors.Num() == 1)
 			{
-				// Iterate over all the folders that have been dragged
-				for (FName DraggedFolder : Payload.GetData<FName>(FFolderPathSelector()))
-				{
-					const FName Leaf = GetFolderLeafName(DraggedFolder);
-					const FName Parent = GetParentPath(DraggedFolder);
-
-					if (Parent == DestinationPath)
-					{
-						FFormatNamedArguments Args;
-						Args.Add(TEXT("SourceName"), FText::FromName(Leaf));
-
-						FText Text;
-						if (DestinationPath.IsNone())
-						{
-							Text = FText::Format(LOCTEXT("FolderAlreadyAssignedRoot", "{SourceName} is already assigned to root"), Args);
-						}
-						else
-						{
-							Args.Add(TEXT("DestPath"), FText::FromName(DestinationPath));
-							Text = FText::Format(LOCTEXT("FolderAlreadyAssigned", "{SourceName} is already assigned to {DestPath}"), Args);
-						}
-
-						return FDragValidationInfo(DropCompatibility::Drop_IncompatibleGeneric, Text);
-					}
+				return FSceneOutlinerDragValidationInfo(ESceneOutlinerDropCompatibility::CompatibleAttach, ActorLabel);
+			}
+			else
+			{
+				return FSceneOutlinerDragValidationInfo(ESceneOutlinerDropCompatibility::CompatibleMultipleAttach, ActorLabel);
+			}
+		}
+		else
+		{
+			if (DragActors.Num() == 1)
+			{
+				return FSceneOutlinerDragValidationInfo(ESceneOutlinerDropCompatibility::IncompatibleGeneric, AttachErrorMsg);
+			}
+			else
+			{
+				const FText ReasonText = FText::Format(LOCTEXT("DropOntoText", "{0}. {1}"), ActorLabel, AttachErrorMsg);
+				return FSceneOutlinerDragValidationInfo(ESceneOutlinerDropCompatibility::IncompatibleMultipleAttach, ReasonText);
+			}
+		}
+	}
+	else if (DropTarget.IsA<FFolderTreeItem>() || DropTarget.IsA<FWorldTreeItem>())
+	{
+		const FFolderTreeItem* FolderItem = DropTarget.CastTo<FFolderTreeItem>();
+		// World items are treated as folders with path = none
+		const FName& DestinationPath = FolderItem ? FolderItem->Path : NAME_None;
+		if (Payload.Has<FFolderTreeItem>())
+		{
+			// Iterate over all the folders that have been dragged
+			for (FName DraggedFolder : Payload.GetData<FName>(SceneOutliner::FFolderPathSelector()))
+			{
+				const FName Leaf = SceneOutliner::GetFolderLeafName(DraggedFolder);
+				const FName Parent = SceneOutliner::GetParentPath(DraggedFolder);
 
-					const FString DragFolderPath = DraggedFolder.ToString();
-					const FString LeafName = Leaf.ToString();
-					const FString DstFolderPath = DestinationPath.IsNone() ? FString() : DestinationPath.ToString();
-					const FString NewPath = DstFolderPath / LeafName;
+				if (Parent == DestinationPath)
+				{
+					FFormatNamedArguments Args;
+					Args.Add(TEXT("SourceName"), FText::FromName(Leaf));
 
-					if (FActorFolders::Get().GetFolderProperties(*RepresentingWorld, FName(*NewPath)))
+					FText Text;
+					if (DestinationPath.IsNone())
 					{
-						// The folder already exists
-						FFormatNamedArguments Args;
-						Args.Add(TEXT("DragName"), FText::FromString(LeafName));
-						return FDragValidationInfo(DropCompatibility::Drop_IncompatibleGeneric,
-							FText::Format(LOCTEXT("FolderAlreadyExistsRoot", "A folder called \"{DragName}\" already exists at this level"), Args));
+						Text = FText::Format(LOCTEXT("FolderAlreadyAssignedRoot", "{SourceName} is already assigned to root"), Args);
 					}
-					else if (DragFolderPath == DstFolderPath || DstFolderPath.StartsWith(DragFolderPath + "/"))
+					else
 					{
-						// Cannot drag as a child of itself
-						FFormatNamedArguments Args;
-						Args.Add(TEXT("FolderPath"), FText::FromName(DraggedFolder));
-						return FDragValidationInfo(
-							DropCompatibility::Drop_IncompatibleGeneric,
-							FText::Format(LOCTEXT("ChildOfItself", "Cannot move \"{FolderPath}\" to be a child of itself"), Args));
+						Args.Add(TEXT("DestPath"), FText::FromName(DestinationPath));
+						Text = FText::Format(LOCTEXT("FolderAlreadyAssigned", "{SourceName} is already assigned to {DestPath}"), Args);
 					}
+
+					return FSceneOutlinerDragValidationInfo(ESceneOutlinerDropCompatibility::IncompatibleGeneric, Text);
+				}
+
+				const FString DragFolderPath = DraggedFolder.ToString();
+				const FString LeafName = Leaf.ToString();
+				const FString DstFolderPath = DestinationPath.IsNone() ? FString() : DestinationPath.ToString();
+				const FString NewPath = DstFolderPath / LeafName;
+
+				if (FActorFolders::Get().GetFolderProperties(*RepresentingWorld, FName(*NewPath)))
+				{
+					// The folder already exists
+					FFormatNamedArguments Args;
+					Args.Add(TEXT("DragName"), FText::FromString(LeafName));
+					return FSceneOutlinerDragValidationInfo(ESceneOutlinerDropCompatibility::IncompatibleGeneric,
+						FText::Format(LOCTEXT("FolderAlreadyExistsRoot", "A folder called \"{DragName}\" already exists at this level"), Args));
+				}
+				else if (DragFolderPath == DstFolderPath || DstFolderPath.StartsWith(DragFolderPath + "/"))
+				{
+					// Cannot drag as a child of itself
+					FFormatNamedArguments Args;
+					Args.Add(TEXT("FolderPath"), FText::FromName(DraggedFolder));
+					return FSceneOutlinerDragValidationInfo(
+						ESceneOutlinerDropCompatibility::IncompatibleGeneric,
+						FText::Format(LOCTEXT("ChildOfItself", "Cannot move \"{FolderPath}\" to be a child of itself"), Args));
 				}
 			}
+		}
 
-			if (Payload.Has<FActorTreeItem>())
+		if (Payload.Has<FActorTreeItem>())
+		{
+			// Iterate over all the folders that have been dragged
+			for (const TWeakObjectPtr<AActor>& WeakActor : Payload.GetData<TWeakObjectPtr<AActor>>(SceneOutliner::FWeakActorSelector()))
 			{
-				// Iterate over all the folders that have been dragged
-				for (const TWeakObjectPtr<AActor>& WeakActor : Payload.GetData<TWeakObjectPtr<AActor>>(FWeakActorSelector()))
+				const AActor* Actor = WeakActor.Get();
+				if (Actor->IsChildActor())
 				{
-					const AActor* Actor = WeakActor.Get();
-					if (Actor->IsChildActor())
+					return FSceneOutlinerDragValidationInfo(ESceneOutlinerDropCompatibility::IncompatibleGeneric, FText::Format(LOCTEXT("Error_AttachChildActor", "Cannot move {0} as it is a child actor."), FText::FromString(Actor->GetActorLabel())));
+				}
+				else if (Actor->GetFolderPath() == DestinationPath && !Actor->GetAttachParentActor())
+				{
+					FFormatNamedArguments Args;
+					Args.Add(TEXT("SourceName"), FText::FromString(Actor->GetActorLabel()));
+
+					FText Text;
+					if (DestinationPath.IsNone())
 					{
-						return FDragValidationInfo(DropCompatibility::Drop_IncompatibleGeneric, FText::Format(LOCTEXT("Error_AttachChildActor", "Cannot move {0} as it is a child actor."), FText::FromString(Actor->GetActorLabel())));
+						Text = FText::Format(LOCTEXT("FolderAlreadyAssignedRoot", "{SourceName} is already assigned to root"), Args);
 					}
-					else if (Actor->GetFolderPath() == DestinationPath && !Actor->GetAttachParentActor())
+					else
 					{
-						FFormatNamedArguments Args;
-						Args.Add(TEXT("SourceName"), FText::FromString(Actor->GetActorLabel()));
-
-						FText Text;
-						if (DestinationPath.IsNone())
-						{
-							Text = FText::Format(LOCTEXT("FolderAlreadyAssignedRoot", "{SourceName} is already assigned to root"), Args);
-						}
-						else
-						{
-							Args.Add(TEXT("DestPath"), FText::FromName(DestinationPath));
-							Text = FText::Format(LOCTEXT("FolderAlreadyAssigned", "{SourceName} is already assigned to {DestPath}"), Args);
-						}
-
-						return FDragValidationInfo(DropCompatibility::Drop_IncompatibleGeneric, Text);
+						Args.Add(TEXT("DestPath"), FText::FromName(DestinationPath));
+						Text = FText::Format(LOCTEXT("FolderAlreadyAssigned", "{SourceName} is already assigned to {DestPath}"), Args);
 					}
+
+					return FSceneOutlinerDragValidationInfo(ESceneOutlinerDropCompatibility::IncompatibleGeneric, Text);
 				}
 			}
+		}
 
-			// Everything else is a valid operation
-			if (DestinationPath.IsNone())
-			{
-				return FDragValidationInfo(DropCompatibility::Drop_CompatibleGeneric, LOCTEXT("MoveToRoot", "Move to root"));
-			}
-			else
-			{
-				FFormatNamedArguments Args;
-				Args.Add(TEXT("DestPath"), FText::FromName(DestinationPath));
-				return FDragValidationInfo(DropCompatibility::Drop_CompatibleGeneric, FText::Format(LOCTEXT("MoveInto", "Move into \"{DestPath}\""), Args));
-			}
+		// Everything else is a valid operation
+		if (DestinationPath.IsNone())
+		{
+			return FSceneOutlinerDragValidationInfo(ESceneOutlinerDropCompatibility::CompatibleGeneric, LOCTEXT("MoveToRoot", "Move to root"));
 		}
-		else if (DropTarget.IsA<FComponentTreeItem>())
+		else
 		{
-			// we don't allow drag and drop on components for now
-			return FDragValidationInfo(DropCompatibility::Drop_IncompatibleGeneric, FText());
+			FFormatNamedArguments Args;
+			Args.Add(TEXT("DestPath"), FText::FromName(DestinationPath));
+			return FSceneOutlinerDragValidationInfo(ESceneOutlinerDropCompatibility::CompatibleGeneric, FText::Format(LOCTEXT("MoveInto", "Move into \"{DestPath}\""), Args));
 		}
-		return FDragValidationInfo::Invalid();
 	}
+	else if (DropTarget.IsA<FComponentTreeItem>())
+	{
+		// we don't allow drag and drop on components for now
+		return FSceneOutlinerDragValidationInfo(ESceneOutlinerDropCompatibility::IncompatibleGeneric, FText());
+	}
+	return FSceneOutlinerDragValidationInfo::Invalid();
+}
 
-	void FActorBrowsingMode::OnDrop(ITreeItem& DropTarget, const FDragDropPayload& Payload, const FDragValidationInfo& ValidationInfo) const
+void FActorBrowsingMode::OnDrop(ISceneOutlinerTreeItem& DropTarget, const FSceneOutlinerDragDropPayload& Payload, const FSceneOutlinerDragValidationInfo& ValidationInfo) const
+{
+	if (const FActorTreeItem* ActorItem = DropTarget.CastTo<FActorTreeItem>())
 	{
-		if (const FActorTreeItem* ActorItem = DropTarget.CastTo<FActorTreeItem>())
+		AActor* DropActor = ActorItem->Actor.Get();
+		if (!DropActor)
 		{
-			AActor* DropActor = ActorItem->Actor.Get();
-			if (!DropActor)
-			{
-				return;
-			}
+			return;
+		}
 
-			TArray<TWeakObjectPtr<AActor>> DraggedActors = Payload.GetData<TWeakObjectPtr<AActor>>(FWeakActorSelector());
+		TArray<TWeakObjectPtr<AActor>> DraggedActors = Payload.GetData<TWeakObjectPtr<AActor>>(SceneOutliner::FWeakActorSelector());
 
-			FMessageLog EditorErrors("EditorErrors");
-			EditorErrors.NewPage(LOCTEXT("ActorAttachmentsPageLabel", "Actor attachment"));
+		FMessageLog EditorErrors("EditorErrors");
+		EditorErrors.NewPage(LOCTEXT("ActorAttachmentsPageLabel", "Actor attachment"));
 
-			if (ValidationInfo.CompatibilityType == DropCompatibility::Drop_CompatibleMultipleDetach || ValidationInfo.CompatibilityType == DropCompatibility::Drop_CompatibleDetach)
-			{
-				const FScopedTransaction Transaction(LOCTEXT("UndoAction_DetachActors", "Detach actors"));
+		if (ValidationInfo.CompatibilityType == ESceneOutlinerDropCompatibility::CompatibleMultipleDetach || ValidationInfo.CompatibilityType == ESceneOutlinerDropCompatibility::CompatibleDetach)
+		{
+			const FScopedTransaction Transaction(LOCTEXT("UndoAction_DetachActors", "Detach actors"));
 
-				for (const auto& WeakActor : DraggedActors)
+			for (const auto& WeakActor : DraggedActors)
+			{
+				if (auto* DragActor = WeakActor.Get())
 				{
-					if (auto* DragActor = WeakActor.Get())
+					// Detach from parent
+					USceneComponent* RootComp = DragActor->GetRootComponent();
+					if (RootComp && RootComp->GetAttachParent())
 					{
-						// Detach from parent
-						USceneComponent* RootComp = DragActor->GetRootComponent();
-						if (RootComp && RootComp->GetAttachParent())
-						{
-							AActor* OldParent = RootComp->GetAttachParent()->GetOwner();
-							OldParent->Modify();
-							RootComp->DetachFromComponent(FDetachmentTransformRules::KeepWorldTransform);
+						AActor* OldParent = RootComp->GetAttachParent()->GetOwner();
+						OldParent->Modify();
+						RootComp->DetachFromComponent(FDetachmentTransformRules::KeepWorldTransform);
 
-							DragActor->SetFolderPath_Recursively(OldParent->GetFolderPath());
-						}
+						DragActor->SetFolderPath_Recursively(OldParent->GetFolderPath());
 					}
 				}
 			}
-			else if (ValidationInfo.CompatibilityType == DropCompatibility::Drop_CompatibleMultipleAttach || ValidationInfo.CompatibilityType == DropCompatibility::Drop_CompatibleAttach)
-			{
-				// Show socket chooser if we have sockets to select
+		}
+		else if (ValidationInfo.CompatibilityType == ESceneOutlinerDropCompatibility::CompatibleMultipleAttach || ValidationInfo.CompatibilityType == ESceneOutlinerDropCompatibility::CompatibleAttach)
+		{
+			// Show socket chooser if we have sockets to select
 
-				auto PerformAttachment = [](FName SocketName, TWeakObjectPtr<AActor> Parent, const TArray<TWeakObjectPtr<AActor>> NewAttachments)
+			auto PerformAttachment = [](FName SocketName, TWeakObjectPtr<AActor> Parent, const TArray<TWeakObjectPtr<AActor>> NewAttachments)
+			{
+				AActor* ParentActor = Parent.Get();
+				if (ParentActor)
 				{
-					AActor* ParentActor = Parent.Get();
-					if (ParentActor)
-					{
-						// modify parent and child
-						const FScopedTransaction Transaction(LOCTEXT("UndoAction_PerformAttachment", "Attach actors"));
+					// modify parent and child
+					const FScopedTransaction Transaction(LOCTEXT("UndoAction_PerformAttachment", "Attach actors"));
 
-						// Attach each child
-						bool bAttached = false;
-						for (auto& Child : NewAttachments)
+					// Attach each child
+					bool bAttached = false;
+					for (auto& Child : NewAttachments)
+					{
+						AActor* ChildActor = Child.Get();
+						if (GEditor->CanParentActors(ParentActor, ChildActor))
 						{
-							AActor* ChildActor = Child.Get();
-							if (GEditor->CanParentActors(ParentActor, ChildActor))
-							{
-								GEditor->ParentActors(ParentActor, ChildActor, SocketName);
+							GEditor->ParentActors(ParentActor, ChildActor, SocketName);
 
-								ChildActor->SetFolderPath_Recursively(ParentActor->GetFolderPath());
-							}
+							ChildActor->SetFolderPath_Recursively(ParentActor->GetFolderPath());
 						}
 					}
-				};
-
-				//@TODO: Should create a menu for each component that contains sockets, or have some form of disambiguation within the menu (like a fully qualified path)
-				// Instead, we currently only display the sockets on the root component
-				USceneComponent* Component = DropActor->GetRootComponent();
-				if ((Component != NULL) && (Component->HasAnySockets()))
-				{
-					// Create the popup
-					FSlateApplication::Get().PushMenu(
-						SceneOutliner->AsShared(),
-						FWidgetPath(),
-						SNew(SSocketChooserPopup)
-						.SceneComponent(Component)
-						.OnSocketChosen_Lambda(PerformAttachment, DropActor, MoveTemp(DraggedActors)),
-						FSlateApplication::Get().GetCursorPos(),
-						FPopupTransitionEffect(FPopupTransitionEffect::TypeInPopup)
-					);
-				}
-				else
-				{
-					PerformAttachment(NAME_None, DropActor, MoveTemp(DraggedActors));
 				}
-			}
+			};
 
-			// Report errors
-			EditorErrors.Notify(NSLOCTEXT("ActorAttachmentError", "AttachmentsFailed", "Attachments Failed!"));
+			//@TODO: Should create a menu for each component that contains sockets, or have some form of disambiguation within the menu (like a fully qualified path)
+			// Instead, we currently only display the sockets on the root component
+			USceneComponent* Component = DropActor->GetRootComponent();
+			if ((Component != NULL) && (Component->HasAnySockets()))
+			{
+				// Create the popup
+				FSlateApplication::Get().PushMenu(
+					SceneOutliner->AsShared(),
+					FWidgetPath(),
+					SNew(SSocketChooserPopup)
+					.SceneComponent(Component)
+					.OnSocketChosen_Lambda(PerformAttachment, DropActor, MoveTemp(DraggedActors)),
+					FSlateApplication::Get().GetCursorPos(),
+					FPopupTransitionEffect(FPopupTransitionEffect::TypeInPopup)
+				);
+			}
+			else
+			{
+				PerformAttachment(NAME_None, DropActor, MoveTemp(DraggedActors));
+			}
 		}
-		else if (DropTarget.IsA<FFolderTreeItem>() || DropTarget.IsA<FWorldTreeItem>())
-		{
-			const FFolderTreeItem* FolderItem = DropTarget.CastTo<FFolderTreeItem>();
-			// If the cast fails, the item must be a WorldTreeItem and we set the path to None to reflect this
-			const FName& DestinationPath = FolderItem ? FolderItem->Path : FName(NAME_None);
 
-			const FScopedTransaction Transaction(LOCTEXT("MoveOutlinerItems", "Move World Outliner Items"));
+		// Report errors
+		EditorErrors.Notify(NSLOCTEXT("ActorAttachmentError", "AttachmentsFailed", "Attachments Failed!"));
+	}
+	else if (DropTarget.IsA<FFolderTreeItem>() || DropTarget.IsA<FWorldTreeItem>())
+	{
+		const FFolderTreeItem* FolderItem = DropTarget.CastTo<FFolderTreeItem>();
+		// If the cast fails, the item must be a WorldTreeItem and we set the path to None to reflect this
+		const FName& DestinationPath = FolderItem ? FolderItem->Path : FName(NAME_None);
+
+		const FScopedTransaction Transaction(LOCTEXT("MoveOutlinerItems", "Move World Outliner Items"));
 			
-			auto MoveToDestination = [&DestinationPath](FFolderTreeItem& Item)
-			{
-				Item.MoveTo(DestinationPath);
-			};
-			Payload.ForEachItem<FFolderTreeItem>(MoveToDestination);
+		auto MoveToDestination = [&DestinationPath](FFolderTreeItem& Item)
+		{
+			Item.MoveTo(DestinationPath);
+		};
+		Payload.ForEachItem<FFolderTreeItem>(MoveToDestination);
 
-			// Set the folder path on all the dragged actors, and detach any that need to be moved
-			if (Payload.Has<FActorTreeItem>())
-			{
-				TSet<const AActor*> ParentActors;
-				TSet<const AActor*> ChildActors;
+		// Set the folder path on all the dragged actors, and detach any that need to be moved
+		if (Payload.Has<FActorTreeItem>())
+		{
+			TSet<const AActor*> ParentActors;
+			TSet<const AActor*> ChildActors;
 
-				for (const auto& WeakActor : Payload.GetData<TWeakObjectPtr<AActor>>(FWeakActorSelector()))
+			for (const auto& WeakActor : Payload.GetData<TWeakObjectPtr<AActor>>(SceneOutliner::FWeakActorSelector()))
+			{
+				AActor* Actor = WeakActor.Get();
+				if (Actor)
 				{
-					AActor* Actor = WeakActor.Get();
-					if (Actor)
-					{
-						// First mark this object as a parent, then set its children's path
-						ParentActors.Add(Actor);
-						Actor->SetFolderPath(DestinationPath);
-
-						FActorEditorUtils::TraverseActorTree_ParentFirst(Actor, [&](AActor* InActor) {
-							ChildActors.Add(InActor);
-							InActor->SetFolderPath(DestinationPath);
-							return true;
-							}, false);
-					}
+					// First mark this object as a parent, then set its children's path
+					ParentActors.Add(Actor);
+					Actor->SetFolderPath(DestinationPath);
+
+					FActorEditorUtils::TraverseActorTree_ParentFirst(Actor, [&](AActor* InActor) {
+						ChildActors.Add(InActor);
+						InActor->SetFolderPath(DestinationPath);
+						return true;
+						}, false);
 				}
+			}
 
-				// Detach parent actors
-				for (const AActor* Parent : ParentActors)
-				{
-					auto* RootComp = Parent->GetRootComponent();
+			// Detach parent actors
+			for (const AActor* Parent : ParentActors)
+			{
+				auto* RootComp = Parent->GetRootComponent();
 
-					// We don't detach if it's a child of another that's been dragged
-					if (RootComp && RootComp->GetAttachParent() && !ChildActors.Contains(Parent))
+				// We don't detach if it's a child of another that's been dragged
+				if (RootComp && RootComp->GetAttachParent() && !ChildActors.Contains(Parent))
+				{
+					if (AActor* OldParentActor = RootComp->GetAttachParent()->GetOwner())
 					{
-						if (AActor* OldParentActor = RootComp->GetAttachParent()->GetOwner())
-						{
-							OldParentActor->Modify();
-						}
-						RootComp->DetachFromComponent(FDetachmentTransformRules::KeepWorldTransform);
+						OldParentActor->Modify();
 					}
+					RootComp->DetachFromComponent(FDetachmentTransformRules::KeepWorldTransform);
 				}
 			}
 		}
 	}
+}
 
-	FName FActorBrowsingMode::CreateNewFolder()
-	{
-		const FScopedTransaction Transaction(LOCTEXT("UndoAction_CreateFolder", "Create Folder"));
+FName FActorBrowsingMode::CreateNewFolder()
+{
+	const FScopedTransaction Transaction(LOCTEXT("UndoAction_CreateFolder", "Create Folder"));
 
-		const FName NewFolderName = FActorFolders::Get().GetDefaultFolderNameForSelection(*RepresentingWorld);
-		FActorFolders::Get().CreateFolderContainingSelection(*RepresentingWorld, NewFolderName);
+	const FName NewFolderName = FActorFolders::Get().GetDefaultFolderNameForSelection(*RepresentingWorld);
+	FActorFolders::Get().CreateFolderContainingSelection(*RepresentingWorld, NewFolderName);
 
-		return NewFolderName;
-	}
+	return NewFolderName;
+}
+
+FName FActorBrowsingMode::CreateFolder(const FName& ParentPath, const FName& LeafName)
+{
+	const FName NewPath = FActorFolders::Get().GetFolderName(*RepresentingWorld, ParentPath, LeafName);
+	FActorFolders::Get().CreateFolder(*RepresentingWorld, NewPath);
+	return NewPath;
+}
 
-	FName FActorBrowsingMode::CreateFolder(const FName& ParentPath, const FName& LeafName)
+bool FActorBrowsingMode::ReparentItemToFolder(const FName& FolderPath, const FSceneOutlinerTreeItemPtr& Item)
+{
+	if (FActorTreeItem* ActorItem = Item->CastTo<FActorTreeItem>())
 	{
-		const FName NewPath = FActorFolders::Get().GetFolderName(*RepresentingWorld, ParentPath, LeafName);
-		FActorFolders::Get().CreateFolder(*RepresentingWorld, NewPath);
-		return NewPath;
+		ActorItem->Actor->SetFolderPath_Recursively(FolderPath);
+		return true;
 	}
+	return false;
+}
 
-	bool FActorBrowsingMode::ReparentItemToFolder(const FName& FolderPath, const FTreeItemPtr& Item)
+void FActorBrowsingMode::SelectFoldersDescendants(const TArray<FFolderTreeItem*>& FolderItems, bool bSelectImmediateChildrenOnly)
+{
+	// Ensure that all folder descendants are expanded
+	TFunctionRef<void(const FSceneOutlinerTreeItemPtr&)> RecursiveExpand = [&](const FSceneOutlinerTreeItemPtr&  Item)
 	{
-		if (FActorTreeItem* ActorItem = Item->CastTo<FActorTreeItem>())
+		SceneOutliner->SetItemExpansion(Item, true);
+			
+		if (!bSelectImmediateChildrenOnly)
 		{
-			ActorItem->Actor->SetFolderPath_Recursively(FolderPath);
-			return true;
+			for (const TWeakPtr<ISceneOutlinerTreeItem> Child : Item->GetChildren())
+			{
+				RecursiveExpand(Child.Pin());
+			}
 		}
-		return false;
-	}
+	};
 
-	void FActorBrowsingMode::SelectFoldersDescendants(const TArray<FFolderTreeItem*>& FolderItems, bool bSelectImmediateChildrenOnly)
+	// Selects all actor descendants of a folder
+	TFunctionRef<void(const FSceneOutlinerTreeItemPtr&)> RecursiveActorSelect = [&](const FSceneOutlinerTreeItemPtr& Item)
 	{
-		// Ensure that all folder descendants are expanded
-		TFunctionRef<void(const FTreeItemPtr&)> RecursiveExpand = [&](const FTreeItemPtr&  Item)
+		// If the current item is an actor, ensure to select it as well
+		if (FActorTreeItem* ActorItem = Item->CastTo<FActorTreeItem>())
 		{
-			SceneOutliner->SetItemExpansion(Item, true);
-			
-			if (!bSelectImmediateChildrenOnly)
+			if (AActor* Actor = ActorItem->Actor.Get())
 			{
-				for (const TWeakPtr<ITreeItem> Child : Item->GetChildren())
-				{
-					RecursiveExpand(Child.Pin());
-				}
+				GEditor->SelectActor(Actor, true, false);
 			}
-		};
-
-		// Selects all actor descendants of a folder
-		TFunctionRef<void(const FTreeItemPtr&)> RecursiveActorSelect = [&](const FTreeItemPtr& Item)
+		}
+		// Select all children
+		for (const TWeakPtr<ISceneOutlinerTreeItem> Child : Item->GetChildren())
 		{
-			// If the current item is an actor, ensure to select it as well
-			if (FActorTreeItem* ActorItem = Item->CastTo<FActorTreeItem>())
+			if (FActorTreeItem* ActorItem = Child.Pin()->CastTo<FActorTreeItem>())
 			{
 				if (AActor* Actor = ActorItem->Actor.Get())
 				{
 					GEditor->SelectActor(Actor, true, false);
 				}
 			}
-			// Select all children
-			for (const TWeakPtr<ITreeItem> Child : Item->GetChildren())
+			else if (FFolderTreeItem* FolderItem = Child.Pin()->CastTo<FFolderTreeItem>())
 			{
-				if (FActorTreeItem* ActorItem = Child.Pin()->CastTo<FActorTreeItem>())
-				{
-					if (AActor* Actor = ActorItem->Actor.Get())
-					{
-						GEditor->SelectActor(Actor, true, false);
-					}
-				}
-				else if (FFolderTreeItem* FolderItem = Child.Pin()->CastTo<FFolderTreeItem>())
-				{
-					SceneOutliner->SetItemSelection(FolderItem->AsShared(), true);
-				}
+				SceneOutliner->SetItemSelection(FolderItem->AsShared(), true);
+			}
 
-				if (!bSelectImmediateChildrenOnly)
+			if (!bSelectImmediateChildrenOnly)
+			{
+				for (const TWeakPtr<ISceneOutlinerTreeItem>& Grandchild : Child.Pin()->GetChildren())
 				{
-					for (const TWeakPtr<ITreeItem>& Grandchild : Child.Pin()->GetChildren())
-					{
-						RecursiveActorSelect(Grandchild.Pin());
-					}
+					RecursiveActorSelect(Grandchild.Pin());
 				}
 			}
-		};
-
-		// Expand everything before beginning selection
-		for (FFolderTreeItem* Folder : FolderItems)
-		{
-			RecursiveExpand(Folder->AsShared());
 		}
+	};
 
-		// batch selection
-		GEditor->GetSelectedActors()->BeginBatchSelectOperation();
+	// Expand everything before beginning selection
+	for (FFolderTreeItem* Folder : FolderItems)
+	{
+		RecursiveExpand(Folder->AsShared());
+	}
 
-		for (FFolderTreeItem* Folder : FolderItems)
-		{
-			RecursiveActorSelect(Folder->AsShared());
-		}
+	// batch selection
+	GEditor->GetSelectedActors()->BeginBatchSelectOperation();
 
-		GEditor->GetSelectedActors()->EndBatchSelectOperation(/*bNotify*/false);
-		GEditor->NoteSelectionChange();
+	for (FFolderTreeItem* Folder : FolderItems)
+	{
+		RecursiveActorSelect(Folder->AsShared());
 	}
 
-	FCreateOutlinerMode FActorBrowsingMode::CreateFolderPickerMode() const
+	GEditor->GetSelectedActors()->EndBatchSelectOperation(/*bNotify*/false);
+	GEditor->NoteSelectionChange();
+}
+
+FCreateSceneOutlinerMode FActorBrowsingMode::CreateFolderPickerMode() const
+{
+	auto MoveSelectionTo = [this](const FSceneOutlinerTreeItemRef& NewParent)
 	{
-		auto MoveSelectionTo = [this](const FTreeItemRef& NewParent)
+		if (NewParent->IsA<FWorldTreeItem>())
 		{
-			if (NewParent->IsA<FWorldTreeItem>())
-			{
-				SceneOutliner->MoveSelectionTo(FName());
-			}
-			else if (FFolderTreeItem* FolderItem = NewParent->CastTo<FFolderTreeItem>())
-			{
-				SceneOutliner->MoveSelectionTo(FolderItem->Path);
-			}
-		};
+			SceneOutliner->MoveSelectionTo(FName());
+		}
+		else if (FFolderTreeItem* FolderItem = NewParent->CastTo<FFolderTreeItem>())
+		{
+			SceneOutliner->MoveSelectionTo(FolderItem->Path);
+		}
+	};
 
-		return FCreateOutlinerMode::CreateLambda([this, MoveSelectionTo](SceneOutliner::SSceneOutliner* Outliner)
-			{
-				return new FActorFolderPickingMode(Outliner, FOnSceneOutlinerItemPicked::CreateLambda(MoveSelectionTo));
-			});
-	}
+	return FCreateSceneOutlinerMode::CreateLambda([this, MoveSelectionTo](SSceneOutliner* Outliner)
+		{
+			return new FActorFolderPickingMode(Outliner, FOnSceneOutlinerItemPicked::CreateLambda(MoveSelectionTo));
+		});
+}
 
-	void FActorBrowsingMode::OnDuplicateSelected()
-	{
-		GUnrealEd->Exec(RepresentingWorld.Get(), TEXT("DUPLICATE"));
-	}
+void FActorBrowsingMode::OnDuplicateSelected()
+{
+	GUnrealEd->Exec(RepresentingWorld.Get(), TEXT("DUPLICATE"));
+}
 
-	void FActorBrowsingMode::OnEditCutActorsBegin()
-	{
-		// Only a callback in actor browsing mode
-		SceneOutliner->CopyFoldersBegin();
-		SceneOutliner->DeleteFoldersBegin();
-	}
+void FActorBrowsingMode::OnEditCutActorsBegin()
+{
+	// Only a callback in actor browsing mode
+	SceneOutliner->CopyFoldersBegin();
+	SceneOutliner->DeleteFoldersBegin();
+}
 
-	void FActorBrowsingMode::OnEditCutActorsEnd()
-	{
-		// Only a callback in actor browsing mode
-		SceneOutliner->CopyFoldersEnd();
-		SceneOutliner->DeleteFoldersEnd();
-	}
+void FActorBrowsingMode::OnEditCutActorsEnd()
+{
+	// Only a callback in actor browsing mode
+	SceneOutliner->CopyFoldersEnd();
+	SceneOutliner->DeleteFoldersEnd();
+}
 
-	void FActorBrowsingMode::OnEditCopyActorsBegin()
-	{
-		// Only a callback in actor browsing mode
-		SceneOutliner->CopyFoldersBegin();
-	}
+void FActorBrowsingMode::OnEditCopyActorsBegin()
+{
+	// Only a callback in actor browsing mode
+	SceneOutliner->CopyFoldersBegin();
+}
 
-	void FActorBrowsingMode::OnEditCopyActorsEnd()
-	{
-		// Only a callback in actor browsing mode
-		SceneOutliner->CopyFoldersEnd();
-	}
+void FActorBrowsingMode::OnEditCopyActorsEnd()
+{
+	// Only a callback in actor browsing mode
+	SceneOutliner->CopyFoldersEnd();
+}
 
-	void FActorBrowsingMode::OnEditPasteActorsBegin()
-	{
-		// Only a callback in actor browsing mode
-		const TArray<FName> FolderPaths = SceneOutliner->GetClipboardPasteFolders();
-		SceneOutliner->PasteFoldersBegin(FolderPaths);
-	}
+void FActorBrowsingMode::OnEditPasteActorsBegin()
+{
+	// Only a callback in actor browsing mode
+	const TArray<FName> FolderPaths = SceneOutliner->GetClipboardPasteFolders();
+	SceneOutliner->PasteFoldersBegin(FolderPaths);
+}
 
-	void FActorBrowsingMode::OnEditPasteActorsEnd()
-	{
-		// Only a callback in actor browsing mode
-		SceneOutliner->PasteFoldersEnd();
-	}
+void FActorBrowsingMode::OnEditPasteActorsEnd()
+{
+	// Only a callback in actor browsing mode
+	SceneOutliner->PasteFoldersEnd();
+}
 
-	void FActorBrowsingMode::OnDuplicateActorsBegin()
-	{
-		// Only a callback in actor browsing mode
-		const TArray<FName> SelectedFolderPaths = SceneOutliner->GetSelection().GetData<FName>(FFolderPathSelector());
-		SceneOutliner->PasteFoldersBegin(SelectedFolderPaths);
-	}
+void FActorBrowsingMode::OnDuplicateActorsBegin()
+{
+	// Only a callback in actor browsing mode
+	const TArray<FName> SelectedFolderPaths = SceneOutliner->GetSelection().GetData<FName>(SceneOutliner::FFolderPathSelector());
+	SceneOutliner->PasteFoldersBegin(SelectedFolderPaths);
+}
 
-	void FActorBrowsingMode::OnDuplicateActorsEnd()
-	{
-		// Only a callback in actor browsing mode
-		SceneOutliner->PasteFoldersEnd();
-	}
+void FActorBrowsingMode::OnDuplicateActorsEnd()
+{
+	// Only a callback in actor browsing mode
+	SceneOutliner->PasteFoldersEnd();
+}
 
-	void FActorBrowsingMode::OnDeleteActorsBegin()
-	{
-		SceneOutliner->DeleteFoldersBegin();
-	}
+void FActorBrowsingMode::OnDeleteActorsBegin()
+{
+	SceneOutliner->DeleteFoldersBegin();
+}
 
-	void FActorBrowsingMode::OnDeleteActorsEnd()
-	{
-		SceneOutliner->DeleteFoldersEnd();
-	}
+void FActorBrowsingMode::OnDeleteActorsEnd()
+{
+	SceneOutliner->DeleteFoldersEnd();
 }
+
 #undef LOCTEXT_NAMESPACE
\ No newline at end of file
diff --git a/Engine/Source/Editor/SceneOutliner/Private/ActorFolderHierarchy.cpp b/Engine/Source/Editor/SceneOutliner/Private/ActorFolderHierarchy.cpp
--- a/Engine/Source/Editor/SceneOutliner/Private/ActorFolderHierarchy.cpp
+++ b/Engine/Source/Editor/SceneOutliner/Private/ActorFolderHierarchy.cpp
@@ -5,147 +5,144 @@
 #include "ActorFolderTreeItem.h"
 #include "EditorActorFolders.h"
 
-namespace SceneOutliner
+FActorFolderHierarchy::FActorFolderHierarchy(ISceneOutlinerMode* InMode, const TWeakObjectPtr<UWorld>& World)
+	: ISceneOutlinerHierarchy(InMode)
+	, RepresentingWorld(World)
 {
-	FActorFolderHierarchy::FActorFolderHierarchy(ISceneOutlinerMode* InMode, const TWeakObjectPtr<UWorld>& World)
-		: ISceneOutlinerHierarchy(InMode)
-		, RepresentingWorld(World)
+	// ActorFolderHierarchy should only be used with a mode which is showing folders
+	check(Mode->ShouldShowFolders());
+}
+
+FSceneOutlinerTreeItemPtr FActorFolderHierarchy::FindParent(const ISceneOutlinerTreeItem& Item, const TMap<FSceneOutlinerTreeItemID, FSceneOutlinerTreeItemPtr>& Items) const
+{
+	if (Item.IsA<FWorldTreeItem>())
 	{
-		// ActorFolderHierarchy should only be used with a mode which is showing folders
-		check(Mode->ShouldShowFolders());
+		return nullptr;
 	}
-
-	FTreeItemPtr FActorFolderHierarchy::FindParent(const ITreeItem& Item, const TMap<FTreeItemID, FTreeItemPtr>& Items) const
+	else if (const FActorFolderTreeItem* ActorFolderItem = Item.CastTo<FActorFolderTreeItem>())
 	{
-		if (Item.IsA<FWorldTreeItem>())
+		const FName ParentPath = SceneOutliner::GetParentPath(ActorFolderItem->Path);
+
+		const FSceneOutlinerTreeItemPtr* ParentItem = nullptr;
+		// If the folder has no parent path, it must be parented to the root world
+		if (ParentPath.IsNone())
 		{
-			return nullptr;
+			ParentItem = Items.Find(ActorFolderItem->World.Get());
 		}
-		else if (const FActorFolderTreeItem* ActorFolderItem = Item.CastTo<FActorFolderTreeItem>())
+		else
 		{
-			const FName ParentPath = GetParentPath(ActorFolderItem->Path);
-
-			const FTreeItemPtr* ParentItem = nullptr;
-			// If the folder has no parent path, it must be parented to the root world
-			if (ParentPath.IsNone())
-			{
-				ParentItem = Items.Find(ActorFolderItem->World.Get());
-			}
-			else
-			{
-				ParentItem = Items.Find(GetParentPath(ActorFolderItem->Path));
-			}
+			ParentItem = Items.Find(SceneOutliner::GetParentPath(ActorFolderItem->Path));
+		}
 
-			if (ParentItem)
-			{
-				return *ParentItem;
-			}
+		if (ParentItem)
+		{
+			return *ParentItem;
 		}
-		return nullptr;
 	}
+	return nullptr;
+}
 
-	void FActorFolderHierarchy::FindChildren(const ITreeItem& Item, const TMap<FTreeItemID, FTreeItemPtr>& Items, TArray<FTreeItemPtr>& OutChildren) const
+void FActorFolderHierarchy::FindChildren(const ISceneOutlinerTreeItem& Item, const TMap<FSceneOutlinerTreeItemID, FSceneOutlinerTreeItemPtr>& Items, TArray<FSceneOutlinerTreeItemPtr>& OutChildren) const
+{
+	if (const FWorldTreeItem* WorldTreeItem = Item.CastTo<FWorldTreeItem>())
 	{
-		if (const FWorldTreeItem* WorldTreeItem = Item.CastTo<FWorldTreeItem>())
+		// Find folders which are located at the root
+		for (const auto& Pair : FActorFolders::Get().GetFolderPropertiesForWorld(*WorldTreeItem->World))
 		{
-			// Find folders which are located at the root
-			for (const auto& Pair : FActorFolders::Get().GetFolderPropertiesForWorld(*WorldTreeItem->World))
+			if (const FSceneOutlinerTreeItemPtr* PotentialChild = Items.Find(Pair.Key))
 			{
-				if (const FTreeItemPtr* PotentialChild = Items.Find(Pair.Key))
+				if (const FFolderTreeItem* FolderItem = (*PotentialChild)->CastTo<FFolderTreeItem>())
 				{
-					if (const FFolderTreeItem* FolderItem = (*PotentialChild)->CastTo<FFolderTreeItem>())
+					if (SceneOutliner::GetParentPath(FolderItem->Path).IsNone())
 					{
-						if (GetParentPath(FolderItem->Path).IsNone())
-						{
-							OutChildren.Add(*PotentialChild);
-						}
+						OutChildren.Add(*PotentialChild);
 					}
 				}
 			}
 		}
-		else if (const FFolderTreeItem* FolderItem = Item.CastTo<FFolderTreeItem>())
+	}
+	else if (const FFolderTreeItem* FolderItem = Item.CastTo<FFolderTreeItem>())
+	{
+		// Search through all items and see if there is an item with a path which is a child of this one
+		for (const auto& Pair : Items)
 		{
-			// Search through all items and see if there is an item with a path which is a child of this one
-			for (const auto& Pair : Items)
+			if (const FFolderTreeItem* PotentialChild = Pair.Value->CastTo<FFolderTreeItem>())
 			{
-				if (const FFolderTreeItem* PotentialChild = Pair.Value->CastTo<FFolderTreeItem>())
+				if (SceneOutliner::PathIsChildOf(PotentialChild->Path, FolderItem->Path))
 				{
-					if (PathIsChildOf(PotentialChild->Path, FolderItem->Path))
-					{
-						OutChildren.Add(Pair.Value);
-					}
+					OutChildren.Add(Pair.Value);
 				}
 			}
 		}
 	}
+}
 
-	void FActorFolderHierarchy::CreateWorldChildren(UWorld* World, TArray<FTreeItemPtr>& OutItems) const
+void FActorFolderHierarchy::CreateWorldChildren(UWorld* World, TArray<FSceneOutlinerTreeItemPtr>& OutItems) const
+{
+	for (const auto& Pair : FActorFolders::Get().GetFolderPropertiesForWorld(*World))
 	{
-		for (const auto& Pair : FActorFolders::Get().GetFolderPropertiesForWorld(*World))
+		if (FSceneOutlinerTreeItemPtr FolderItem = Mode->CreateItemFor<FActorFolderTreeItem>(FActorFolderTreeItem(Pair.Key, World)))
 		{
-			if (FTreeItemPtr FolderItem = Mode->CreateItemFor<FActorFolderTreeItem>(FActorFolderTreeItem(Pair.Key, World)))
-			{
-				OutItems.Add(FolderItem);
-			}
+			OutItems.Add(FolderItem);
 		}
 	}
+}
+
+void FActorFolderHierarchy::CreateItems(TArray<FSceneOutlinerTreeItemPtr>& OutItems) const
+{
+	check(RepresentingWorld.IsValid());
 
-	void FActorFolderHierarchy::CreateItems(TArray<FTreeItemPtr>& OutItems) const
+	if (FSceneOutlinerTreeItemPtr WorldItem = Mode->CreateItemFor<FWorldTreeItem>(RepresentingWorld.Get()))
 	{
-		check(RepresentingWorld.IsValid());
+		OutItems.Add(WorldItem);
+	}
 
-		if (FTreeItemPtr WorldItem = Mode->CreateItemFor<FWorldTreeItem>(RepresentingWorld.Get()))
-		{
-			OutItems.Add(WorldItem);
-		}
+	CreateWorldChildren(RepresentingWorld.Get(), OutItems);
+}
 
-		CreateWorldChildren(RepresentingWorld.Get(), OutItems);
+void FActorFolderHierarchy::CreateChildren(const FSceneOutlinerTreeItemPtr& Item, TArray<FSceneOutlinerTreeItemPtr>& OutChildren) const
+{
+	if (const FWorldTreeItem* WorldItem = Item->CastTo<FWorldTreeItem>())
+	{
+		check(WorldItem->World == RepresentingWorld);
+		CreateWorldChildren(WorldItem->World.Get(), OutChildren);
 	}
-
-	void FActorFolderHierarchy::CreateChildren(const FTreeItemPtr& Item, TArray<FTreeItemPtr>& OutChildren) const
+	else if (FActorFolderTreeItem* FolderItem = Item->CastTo<FActorFolderTreeItem>())
 	{
-		if (const FWorldTreeItem* WorldItem = Item->CastTo<FWorldTreeItem>())
-		{
-			check(WorldItem->World == RepresentingWorld);
-			CreateWorldChildren(WorldItem->World.Get(), OutChildren);
-		}
-		else if (FActorFolderTreeItem* FolderItem = Item->CastTo<FActorFolderTreeItem>())
+		// Since no map of folder->children exist for ActorFolders, must iterate through all
+		// and manually check the path to know if it a child
+		for (const auto& Pair : FActorFolders::Get().GetFolderPropertiesForWorld(*(FolderItem->World)))
 		{
-			// Since no map of folder->children exist for ActorFolders, must iterate through all
-			// and manually check the path to know if it a child
-			for (const auto& Pair : FActorFolders::Get().GetFolderPropertiesForWorld(*(FolderItem->World)))
+			if (SceneOutliner::PathIsChildOf(Pair.Key, FolderItem->Path))
 			{
-				if (PathIsChildOf(Pair.Key, FolderItem->Path))
+				if (FSceneOutlinerTreeItemPtr NewFolderItem = Mode->CreateItemFor<FActorFolderTreeItem>(FActorFolderTreeItem(Pair.Key, FolderItem->World), true))
 				{
-					if (FTreeItemPtr NewFolderItem = Mode->CreateItemFor<FActorFolderTreeItem>(FActorFolderTreeItem(Pair.Key, FolderItem->World), true))
-					{
-						OutChildren.Add(NewFolderItem);
-					}
+					OutChildren.Add(NewFolderItem);
 				}
 			}
 		}
 	}
+}
 
-	FTreeItemPtr FActorFolderHierarchy::CreateParentItem(const FTreeItemPtr& Item) const
+FSceneOutlinerTreeItemPtr FActorFolderHierarchy::CreateParentItem(const FSceneOutlinerTreeItemPtr& Item) const
+{
+	if (Item->IsA<FWorldTreeItem>())
+	{
+		return nullptr;
+	}
+	else if (const FActorFolderTreeItem* FolderTreeItem = Item->CastTo<FActorFolderTreeItem>())
 	{
-		if (Item->IsA<FWorldTreeItem>())
+		const FName ParentPath = SceneOutliner::GetParentPath(FolderTreeItem->Path);
+		if (ParentPath.IsNone())
 		{
-			return nullptr;
+			UWorld* OwningWorld = FolderTreeItem->World.Get();
+			check(OwningWorld);
+			return Mode->CreateItemFor<FWorldTreeItem>(OwningWorld, true);
 		}
-		else if (const FActorFolderTreeItem* FolderTreeItem = Item->CastTo<FActorFolderTreeItem>())
+		else
 		{
-			const FName ParentPath = GetParentPath(FolderTreeItem->Path);
-			if (ParentPath.IsNone())
-			{
-				UWorld* OwningWorld = FolderTreeItem->World.Get();
-				check(OwningWorld);
-				return Mode->CreateItemFor<FWorldTreeItem>(OwningWorld, true);
-			}
-			else
-			{
-				return Mode->CreateItemFor<FActorFolderTreeItem>(FActorFolderTreeItem(ParentPath, FolderTreeItem->World));
-			}
+			return Mode->CreateItemFor<FActorFolderTreeItem>(FActorFolderTreeItem(ParentPath, FolderTreeItem->World));
 		}
-		return nullptr;
 	}
-}
\ No newline at end of file
+	return nullptr;
+}
diff --git a/Engine/Source/Editor/SceneOutliner/Private/ActorFolderPickingMode.cpp b/Engine/Source/Editor/SceneOutliner/Private/ActorFolderPickingMode.cpp
--- a/Engine/Source/Editor/SceneOutliner/Private/ActorFolderPickingMode.cpp
+++ b/Engine/Source/Editor/SceneOutliner/Private/ActorFolderPickingMode.cpp
@@ -6,53 +6,50 @@
 
 #define LOCTEXT_NAMESPACE "SceneOutliner_ActoFolderPickingMode"
 
-namespace SceneOutliner
+FActorFolderPickingMode::FActorFolderPickingMode(SSceneOutliner* InSceneOutliner, FOnSceneOutlinerItemPicked InOnItemPicked, TWeakObjectPtr<UWorld> InSpecifiedWorldToDisplay)
+	: FActorMode(InSceneOutliner, true, InSpecifiedWorldToDisplay)
+	, OnItemPicked(InOnItemPicked)
 {
-	FActorFolderPickingMode::FActorFolderPickingMode(SSceneOutliner* InSceneOutliner, FOnSceneOutlinerItemPicked InOnItemPicked, TWeakObjectPtr<UWorld> InSpecifiedWorldToDisplay)
-		: FActorMode(InSceneOutliner, true, InSpecifiedWorldToDisplay)
-		, OnItemPicked(InOnItemPicked)
-	{
-		Rebuild();
-	}
+	Rebuild();
+}
 
-	TUniquePtr<ISceneOutlinerHierarchy> FActorFolderPickingMode::CreateHierarchy()
-	{
-		return MakeUnique<FActorFolderHierarchy>(this, RepresentingWorld);
-	}
+TUniquePtr<ISceneOutlinerHierarchy> FActorFolderPickingMode::CreateHierarchy()
+{
+	return MakeUnique<FActorFolderHierarchy>(this, RepresentingWorld);
+}
 
-	void FActorFolderPickingMode::OnItemSelectionChanged(FTreeItemPtr Item, ESelectInfo::Type SelectionType, const FItemSelection& Selection)
+void FActorFolderPickingMode::OnItemSelectionChanged(FSceneOutlinerTreeItemPtr Item, ESelectInfo::Type SelectionType, const FSceneOutlinerItemSelection& Selection)
+{
+	// fire off the notification to whoever is listening.
+	// This may often cause the widget itself to be enqueued for destruction
+	auto SelectedItems = SceneOutliner->GetSelectedItems();
+	if (SelectedItems.Num() > 0)
 	{
-		// fire off the notification to whoever is listening.
-		// This may often cause the widget itself to be enqueued for destruction
-		auto SelectedItems = SceneOutliner->GetSelectedItems();
-		if (SelectedItems.Num() > 0)
+		auto FirstItem = SelectedItems[0];
+		if (FirstItem->CanInteract())
 		{
-			auto FirstItem = SelectedItems[0];
-			if (FirstItem->CanInteract())
-			{
-				OnItemPicked.ExecuteIfBound(FirstItem.ToSharedRef());
-			}
+			OnItemPicked.ExecuteIfBound(FirstItem.ToSharedRef());
 		}
 	}
+}
 
-	void FActorFolderPickingMode::OnFilterTextCommited(FItemSelection& Selection, ETextCommit::Type CommitType)
+void FActorFolderPickingMode::OnFilterTextCommited(FSceneOutlinerItemSelection& Selection, ETextCommit::Type CommitType)
+{
+	// In folder picking mode, we check to see if we have a selected folder, and if so, fire
+	// off the notification to whoever is listening.  This may often cause the widget itself
+	// to be enqueued for destruction
+	TArray<FActorFolderTreeItem*> Folders;
+	Selection.Get(Folders);
+	if (Folders.Num() == 1 && Folders[0])
 	{
-		// In folder picking mode, we check to see if we have a selected folder, and if so, fire
-		// off the notification to whoever is listening.  This may often cause the widget itself
-		// to be enqueued for destruction
-		TArray<FActorFolderTreeItem*> Folders;
-		Selection.Get(Folders);
-		if (Folders.Num() == 1 && Folders[0])
-		{
-			// Signal that a folder was selected. We assume it is valid as it wouldn't not have been added to the selection if not
-			SceneOutliner->SetItemSelection(Folders[0]->AsShared(), true, ESelectInfo::OnKeyPress);
-		}
+		// Signal that a folder was selected. We assume it is valid as it wouldn't not have been added to the selection if not
+		SceneOutliner->SetItemSelection(Folders[0]->AsShared(), true, ESelectInfo::OnKeyPress);
 	}
+}
 
-	void FActorFolderPickingMode::CreateViewContent(FMenuBuilder& MenuBuilder)
-	{
-		FActorMode::BuildWorldPickerMenu(MenuBuilder);
-	}
+void FActorFolderPickingMode::CreateViewContent(FMenuBuilder& MenuBuilder)
+{
+	FActorMode::BuildWorldPickerMenu(MenuBuilder);
 }
 
 #undef LOCTEXT_NAMESPACE
\ No newline at end of file
