Removed SceneOutliner namespace

diff --git a/Engine/Source/Editor/SceneOutliner/Private/ActorFolderTreeItem.cpp b/Engine/Source/Editor/SceneOutliner/Private/ActorFolderTreeItem.cpp
--- a/Engine/Source/Editor/SceneOutliner/Private/ActorFolderTreeItem.cpp
+++ b/Engine/Source/Editor/SceneOutliner/Private/ActorFolderTreeItem.cpp
@@ -10,268 +10,265 @@
 
 #define LOCTEXT_NAMESPACE "SceneOutliner_ActorFolderTreeItem"
 
-namespace SceneOutliner
+const FSceneOutlinerTreeItemType FActorFolderTreeItem::Type(&FFolderTreeItem::Type);
+
+struct SActorFolderTreeLabel : FSceneOutlinerCommonLabelData, public SCompoundWidget
 {
-	const FTreeItemType FActorFolderTreeItem::Type(&FFolderTreeItem::Type);
+	SLATE_BEGIN_ARGS(SActorFolderTreeLabel) {}
+	SLATE_END_ARGS()
 
-	struct SActorFolderTreeLabel : FCommonLabelData, public SCompoundWidget
+	void Construct(const FArguments& InArgs, FActorFolderTreeItem& FolderItem, ISceneOutliner& SceneOutliner, const STableRow<FSceneOutlinerTreeItemPtr>& InRow)
 	{
-		SLATE_BEGIN_ARGS(SActorFolderTreeLabel) {}
-		SLATE_END_ARGS()
+		TreeItemPtr = StaticCastSharedRef<FActorFolderTreeItem>(FolderItem.AsShared());
+		WeakSceneOutliner = StaticCastSharedRef<ISceneOutliner>(SceneOutliner.AsShared());
+
+		TSharedPtr<SInlineEditableTextBlock> InlineTextBlock = SNew(SInlineEditableTextBlock)
+			.Text(this, &SActorFolderTreeLabel::GetDisplayText)
+			.HighlightText(SceneOutliner.GetFilterHighlightText())
+			.ColorAndOpacity(this, &SActorFolderTreeLabel::GetForegroundColor)
+			.OnTextCommitted(this, &SActorFolderTreeLabel::OnLabelCommitted)
+			.OnVerifyTextChanged(this, &SActorFolderTreeLabel::OnVerifyItemLabelChanged)
+			.IsSelected(FIsSelected::CreateSP(&InRow, &STableRow<FSceneOutlinerTreeItemPtr>::IsSelectedExclusively))
+			.IsReadOnly_Lambda([Item = FolderItem.AsShared(), this]()
+		{
+			return !CanExecuteRenameRequest(Item.Get());
+		});
 
-		void Construct(const FArguments& InArgs, FActorFolderTreeItem& FolderItem, ISceneOutliner& SceneOutliner, const STableRow<FTreeItemPtr>& InRow)
+		if (WeakSceneOutliner.Pin()->GetMode()->IsInteractive())
 		{
-			TreeItemPtr = StaticCastSharedRef<FActorFolderTreeItem>(FolderItem.AsShared());
-			WeakSceneOutliner = StaticCastSharedRef<ISceneOutliner>(SceneOutliner.AsShared());
-
-			TSharedPtr<SInlineEditableTextBlock> InlineTextBlock = SNew(SInlineEditableTextBlock)
-				.Text(this, &SActorFolderTreeLabel::GetDisplayText)
-				.HighlightText(SceneOutliner.GetFilterHighlightText())
-				.ColorAndOpacity(this, &SActorFolderTreeLabel::GetForegroundColor)
-				.OnTextCommitted(this, &SActorFolderTreeLabel::OnLabelCommitted)
-				.OnVerifyTextChanged(this, &SActorFolderTreeLabel::OnVerifyItemLabelChanged)
-				.IsSelected(FIsSelected::CreateSP(&InRow, &STableRow<FTreeItemPtr>::IsSelectedExclusively))
-				.IsReadOnly_Lambda([Item = FolderItem.AsShared(), this]()
-			{
-				return !CanExecuteRenameRequest(Item.Get());
-			});
+			FolderItem.RenameRequestEvent.BindSP(InlineTextBlock.Get(), &SInlineEditableTextBlock::EnterEditingMode);
+		}
 
-			if (WeakSceneOutliner.Pin()->GetMode()->IsInteractive())
-			{
-				FolderItem.RenameRequestEvent.BindSP(InlineTextBlock.Get(), &SInlineEditableTextBlock::EnterEditingMode);
-			}
+		ChildSlot
+			[
+				SNew(SHorizontalBox)
+
+				+ SHorizontalBox::Slot()
+			.AutoWidth()
+			.VAlign(VAlign_Center)
+			.Padding(FSceneOutlinerDefaultTreeItemMetrics::IconPadding())
+			[
+				SNew(SBox)
+				.WidthOverride(FSceneOutlinerDefaultTreeItemMetrics::IconSize())
+			[
+				SNew(SImage)
+				.Image(this, &SActorFolderTreeLabel::GetIcon)
+			]
+			]
+
+		+ SHorizontalBox::Slot()
+			.FillWidth(1.0f)
+			.VAlign(VAlign_Center)
+			.Padding(0.0f, 2.0f)
+			[
+				InlineTextBlock.ToSharedRef()
+			]
+			];
+	}
 
-			ChildSlot
-				[
-					SNew(SHorizontalBox)
-
-					+ SHorizontalBox::Slot()
-				.AutoWidth()
-				.VAlign(VAlign_Center)
-				.Padding(FDefaultTreeItemMetrics::IconPadding())
-				[
-					SNew(SBox)
-					.WidthOverride(FDefaultTreeItemMetrics::IconSize())
-				[
-					SNew(SImage)
-					.Image(this, &SActorFolderTreeLabel::GetIcon)
-				]
-				]
-
-			+ SHorizontalBox::Slot()
-				.FillWidth(1.0f)
-				.VAlign(VAlign_Center)
-				.Padding(0.0f, 2.0f)
-				[
-					InlineTextBlock.ToSharedRef()
-				]
-				];
-		}
+private:
+	TWeakPtr<FActorFolderTreeItem> TreeItemPtr;
 
-	private:
-		TWeakPtr<FActorFolderTreeItem> TreeItemPtr;
+	FText GetDisplayText() const
+	{
+		auto Folder = TreeItemPtr.Pin();
+		return Folder.IsValid() ? FText::FromName(Folder->LeafName) : FText();
+	}
 
-		FText GetDisplayText() const
+	const FSlateBrush* GetIcon() const
+	{
+		auto TreeItem = TreeItemPtr.Pin();
+		if (!TreeItem.IsValid())
 		{
-			auto Folder = TreeItemPtr.Pin();
-			return Folder.IsValid() ? FText::FromName(Folder->LeafName) : FText();
+			return FEditorStyle::Get().GetBrush(TEXT("SceneOutliner.FolderClosed"));
 		}
 
-		const FSlateBrush* GetIcon() const
+		if (TreeItem->Flags.bIsExpanded && TreeItem->GetChildren().Num())
 		{
-			auto TreeItem = TreeItemPtr.Pin();
-			if (!TreeItem.IsValid())
-			{
-				return FEditorStyle::Get().GetBrush(TEXT("SceneOutliner.FolderClosed"));
-			}
+			return FEditorStyle::Get().GetBrush(TEXT("SceneOutliner.FolderOpen"));
+		}
+		else
+		{
+			return FEditorStyle::Get().GetBrush(TEXT("SceneOutliner.FolderClosed"));
+		}
+	}
 
-			if (TreeItem->Flags.bIsExpanded && TreeItem->GetChildren().Num())
-			{
-				return FEditorStyle::Get().GetBrush(TEXT("SceneOutliner.FolderOpen"));
-			}
-			else
-			{
-				return FEditorStyle::Get().GetBrush(TEXT("SceneOutliner.FolderClosed"));
-			}
+	FSlateColor GetForegroundColor() const
+	{
+		if (auto BaseColor = FSceneOutlinerCommonLabelData::GetForegroundColor(*TreeItemPtr.Pin()))
+		{
+			return BaseColor.GetValue();
 		}
 
-		FSlateColor GetForegroundColor() const
+		return FSlateColor::UseForeground();
+	}
+
+	bool OnVerifyItemLabelChanged(const FText& InLabel, FText& OutErrorMessage)
+	{
+		auto TreeItem = TreeItemPtr.Pin();
+		if (!TreeItem.IsValid())
 		{
-			if (auto BaseColor = FCommonLabelData::GetForegroundColor(*TreeItemPtr.Pin()))
-			{
-				return BaseColor.GetValue();
-			}
+			OutErrorMessage = LOCTEXT("RenameFailed_TreeItemDeleted", "Tree item no longer exists");
+			return false;
+		}
 
-			return FSlateColor::UseForeground();
+		const FText TrimmedLabel = FText::TrimPrecedingAndTrailing(InLabel);
+
+		if (TrimmedLabel.IsEmpty())
+		{
+			OutErrorMessage = LOCTEXT("RenameFailed_LeftBlank", "Names cannot be left blank");
+			return false;
 		}
 
-		bool OnVerifyItemLabelChanged(const FText& InLabel, FText& OutErrorMessage)
+		if (TrimmedLabel.ToString().Len() >= NAME_SIZE)
 		{
-			auto TreeItem = TreeItemPtr.Pin();
-			if (!TreeItem.IsValid())
-			{
-				OutErrorMessage = LOCTEXT("RenameFailed_TreeItemDeleted", "Tree item no longer exists");
-				return false;
-			}
+			FFormatNamedArguments Arguments;
+			Arguments.Add(TEXT("CharCount"), NAME_SIZE);
+			OutErrorMessage = FText::Format(LOCTEXT("RenameFailed_TooLong", "Names must be less than {CharCount} characters long."), Arguments);
+			return false;
+		}
 
-			const FText TrimmedLabel = FText::TrimPrecedingAndTrailing(InLabel);
+		const FString LabelString = TrimmedLabel.ToString();
+		if (TreeItem->LeafName.ToString() == LabelString)
+		{
+			return true;
+		}
 
-			if (TrimmedLabel.IsEmpty())
-			{
-				OutErrorMessage = LOCTEXT("RenameFailed_LeftBlank", "Names cannot be left blank");
-				return false;
-			}
+		int32 Dummy = 0;
+		if (LabelString.FindChar('/', Dummy) || LabelString.FindChar('\\', Dummy))
+		{
+			OutErrorMessage = LOCTEXT("RenameFailed_InvalidChar", "Folder names cannot contain / or \\.");
+			return false;
+		}
 
-			if (TrimmedLabel.ToString().Len() >= NAME_SIZE)
-			{
-				FFormatNamedArguments Arguments;
-				Arguments.Add(TEXT("CharCount"), NAME_SIZE);
-				OutErrorMessage = FText::Format(LOCTEXT("RenameFailed_TooLong", "Names must be less than {CharCount} characters long."), Arguments);
-				return false;
-			}
+		// Validate that this folder doesn't exist already
+		FName NewPath = SceneOutliner::GetParentPath(TreeItem->Path);
+		if (NewPath.IsNone())
+		{
+			NewPath = FName(*LabelString);
+		}
+		else
+		{
+			NewPath = FName(*(NewPath.ToString() / LabelString));
+		}
 
-			const FString LabelString = TrimmedLabel.ToString();
-			if (TreeItem->LeafName.ToString() == LabelString)
-			{
-				return true;
-			}
+		if (FActorFolders::Get().GetFolderProperties(*GWorld, NewPath))
+		{
+			OutErrorMessage = LOCTEXT("RenameFailed_AlreadyExists", "A folder with this name already exists at this level");
+			return false;
+		}
 
-			int32 Dummy = 0;
-			if (LabelString.FindChar('/', Dummy) || LabelString.FindChar('\\', Dummy))
-			{
-				OutErrorMessage = LOCTEXT("RenameFailed_InvalidChar", "Folder names cannot contain / or \\.");
-				return false;
-			}
+		return true;
+	}
 
-			// Validate that this folder doesn't exist already
-			FName NewPath = GetParentPath(TreeItem->Path);
+	void OnLabelCommitted(const FText& InLabel, ETextCommit::Type InCommitInfo)
+	{
+		auto TreeItem = TreeItemPtr.Pin();
+		if (TreeItem.IsValid() && !InLabel.ToString().Equals(TreeItem->LeafName.ToString(), ESearchCase::CaseSensitive))
+		{
+			// Rename the item
+			FName NewPath = SceneOutliner::GetParentPath(TreeItem->Path);
 			if (NewPath.IsNone())
 			{
-				NewPath = FName(*LabelString);
+				NewPath = FName(*InLabel.ToString());
 			}
 			else
 			{
-				NewPath = FName(*(NewPath.ToString() / LabelString));
+				NewPath = FName(*(NewPath.ToString() / InLabel.ToString()));
 			}
 
-			if (FActorFolders::Get().GetFolderProperties(*GWorld, NewPath))
-			{
-				OutErrorMessage = LOCTEXT("RenameFailed_AlreadyExists", "A folder with this name already exists at this level");
-				return false;
-			}
-
-			return true;
-		}
+			FActorFolders::Get().RenameFolderInWorld(*GWorld, TreeItem->Path, NewPath);
 
-		void OnLabelCommitted(const FText& InLabel, ETextCommit::Type InCommitInfo)
-		{
-			auto TreeItem = TreeItemPtr.Pin();
-			if (TreeItem.IsValid() && !InLabel.ToString().Equals(TreeItem->LeafName.ToString(), ESearchCase::CaseSensitive))
+			auto Outliner = WeakSceneOutliner.Pin();
+			if (Outliner.IsValid())
 			{
-				// Rename the item
-				FName NewPath = GetParentPath(TreeItem->Path);
-				if (NewPath.IsNone())
-				{
-					NewPath = FName(*InLabel.ToString());
-				}
-				else
-				{
-					NewPath = FName(*(NewPath.ToString() / InLabel.ToString()));
-				}
-
-				FActorFolders::Get().RenameFolderInWorld(*GWorld, TreeItem->Path, NewPath);
-
-				auto Outliner = WeakSceneOutliner.Pin();
-				if (Outliner.IsValid())
-				{
-					Outliner->SetKeyboardFocus();
-				}
+				Outliner->SetKeyboardFocus();
 			}
 		}
-	};
+	}
+};
 
-	FActorFolderTreeItem::FActorFolderTreeItem(FName InPath, TWeakObjectPtr<UWorld> InWorld)
-		: FFolderTreeItem(InPath, Type)
-		, World(InWorld)
+FActorFolderTreeItem::FActorFolderTreeItem(FName InPath, TWeakObjectPtr<UWorld> InWorld)
+	: FFolderTreeItem(InPath, Type)
+	, World(InWorld)
+{
+}
+
+void FActorFolderTreeItem::OnExpansionChanged()
+{
+	if (!World.IsValid())
 	{
+		return;
 	}
 
-	void FActorFolderTreeItem::OnExpansionChanged()
+	// Update the central store of folder properties with this folder's new expansion state
+	if (FActorFolderProps* Props = FActorFolders::Get().GetFolderProperties(*World, Path))
 	{
-		if (!World.IsValid())
-		{
-			return;
-		}
-
-		// Update the central store of folder properties with this folder's new expansion state
-		if (FActorFolderProps* Props = FActorFolders::Get().GetFolderProperties(*World, Path))
-		{
-			Props->bIsExpanded = Flags.bIsExpanded;
-		}
+		Props->bIsExpanded = Flags.bIsExpanded;
 	}
+}
 
-	void FActorFolderTreeItem::Delete(FName InNewParentPath)
+void FActorFolderTreeItem::Delete(FName InNewParentPath)
+{
+	if (!World.IsValid())
 	{
-		if (!World.IsValid())
-		{
-			return;
-		}
-		const FScopedTransaction Transaction(LOCTEXT("DeleteFolderTransaction", "Delete Folder"));
-
-		for (const TWeakPtr<ITreeItem>& ChildPtr : GetChildren())
-		{
-			FTreeItemPtr Child = ChildPtr.Pin();
-			if (const FActorTreeItem* ActorItem = Child->CastTo<FActorTreeItem>())
-			{
-				if (AActor* Actor = ActorItem->Actor.Get())
-				{
-					Actor->SetFolderPath_Recursively(InNewParentPath);
-				}
-			}
-			else if (FFolderTreeItem* FolderItem = Child->CastTo<FFolderTreeItem>())
-			{
-				FolderItem->MoveTo(InNewParentPath);
-			}
-		}
-
-		FActorFolders::Get().DeleteFolder(*World, Path);
+		return;
 	}
+	const FScopedTransaction Transaction(LOCTEXT("DeleteFolderTransaction", "Delete Folder"));
 
-	FName FActorFolderTreeItem::MoveTo(const FName& NewParent)
+	for (const TWeakPtr<ISceneOutlinerTreeItem>& ChildPtr : GetChildren())
 	{
-		if (World.IsValid())
+		FSceneOutlinerTreeItemPtr Child = ChildPtr.Pin();
+		if (const FActorTreeItem* ActorItem = Child->CastTo<FActorTreeItem>())
 		{
-			// Get unique name
-			const FName NewPath = FActorFolders::Get().GetFolderName(*World, NewParent, LeafName);
-
-			if (FActorFolders::Get().RenameFolderInWorld(*World, Path, NewPath))
+			if (AActor* Actor = ActorItem->Actor.Get())
 			{
-				return NewPath;
+				Actor->SetFolderPath_Recursively(InNewParentPath);
 			}
 		}
-		return FName();
+		else if (FFolderTreeItem* FolderItem = Child->CastTo<FFolderTreeItem>())
+		{
+			FolderItem->MoveTo(InNewParentPath);
+		}
 	}
 
-	void FActorFolderTreeItem::CreateSubFolder(TWeakPtr<SSceneOutliner> WeakOutliner)
+	FActorFolders::Get().DeleteFolder(*World, Path);
+}
+
+FName FActorFolderTreeItem::MoveTo(const FName& NewParent)
+{
+	if (World.IsValid())
 	{
-		auto Outliner = WeakOutliner.Pin();
+		// Get unique name
+		const FName NewPath = FActorFolders::Get().GetFolderName(*World, NewParent, LeafName);
 
-		if (Outliner.IsValid() && World.IsValid())
+		if (FActorFolders::Get().RenameFolderInWorld(*World, Path, NewPath))
 		{
-			const FScopedTransaction Transaction(LOCTEXT("UndoAction_CreateFolder", "Create Folder"));
-
-			const FName NewFolderName = FActorFolders::Get().GetDefaultFolderName(*World, Path);
-			FActorFolders::Get().CreateFolder(*World, NewFolderName);
-
-			// At this point the new folder will be in our newly added list, so select it and open a rename when it gets refreshed
-			Outliner->OnItemAdded(NewFolderName, ENewItemAction::Select | ENewItemAction::Rename);
+			return NewPath;
 		}
 	}
+	return FName();
+}
+
+void FActorFolderTreeItem::CreateSubFolder(TWeakPtr<SSceneOutliner> WeakOutliner)
+{
+	auto Outliner = WeakOutliner.Pin();
 
-	TSharedRef<SWidget> FActorFolderTreeItem::GenerateLabelWidget(ISceneOutliner& Outliner, const STableRow<FTreeItemPtr>& InRow)
+	if (Outliner.IsValid() && World.IsValid())
 	{
-		return SNew(SActorFolderTreeLabel, *this, Outliner, InRow);
+		const FScopedTransaction Transaction(LOCTEXT("UndoAction_CreateFolder", "Create Folder"));
+
+		const FName NewFolderName = FActorFolders::Get().GetDefaultFolderName(*World, Path);
+		FActorFolders::Get().CreateFolder(*World, NewFolderName);
+
+		// At this point the new folder will be in our newly added list, so select it and open a rename when it gets refreshed
+		Outliner->OnItemAdded(NewFolderName, SceneOutliner::ENewItemAction::Select | SceneOutliner::ENewItemAction::Rename);
 	}
 }
 
+TSharedRef<SWidget> FActorFolderTreeItem::GenerateLabelWidget(ISceneOutliner& Outliner, const STableRow<FSceneOutlinerTreeItemPtr>& InRow)
+{
+	return SNew(SActorFolderTreeLabel, *this, Outliner, InRow);
+}
+
 #undef LOCTEXT_NAMESPACE
\ No newline at end of file
diff --git a/Engine/Source/Editor/SceneOutliner/Private/ActorHierarchy.cpp b/Engine/Source/Editor/SceneOutliner/Private/ActorHierarchy.cpp
--- a/Engine/Source/Editor/SceneOutliner/Private/ActorHierarchy.cpp
+++ b/Engine/Source/Editor/SceneOutliner/Private/ActorHierarchy.cpp
@@ -14,498 +14,496 @@
 #include "GameFramework/WorldSettings.h"
 #include "EditorActorFolders.h"
 
-namespace SceneOutliner
+TUniquePtr<FActorHierarchy> FActorHierarchy::Create(ISceneOutlinerMode* Mode, const TWeakObjectPtr<UWorld>& World, TFunction<bool(const AActor*)> InIsActorDisplayable)
 {
-	TUniquePtr<FActorHierarchy> FActorHierarchy::Create(ISceneOutlinerMode* Mode, const TWeakObjectPtr<UWorld>& World, TFunction<bool(const AActor*)> InIsActorDisplayable)
-	{
-		FActorHierarchy* Hierarchy = new FActorHierarchy(Mode, World, InIsActorDisplayable);
+	FActorHierarchy* Hierarchy = new FActorHierarchy(Mode, World, InIsActorDisplayable);
 		
-		GEngine->OnLevelActorAdded().AddRaw(Hierarchy, &FActorHierarchy::OnLevelActorAdded);
-		GEngine->OnLevelActorDeleted().AddRaw(Hierarchy, &FActorHierarchy::OnLevelActorDeleted);
-		GEngine->OnLevelActorDetached().AddRaw(Hierarchy, &FActorHierarchy::OnLevelActorDetached);
-		GEngine->OnLevelActorAttached().AddRaw(Hierarchy, &FActorHierarchy::OnLevelActorAttached);
-		GEngine->OnLevelActorFolderChanged().AddRaw(Hierarchy, &FActorHierarchy::OnLevelActorFolderChanged);
-		GEngine->OnLevelActorListChanged().AddRaw(Hierarchy, &FActorHierarchy::OnLevelActorListChanged);
-
-		FWorldDelegates::LevelAddedToWorld.AddRaw(Hierarchy, &FActorHierarchy::OnLevelAdded);
-		FWorldDelegates::LevelRemovedFromWorld.AddRaw(Hierarchy, &FActorHierarchy::OnLevelRemoved);
-
-		auto& Folders = FActorFolders::Get();
-		Folders.OnFolderCreate.AddRaw(Hierarchy, &FActorHierarchy::OnBroadcastFolderCreate);
-		Folders.OnFolderMove.AddRaw(Hierarchy, &FActorHierarchy::OnBroadcastFolderMove);
-		Folders.OnFolderDelete.AddRaw(Hierarchy, &FActorHierarchy::OnBroadcastFolderDelete);
+	GEngine->OnLevelActorAdded().AddRaw(Hierarchy, &FActorHierarchy::OnLevelActorAdded);
+	GEngine->OnLevelActorDeleted().AddRaw(Hierarchy, &FActorHierarchy::OnLevelActorDeleted);
+	GEngine->OnLevelActorDetached().AddRaw(Hierarchy, &FActorHierarchy::OnLevelActorDetached);
+	GEngine->OnLevelActorAttached().AddRaw(Hierarchy, &FActorHierarchy::OnLevelActorAttached);
+	GEngine->OnLevelActorFolderChanged().AddRaw(Hierarchy, &FActorHierarchy::OnLevelActorFolderChanged);
+	GEngine->OnLevelActorListChanged().AddRaw(Hierarchy, &FActorHierarchy::OnLevelActorListChanged);
+
+	FWorldDelegates::LevelAddedToWorld.AddRaw(Hierarchy, &FActorHierarchy::OnLevelAdded);
+	FWorldDelegates::LevelRemovedFromWorld.AddRaw(Hierarchy, &FActorHierarchy::OnLevelRemoved);
+
+	auto& Folders = FActorFolders::Get();
+	Folders.OnFolderCreate.AddRaw(Hierarchy, &FActorHierarchy::OnBroadcastFolderCreate);
+	Folders.OnFolderMove.AddRaw(Hierarchy, &FActorHierarchy::OnBroadcastFolderMove);
+	Folders.OnFolderDelete.AddRaw(Hierarchy, &FActorHierarchy::OnBroadcastFolderDelete);
+
+	return TUniquePtr<FActorHierarchy>(Hierarchy);
+}
 
-		return TUniquePtr<FActorHierarchy>(Hierarchy);
-	}
+FActorHierarchy::FActorHierarchy(ISceneOutlinerMode* Mode, const TWeakObjectPtr<UWorld>& World, TFunction<bool(const AActor*)> InIsActorDisplayable)
+	: ISceneOutlinerHierarchy(Mode)
+	, RepresentingWorld(World)
+	, IsActorDisplayableCallback(InIsActorDisplayable)
+{
+}
 
-	FActorHierarchy::FActorHierarchy(ISceneOutlinerMode* Mode, const TWeakObjectPtr<UWorld>& World, TFunction<bool(const AActor*)> InIsActorDisplayable)
-		: ISceneOutlinerHierarchy(Mode)
-		, RepresentingWorld(World)
-		, IsActorDisplayableCallback(InIsActorDisplayable)
+FActorHierarchy::~FActorHierarchy()
+{
+	if (GEngine)
 	{
+		GEngine->OnLevelActorAdded().RemoveAll(this);
+		GEngine->OnLevelActorDeleted().RemoveAll(this);
+		GEngine->OnLevelActorDetached().RemoveAll(this);
+		GEngine->OnLevelActorAttached().RemoveAll(this);
+		GEngine->OnLevelActorFolderChanged().RemoveAll(this);
+		GEngine->OnLevelActorListChanged().RemoveAll(this);
 	}
 
-	FActorHierarchy::~FActorHierarchy()
-	{
-		if (GEngine)
-		{
-			GEngine->OnLevelActorAdded().RemoveAll(this);
-			GEngine->OnLevelActorDeleted().RemoveAll(this);
-			GEngine->OnLevelActorDetached().RemoveAll(this);
-			GEngine->OnLevelActorAttached().RemoveAll(this);
-			GEngine->OnLevelActorFolderChanged().RemoveAll(this);
-			GEngine->OnLevelActorListChanged().RemoveAll(this);
-		}
-
-		FWorldDelegates::LevelAddedToWorld.RemoveAll(this);
-		FWorldDelegates::LevelRemovedFromWorld.RemoveAll(this);
+	FWorldDelegates::LevelAddedToWorld.RemoveAll(this);
+	FWorldDelegates::LevelRemovedFromWorld.RemoveAll(this);
 
 		
-		if (FActorFolders::IsAvailable())
-		{
-			auto& Folders = FActorFolders::Get();
-			Folders.OnFolderCreate.RemoveAll(this);
-			Folders.OnFolderMove.RemoveAll(this);
-			Folders.OnFolderDelete.RemoveAll(this);
-		}
+	if (FActorFolders::IsAvailable())
+	{
+		auto& Folders = FActorFolders::Get();
+		Folders.OnFolderCreate.RemoveAll(this);
+		Folders.OnFolderMove.RemoveAll(this);
+		Folders.OnFolderDelete.RemoveAll(this);
 	}
+}
 
-	void FActorHierarchy::FindChildren(const ITreeItem& Item, const TMap<FTreeItemID, FTreeItemPtr>& Items, TArray<FTreeItemPtr>& OutChildren) const
+void FActorHierarchy::FindChildren(const ISceneOutlinerTreeItem& Item, const TMap<FSceneOutlinerTreeItemID, FSceneOutlinerTreeItemPtr>& Items, TArray<FSceneOutlinerTreeItemPtr>& OutChildren) const
+{
+	if (const FWorldTreeItem* WorldTreeItem = Item.CastTo<FWorldTreeItem>())
 	{
-		if (const FWorldTreeItem* WorldTreeItem = Item.CastTo<FWorldTreeItem>())
+		// All actors in the world are children of the world
+		for (FActorIterator ActorIt(WorldTreeItem->World.Get()); ActorIt; ++ActorIt)
 		{
-			// All actors in the world are children of the world
-			for (FActorIterator ActorIt(WorldTreeItem->World.Get()); ActorIt; ++ActorIt)
+			// Actors which have parents are not direct children of a world
+			if (ActorIt->GetAttachParentActor() == nullptr)
 			{
-				// Actors which have parents are not direct children of a world
-				if (ActorIt->GetAttachParentActor() == nullptr)
+				if (const FSceneOutlinerTreeItemPtr* ChildItem = Items.Find(*ActorIt))
 				{
-					if (const FTreeItemPtr* ChildItem = Items.Find(*ActorIt))
-					{
-						OutChildren.Add(*ChildItem);
-					}
+					OutChildren.Add(*ChildItem);
 				}
 			}
+		}
 
-			if (Mode->ShouldShowFolders())
+		if (Mode->ShouldShowFolders())
+		{
+			// Find folders which are located at the root
+			for (const auto& Pair : FActorFolders::Get().GetFolderPropertiesForWorld(*WorldTreeItem->World))
 			{
-				// Find folders which are located at the root
-				for (const auto& Pair : FActorFolders::Get().GetFolderPropertiesForWorld(*WorldTreeItem->World))
+				if (const FSceneOutlinerTreeItemPtr* PotentialChild = Items.Find(Pair.Key))
 				{
-					if (const FTreeItemPtr* PotentialChild = Items.Find(Pair.Key))
+					if (const FFolderTreeItem* FolderItem = (*PotentialChild)->CastTo<FFolderTreeItem>())
 					{
-						if (const FFolderTreeItem* FolderItem = (*PotentialChild)->CastTo<FFolderTreeItem>())
+						if (SceneOutliner::GetParentPath(FolderItem->Path).IsNone())
 						{
-							if (GetParentPath(FolderItem->Path).IsNone())
-							{
-								OutChildren.Add(*PotentialChild);
-							}
+							OutChildren.Add(*PotentialChild);
 						}
 					}
 				}
 			}
 		}
-		else if (const FActorTreeItem* ActorTreeItem = Item.CastTo<FActorTreeItem>())
+	}
+	else if (const FActorTreeItem* ActorTreeItem = Item.CastTo<FActorTreeItem>())
+	{
+		// All attached actors are children of this actor
+		TArray<AActor*> AttachedActors;
+		ActorTreeItem->Actor->GetAttachedActors(AttachedActors);
+		for (const AActor* AttachedActor : AttachedActors)
 		{
-			// All attached actors are children of this actor
-			TArray<AActor*> AttachedActors;
-			ActorTreeItem->Actor->GetAttachedActors(AttachedActors);
-			for (const AActor* AttachedActor : AttachedActors)
+			if (const FSceneOutlinerTreeItemPtr* AttachedActorItem = Items.Find(AttachedActor))
 			{
-				if (const FTreeItemPtr* AttachedActorItem = Items.Find(AttachedActor))
-				{
-					OutChildren.Add(*AttachedActorItem);
-				}
+				OutChildren.Add(*AttachedActorItem);
 			}
-			if (bShowingComponents)
+		}
+		if (bShowingComponents)
+		{
+			for (const UActorComponent* Component : ActorTreeItem->Actor->GetComponents())
 			{
-				for (const UActorComponent* Component : ActorTreeItem->Actor->GetComponents())
+				if (const FSceneOutlinerTreeItemPtr* ComponentItem = Items.Find(Component))
 				{
-					if (const FTreeItemPtr* ComponentItem = Items.Find(Component))
-					{
-						OutChildren.Add(*ComponentItem);
-					}
+					OutChildren.Add(*ComponentItem);
 				}
 			}
 		}
-		else if (const FFolderTreeItem* FolderItem = Item.CastTo<FFolderTreeItem>())
-		{
-			// We should never call FindParents on a folder item if folders are not being shown
-			check(Mode->ShouldShowFolders());
+	}
+	else if (const FFolderTreeItem* FolderItem = Item.CastTo<FFolderTreeItem>())
+	{
+		// We should never call FindParents on a folder item if folders are not being shown
+		check(Mode->ShouldShowFolders());
 
-			// Search through all folders and find any folders with paths which are children of this item
-			for (const auto& Pair : FActorFolders::Get().GetFolderPropertiesForWorld(*RepresentingWorld))
+		// Search through all folders and find any folders with paths which are children of this item
+		for (const auto& Pair : FActorFolders::Get().GetFolderPropertiesForWorld(*RepresentingWorld))
+		{
+			if (SceneOutliner::PathIsChildOf(Pair.Key, FolderItem->Path))
 			{
-				if (PathIsChildOf(Pair.Key, FolderItem->Path))
+				if (const FSceneOutlinerTreeItemPtr* Child = Items.Find(Pair.Key))
 				{
-					if (const FTreeItemPtr* Child = Items.Find(Pair.Key))
-					{
-						OutChildren.Add(*Child);
-					}
+					OutChildren.Add(*Child);
 				}
 			}
 		}
 	}
+}
 
-	FTreeItemPtr FActorHierarchy::FindParent(const ITreeItem& Item, const TMap<FTreeItemID, FTreeItemPtr>& Items) const
+FSceneOutlinerTreeItemPtr FActorHierarchy::FindParent(const ISceneOutlinerTreeItem& Item, const TMap<FSceneOutlinerTreeItemID, FSceneOutlinerTreeItemPtr>& Items) const
+{
+	if (Item.IsA<FWorldTreeItem>())
 	{
-		if (Item.IsA<FWorldTreeItem>())
-		{
-			return nullptr;
-		}
-		else if (const FActorTreeItem* ActorTreeItem = Item.CastTo<FActorTreeItem>())
+		return nullptr;
+	}
+	else if (const FActorTreeItem* ActorTreeItem = Item.CastTo<FActorTreeItem>())
+	{
+		if (ActorTreeItem->Actor.IsValid())
 		{
-			if (ActorTreeItem->Actor.IsValid())
+			const FSceneOutlinerTreeItemPtr* ParentItem;
+			if (const AActor* ParentActor = ActorTreeItem->Actor->GetAttachParentActor())
 			{
-				if (const AActor* ParentActor = ActorTreeItem->Actor->GetAttachParentActor())
+				if (ParentItem = Items.Find(ParentActor))
 				{
-					if (const FTreeItemPtr* ParentItem = Items.Find(ParentActor))
-					{
-						return *ParentItem;
-					}
-				}
-				else if (Mode->ShouldShowFolders() && !ActorTreeItem->Actor->GetFolderPath().IsNone())
-				{
-					if (const FTreeItemPtr* ParentItem = Items.Find(ActorTreeItem->Actor->GetFolderPath()))
-					{
-						return *ParentItem;
-					}
+					return *ParentItem;
 				}
-				else if (const FTreeItemPtr* ParentItem = Items.Find(ActorTreeItem->Actor->GetWorld()))
+			}
+			else if (Mode->ShouldShowFolders() && !ActorTreeItem->Actor->GetFolderPath().IsNone())
+			{
+				if (ParentItem = Items.Find(ActorTreeItem->Actor->GetFolderPath()))
 				{
 					return *ParentItem;
 				}
 			}
+			else if (ParentItem = Items.Find(ActorTreeItem->Actor->GetWorld()))
+			{
+				return *ParentItem;
+			}
 		}
-		else if (const FFolderTreeItem* FolderItem = Item.CastTo<FFolderTreeItem>())
-		{
-			// We should never call FindParents on a folder item if folders are not being shown
-			check(Mode->ShouldShowFolders());
+	}
+	else if (const FFolderTreeItem* FolderItem = Item.CastTo<FFolderTreeItem>())
+	{
+		// We should never call FindParents on a folder item if folders are not being shown
+		check(Mode->ShouldShowFolders());
 
-			const FName ParentPath = GetParentPath(FolderItem->Path);
+		const FName ParentPath = SceneOutliner::GetParentPath(FolderItem->Path);
 
-			const FTreeItemPtr* ParentItem = nullptr;
-			// If the folder has no parent path, it must be parented to the root world
-			if (ParentPath.IsNone())
-			{
-				ParentItem = Items.Find(RepresentingWorld.Get());
-			}
-			else
-			{
-				ParentItem = Items.Find(GetParentPath(FolderItem->Path));
-			}
+		const FSceneOutlinerTreeItemPtr* ParentItem = nullptr;
+		// If the folder has no parent path, it must be parented to the root world
+		if (ParentPath.IsNone())
+		{
+			ParentItem = Items.Find(RepresentingWorld.Get());
+		}
+		else
+		{
+			ParentItem = Items.Find(SceneOutliner::GetParentPath(FolderItem->Path));
+		}
 			
-			if (ParentItem)
-			{
-				return *ParentItem;
-			}
+		if (ParentItem)
+		{
+			return *ParentItem;
 		}
-		else if (const FComponentTreeItem* ComponentTreeItem = Item.CastTo<FComponentTreeItem>())
+	}
+	else if (const FComponentTreeItem* ComponentTreeItem = Item.CastTo<FComponentTreeItem>())
+	{
+		const AActor* Owner = ComponentTreeItem->Component->GetOwner();
+		if (Owner)
 		{
-			const AActor* Owner = ComponentTreeItem->Component->GetOwner();
-			if (Owner)
+			if (const FSceneOutlinerTreeItemPtr* ParentItem = Items.Find(Owner))
 			{
-				if (const FTreeItemPtr* ParentItem = Items.Find(Owner))
-				{
-					return *ParentItem;
-				}
+				return *ParentItem;
 			}
 		}
-		return nullptr;
 	}
+	return nullptr;
+}
 
-	void FActorHierarchy::CreateComponentItems(const AActor* Actor, TArray<FTreeItemPtr>& OutItems) const
+void FActorHierarchy::CreateComponentItems(const AActor* Actor, TArray<FSceneOutlinerTreeItemPtr>& OutItems) const
+{
+	// Add all this actors components if showing components and the owning actor was created
+	if (bShowingComponents)
 	{
-		// Add all this actors components if showing components and the owning actor was created
-		if (bShowingComponents)
+		for (UActorComponent* Component : Actor->GetComponents())
 		{
-			for (UActorComponent* Component : Actor->GetComponents())
+			if (FSceneOutlinerTreeItemPtr ComponentItem = Mode->CreateItemFor<FComponentTreeItem>(Component))
 			{
-				if (FTreeItemPtr ComponentItem = Mode->CreateItemFor<FComponentTreeItem>(Component))
-				{
-					OutItems.Add(ComponentItem);
-				}
+				OutItems.Add(ComponentItem);
 			}
 		}
 	}
+}
 
-	void FActorHierarchy::CreateWorldChildren(UWorld* World, TArray<FTreeItemPtr>& OutItems) const
+void FActorHierarchy::CreateWorldChildren(UWorld* World, TArray<FSceneOutlinerTreeItemPtr>& OutItems) const
+{
+	// Create all actor items
+	for (FActorIterator ActorIt(World); ActorIt; ++ActorIt)
 	{
-		// Create all actor items
-		for (FActorIterator ActorIt(World); ActorIt; ++ActorIt)
+		if (IsActorDisplayableCallback(*ActorIt))
 		{
-			if (IsActorDisplayableCallback(*ActorIt))
+			if (FSceneOutlinerTreeItemPtr ActorItem = Mode->CreateItemFor<FActorTreeItem>(*ActorIt))
 			{
-				if (FTreeItemPtr ActorItem = Mode->CreateItemFor<FActorTreeItem>(*ActorIt))
-				{
-					OutItems.Add(ActorItem);
+				OutItems.Add(ActorItem);
 
-					// Create all component items
-					CreateComponentItems(*ActorIt, OutItems);
-				}
+				// Create all component items
+				CreateComponentItems(*ActorIt, OutItems);
 			}
 		}
+	}
 
-		// Create all folder items
-		if (Mode->ShouldShowFolders())
+	// Create all folder items
+	if (Mode->ShouldShowFolders())
+	{
+		for (const auto& Pair : FActorFolders::Get().GetFolderPropertiesForWorld(*RepresentingWorld))
 		{
-			for (const auto& Pair : FActorFolders::Get().GetFolderPropertiesForWorld(*RepresentingWorld))
+			if (FSceneOutlinerTreeItemPtr FolderItem = Mode->CreateItemFor<FActorFolderTreeItem>(FActorFolderTreeItem(Pair.Key, RepresentingWorld)))
 			{
-				if (FTreeItemPtr FolderItem = Mode->CreateItemFor<FActorFolderTreeItem>(FActorFolderTreeItem(Pair.Key, RepresentingWorld)))
-				{
-					OutItems.Add(FolderItem);
-				}
+				OutItems.Add(FolderItem);
 			}
 		}
 	}
+}
+
+void FActorHierarchy::CreateItems(TArray<FSceneOutlinerTreeItemPtr>& OutItems) const
+{
+	check(RepresentingWorld.IsValid());
 
-	void FActorHierarchy::CreateItems(TArray<FTreeItemPtr>& OutItems) const
+	if (FSceneOutlinerTreeItemPtr WorldItem = Mode->CreateItemFor<FWorldTreeItem>(RepresentingWorld))
 	{
-		check(RepresentingWorld.IsValid());
+		OutItems.Add(WorldItem);
+	}
 
-		if (FTreeItemPtr WorldItem = Mode->CreateItemFor<FWorldTreeItem>(RepresentingWorld))
-		{
-			OutItems.Add(WorldItem);
-		}
+	// Create world children regardless of if a world item was created
+	CreateWorldChildren(RepresentingWorld.Get(), OutItems);
+}
 
-		// Create world children regardless of if a world item was created
-		CreateWorldChildren(RepresentingWorld.Get(), OutItems);
+void FActorHierarchy::CreateChildren(const FSceneOutlinerTreeItemPtr& Item, TArray<FSceneOutlinerTreeItemPtr>& OutChildren) const
+{
+	if (FWorldTreeItem* WorldItem = Item->CastTo<FWorldTreeItem>())
+	{
+		check(WorldItem->World == RepresentingWorld);
+		CreateWorldChildren(WorldItem->World.Get(), OutChildren);
 	}
-
-	void FActorHierarchy::CreateChildren(const FTreeItemPtr& Item, TArray<FTreeItemPtr>& OutChildren) const
+	else if (const FActorTreeItem* ParentActorItem = Item->CastTo<FActorTreeItem>())
 	{
-		if (FWorldTreeItem* WorldItem = Item->CastTo<FWorldTreeItem>())
-		{
-			check(WorldItem->World == RepresentingWorld);
-			CreateWorldChildren(WorldItem->World.Get(), OutChildren);
-		}
-		else if (const FActorTreeItem* ParentActorItem = Item->CastTo<FActorTreeItem>())
+		const AActor* ParentActor = ParentActorItem->Actor.Get();
+
+		if (IsActorDisplayableCallback(ParentActor))
 		{
-			const AActor* ParentActor = ParentActorItem->Actor.Get();
+			CreateComponentItems(ParentActor, OutChildren);
 
-			if (IsActorDisplayableCallback(ParentActor))
+			TArray<AActor*> ChildActors;
+			TFunction<bool(AActor*)> GetAttachedActors = [&ChildActors, &GetAttachedActors](AActor* Child)
 			{
-				CreateComponentItems(ParentActor, OutChildren);
+				ChildActors.Add(Child);
+				Child->ForEachAttachedActors(GetAttachedActors);
 
-				TArray<AActor*> ChildActors;
-				TFunction<bool(AActor*)> GetAttachedActors = [&ChildActors, &GetAttachedActors](AActor* Child)
-				{
-					ChildActors.Add(Child);
-					Child->ForEachAttachedActors(GetAttachedActors);
-
-					// Always continue
-					return true;
-				};
+				// Always continue
+				return true;
+			};
 
-				// Grab all direct/indirect children of an actor
-				ParentActor->ForEachAttachedActors(GetAttachedActors);
+			// Grab all direct/indirect children of an actor
+			ParentActor->ForEachAttachedActors(GetAttachedActors);
 
-				for (auto ChildActor : ChildActors)
+			for (auto ChildActor : ChildActors)
+			{
+				if (FSceneOutlinerTreeItemPtr ChildActorItem = Mode->CreateItemFor<FActorTreeItem>(ChildActor))
 				{
-					if (FTreeItemPtr ChildActorItem = Mode->CreateItemFor<FActorTreeItem>(ChildActor))
-					{
-						OutChildren.Add(ChildActorItem);
+					OutChildren.Add(ChildActorItem);
 
-						CreateComponentItems(ChildActor, OutChildren);
-					}
+					CreateComponentItems(ChildActor, OutChildren);
 				}
 			}
 		}
-		else if (FActorFolderTreeItem* FolderItem = Item->CastTo<FActorFolderTreeItem>())
-		{
-			check(Mode->ShouldShowFolders());
+	}
+	else if (FActorFolderTreeItem* FolderItem = Item->CastTo<FActorFolderTreeItem>())
+	{
+		check(Mode->ShouldShowFolders());
 
-			for (const auto& Pair : FActorFolders::Get().GetFolderPropertiesForWorld(*FolderItem->World))
+		for (const auto& Pair : FActorFolders::Get().GetFolderPropertiesForWorld(*FolderItem->World))
+		{
+			if (SceneOutliner::PathIsChildOf(Pair.Key, FolderItem->Path))
 			{
-				if (PathIsChildOf(Pair.Key, FolderItem->Path))
+				if (FSceneOutlinerTreeItemPtr NewFolderItem = Mode->CreateItemFor<FActorFolderTreeItem>(FActorFolderTreeItem(Pair.Key, FolderItem->World)))
 				{
-					if (FTreeItemPtr NewFolderItem = Mode->CreateItemFor<FActorFolderTreeItem>(FActorFolderTreeItem(Pair.Key, FolderItem->World)))
-					{
-						OutChildren.Add(NewFolderItem);
-					}
+					OutChildren.Add(NewFolderItem);
 				}
 			}
 		}
 	}
+}
 
-	FTreeItemPtr FActorHierarchy::CreateParentItem(const FTreeItemPtr& Item) const
+FSceneOutlinerTreeItemPtr FActorHierarchy::CreateParentItem(const FSceneOutlinerTreeItemPtr& Item) const
+{
+	if (Item->IsA<FWorldTreeItem>())
 	{
-		if (Item->IsA<FWorldTreeItem>())
-		{
-			return nullptr;
-		}
-		else if (const FActorTreeItem* ActorTreeItem = Item->CastTo<FActorTreeItem>())
-		{
-			if (ActorTreeItem->Actor.IsValid())
-			{
-				AActor* ParentActor = ActorTreeItem->Actor->GetAttachParentActor();
-				if (ParentActor)
-				{
-					return Mode->CreateItemFor<FActorTreeItem>(ParentActor, true);
-				}
-				// if this item does not belong to a folder (or folders are disabled)
-				else if (!Mode->ShouldShowFolders() || ActorTreeItem->Actor->GetFolderPath().IsNone())
-				{
-					UWorld* OwningWorld = ActorTreeItem->Actor->GetWorld();
-					check(OwningWorld);
-					return Mode->CreateItemFor<FWorldTreeItem>(OwningWorld, true);
-				}
-				else
-				{
-					return Mode->CreateItemFor<FActorFolderTreeItem>(FActorFolderTreeItem(ActorTreeItem->Actor->GetFolderPath(), ActorTreeItem->Actor->GetWorld()), true);
-				}
-			}
-		}
-		else if (const FComponentTreeItem* ComponentTreeItem = Item->CastTo<FComponentTreeItem>())
+		return nullptr;
+	}
+	else if (const FActorTreeItem* ActorTreeItem = Item->CastTo<FActorTreeItem>())
+	{
+		if (ActorTreeItem->Actor.IsValid())
 		{
-			if (AActor* ParentActor = ComponentTreeItem->Component->GetOwner())
+			AActor* ParentActor = ActorTreeItem->Actor->GetAttachParentActor();
+			if (ParentActor)
 			{
 				return Mode->CreateItemFor<FActorTreeItem>(ParentActor, true);
 			}
-		}
-		else if (const FActorFolderTreeItem* FolderTreeItem = Item->CastTo<FActorFolderTreeItem>())
-		{
-			check(Mode->ShouldShowFolders());
-			
-			const FName ParentPath = GetParentPath(FolderTreeItem->Path);
-			if (ParentPath.IsNone())
+			// if this item does not belong to a folder (or folders are disabled)
+			else if (!Mode->ShouldShowFolders() || ActorTreeItem->Actor->GetFolderPath().IsNone())
 			{
-				UWorld* OwningWorld = FolderTreeItem->World.Get();
+				UWorld* OwningWorld = ActorTreeItem->Actor->GetWorld();
 				check(OwningWorld);
 				return Mode->CreateItemFor<FWorldTreeItem>(OwningWorld, true);
 			}
 			else
 			{
-				return Mode->CreateItemFor<FActorFolderTreeItem>(FActorFolderTreeItem(ParentPath, FolderTreeItem->World), true);
+				return Mode->CreateItemFor<FActorFolderTreeItem>(FActorFolderTreeItem(ActorTreeItem->Actor->GetFolderPath(), ActorTreeItem->Actor->GetWorld()), true);
 			}
 		}
-		return nullptr;
-	}
-
-	void FActorHierarchy::FullRefreshEvent()
-	{
-		FHierarchyChangedData EventData;
-		EventData.Type = FHierarchyChangedData::FullRefresh;
-		
-		HierarchyChangedEvent.Broadcast(EventData);
 	}
-
-	void FActorHierarchy::OnLevelActorAdded(AActor* InActor)
+	else if (const FComponentTreeItem* ComponentTreeItem = Item->CastTo<FComponentTreeItem>())
 	{
-		if (InActor && RepresentingWorld.Get() == InActor->GetWorld() && IsActorDisplayableCallback(InActor))
+		if (AActor* ParentActor = ComponentTreeItem->Component->GetOwner())
 		{
-			FHierarchyChangedData EventData;
-			EventData.Type = FHierarchyChangedData::Added;
-			EventData.Item = Mode->CreateItemFor<FActorTreeItem>(InActor);
-			HierarchyChangedEvent.Broadcast(EventData);
+			return Mode->CreateItemFor<FActorTreeItem>(ParentActor, true);
 		}
 	}
-
-	void FActorHierarchy::OnLevelActorDeleted(AActor* InActor)
+	else if (const FActorFolderTreeItem* FolderTreeItem = Item->CastTo<FActorFolderTreeItem>())
 	{
-		if (RepresentingWorld.Get() == InActor->GetWorld())
+		check(Mode->ShouldShowFolders());
+			
+		const FName ParentPath = SceneOutliner::GetParentPath(FolderTreeItem->Path);
+		if (ParentPath.IsNone())
+		{
+			UWorld* OwningWorld = FolderTreeItem->World.Get();
+			check(OwningWorld);
+			return Mode->CreateItemFor<FWorldTreeItem>(OwningWorld, true);
+		}
+		else
 		{
-			FHierarchyChangedData EventData;
-			EventData.Type = FHierarchyChangedData::Removed;
-			EventData.ItemID = InActor;
-			HierarchyChangedEvent.Broadcast(EventData);
+			return Mode->CreateItemFor<FActorFolderTreeItem>(FActorFolderTreeItem(ParentPath, FolderTreeItem->World), true);
 		}
 	}
+	return nullptr;
+}
+
+void FActorHierarchy::FullRefreshEvent()
+{
+	FSceneOutlinerHierarchyChangedData EventData;
+	EventData.Type = FSceneOutlinerHierarchyChangedData::FullRefresh;
+		
+	HierarchyChangedEvent.Broadcast(EventData);
+}
 
-	void FActorHierarchy::OnLevelActorAttached(AActor* InActor, const AActor* InParent)
+void FActorHierarchy::OnLevelActorAdded(AActor* InActor)
+{
+	if (InActor && RepresentingWorld.Get() == InActor->GetWorld() && IsActorDisplayableCallback(InActor))
 	{
-		if (RepresentingWorld.Get() == InActor->GetWorld())
-		{
-			FHierarchyChangedData EventData;
-			EventData.Type = FHierarchyChangedData::Moved;
-			EventData.ItemID = InActor;
-			HierarchyChangedEvent.Broadcast(EventData);
-		}
+		FSceneOutlinerHierarchyChangedData EventData;
+		EventData.Type = FSceneOutlinerHierarchyChangedData::Added;
+		EventData.Item = Mode->CreateItemFor<FActorTreeItem>(InActor);
+		HierarchyChangedEvent.Broadcast(EventData);
 	}
+}
 
-	void FActorHierarchy::OnLevelActorDetached(AActor* InActor, const AActor* InParent)
+void FActorHierarchy::OnLevelActorDeleted(AActor* InActor)
+{
+	if (RepresentingWorld.Get() == InActor->GetWorld())
 	{
-		if (RepresentingWorld.Get() == InActor->GetWorld())
-		{
-			FHierarchyChangedData EventData;
-			EventData.Type = FHierarchyChangedData::Moved;
-			EventData.ItemID = InActor;
-			HierarchyChangedEvent.Broadcast(EventData);
-		}
+		FSceneOutlinerHierarchyChangedData EventData;
+		EventData.Type = FSceneOutlinerHierarchyChangedData::Removed;
+		EventData.ItemID = InActor;
+		HierarchyChangedEvent.Broadcast(EventData);
 	}
+}
 
-	void FActorHierarchy::OnComponentsUpdated()
+void FActorHierarchy::OnLevelActorAttached(AActor* InActor, const AActor* InParent)
+{
+	if (RepresentingWorld.Get() == InActor->GetWorld())
 	{
-		FullRefreshEvent();
+		FSceneOutlinerHierarchyChangedData EventData;
+		EventData.Type = FSceneOutlinerHierarchyChangedData::Moved;
+		EventData.ItemID = InActor;
+		HierarchyChangedEvent.Broadcast(EventData);
 	}
+}
 
-	void FActorHierarchy::OnLevelActorListChanged()
+void FActorHierarchy::OnLevelActorDetached(AActor* InActor, const AActor* InParent)
+{
+	if (RepresentingWorld.Get() == InActor->GetWorld())
 	{
-		FullRefreshEvent();
+		FSceneOutlinerHierarchyChangedData EventData;
+		EventData.Type = FSceneOutlinerHierarchyChangedData::Moved;
+		EventData.ItemID = InActor;
+		HierarchyChangedEvent.Broadcast(EventData);
 	}
+}
+
+void FActorHierarchy::OnComponentsUpdated()
+{
+	FullRefreshEvent();
+}
 
-	void FActorHierarchy::OnLevelAdded(ULevel* InLevel, UWorld* InWorld)
+void FActorHierarchy::OnLevelActorListChanged()
+{
+	FullRefreshEvent();
+}
+
+void FActorHierarchy::OnLevelAdded(ULevel* InLevel, UWorld* InWorld)
+{
+	if (RepresentingWorld.Get() == InWorld)
 	{
-		if (RepresentingWorld.Get() == InWorld)
-		{
-			FullRefreshEvent();
-		}
+		FullRefreshEvent();
 	}
+}
 
-	void FActorHierarchy::OnLevelRemoved(ULevel* InLevel, UWorld* InWorld)
+void FActorHierarchy::OnLevelRemoved(ULevel* InLevel, UWorld* InWorld)
+{
+	if (RepresentingWorld.Get() == InWorld)
 	{
-		if (RepresentingWorld.Get() == InWorld)
-		{
-			FullRefreshEvent();
-		}
+		FullRefreshEvent();
 	}
+}
 
-	/** Called when a folder is to be created */
-	void FActorHierarchy::OnBroadcastFolderCreate(UWorld& InWorld, FName NewPath)
+/** Called when a folder is to be created */
+void FActorHierarchy::OnBroadcastFolderCreate(UWorld& InWorld, FName NewPath)
+{
+	if (Mode->ShouldShowFolders() && RepresentingWorld.Get() == &InWorld)
 	{
-		if (Mode->ShouldShowFolders() && RepresentingWorld.Get() == &InWorld)
-		{
-			FHierarchyChangedData EventData;
-			EventData.Type = FHierarchyChangedData::Added;
-			EventData.Item = Mode->CreateItemFor<FActorFolderTreeItem>(FActorFolderTreeItem(NewPath, &InWorld));
-			EventData.ItemActions = ENewItemAction::Select | ENewItemAction::Rename;
-			HierarchyChangedEvent.Broadcast(EventData);
-		}
+		FSceneOutlinerHierarchyChangedData EventData;
+		EventData.Type = FSceneOutlinerHierarchyChangedData::Added;
+		EventData.Item = Mode->CreateItemFor<FActorFolderTreeItem>(FActorFolderTreeItem(NewPath, &InWorld));
+		EventData.ItemActions = SceneOutliner::ENewItemAction::Select | SceneOutliner::ENewItemAction::Rename;
+		HierarchyChangedEvent.Broadcast(EventData);
 	}
+}
 
-	/** Called when a folder is to be moved */
-	void FActorHierarchy::OnBroadcastFolderMove(UWorld& InWorld, FName OldPath, FName NewPath)
+/** Called when a folder is to be moved */
+void FActorHierarchy::OnBroadcastFolderMove(UWorld& InWorld, FName OldPath, FName NewPath)
+{
+	if (Mode->ShouldShowFolders() && RepresentingWorld.Get() == &InWorld)
 	{
-		if (Mode->ShouldShowFolders() && RepresentingWorld.Get() == &InWorld)
-		{
-			FHierarchyChangedData EventData;
-			EventData.Type = FHierarchyChangedData::FolderMoved;
-			EventData.ItemID = OldPath;
-			EventData.NewPath = NewPath;
-			HierarchyChangedEvent.Broadcast(EventData);
-		}
+		FSceneOutlinerHierarchyChangedData EventData;
+		EventData.Type = FSceneOutlinerHierarchyChangedData::FolderMoved;
+		EventData.ItemID = OldPath;
+		EventData.NewPath = NewPath;
+		HierarchyChangedEvent.Broadcast(EventData);
 	}
+}
 
-	/** Called when a folder is to be deleted */
-	void FActorHierarchy::OnBroadcastFolderDelete(UWorld& InWorld, FName Path)
+/** Called when a folder is to be deleted */
+void FActorHierarchy::OnBroadcastFolderDelete(UWorld& InWorld, FName Path)
+{
+	if (Mode->ShouldShowFolders() && RepresentingWorld.Get() == &InWorld)
 	{
-		if (Mode->ShouldShowFolders() && RepresentingWorld.Get() == &InWorld)
-		{
-			FHierarchyChangedData EventData;
-			EventData.Type = FHierarchyChangedData::Removed;
-			EventData.ItemID = Path;
-			HierarchyChangedEvent.Broadcast(EventData);
-		}
+		FSceneOutlinerHierarchyChangedData EventData;
+		EventData.Type = FSceneOutlinerHierarchyChangedData::Removed;
+		EventData.ItemID = Path;
+		HierarchyChangedEvent.Broadcast(EventData);
 	}
+}
 
-	void FActorHierarchy::OnLevelActorFolderChanged(const AActor* InActor, FName OldPath)
+void FActorHierarchy::OnLevelActorFolderChanged(const AActor* InActor, FName OldPath)
+{
+	if (Mode->ShouldShowFolders() && RepresentingWorld.Get() == InActor->GetWorld())
 	{
-		if (Mode->ShouldShowFolders() && RepresentingWorld.Get() == InActor->GetWorld())
-		{
-			FHierarchyChangedData EventData;
-			EventData.Type = FHierarchyChangedData::Moved;
-			EventData.ItemID = FTreeItemID(InActor);
-			HierarchyChangedEvent.Broadcast(EventData);
-		}
+		FSceneOutlinerHierarchyChangedData EventData;
+		EventData.Type = FSceneOutlinerHierarchyChangedData::Moved;
+		EventData.ItemID = FSceneOutlinerTreeItemID(InActor);
+		HierarchyChangedEvent.Broadcast(EventData);
 	}
 }
diff --git a/Engine/Source/Editor/SceneOutliner/Private/ActorMode.cpp b/Engine/Source/Editor/SceneOutliner/Private/ActorMode.cpp
--- a/Engine/Source/Editor/SceneOutliner/Private/ActorMode.cpp
+++ b/Engine/Source/Editor/SceneOutliner/Private/ActorMode.cpp
@@ -17,11 +17,11 @@
 
 #define LOCTEXT_NAMESPACE "SceneOutliner_ActorMode"
 
+using FActorFilter = TSceneOutlinerPredicateFilter<FActorTreeItem>;
+
 namespace SceneOutliner
 {
-	using FActorFilter = TOutlinerPredicateFilter<FActorTreeItem>;
-
-	bool FWeakActorSelector::operator()(const TWeakPtr<ITreeItem>& Item, TWeakObjectPtr<AActor>& DataOut) const
+	bool FWeakActorSelector::operator()(const TWeakPtr<ISceneOutlinerTreeItem>& Item, TWeakObjectPtr<AActor>& DataOut) const
 	{
 		if (FActorTreeItem* ActorItem = Item.Pin()->CastTo<FActorTreeItem>())
 		{
@@ -34,7 +34,7 @@ namespace SceneOutliner
 		return false;
 	}
 
-	bool FActorSelector::operator()(const TWeakPtr<ITreeItem>& Item, AActor*& ActorPtrOut) const
+	bool FActorSelector::operator()(const TWeakPtr<ISceneOutlinerTreeItem>& Item, AActor*& ActorPtrOut) const
 	{
 		if (FActorTreeItem* ActorItem = Item.Pin()->CastTo<FActorTreeItem>())
 		{
@@ -64,267 +64,268 @@ namespace SceneOutliner
 
 		return false;
 	}
+}
 
-	FActorMode::FActorMode(SSceneOutliner* InSceneOutliner, bool bInHideComponents, TWeakObjectPtr<UWorld> InSpecifiedWorldToDisplay)
-		: ISceneOutlinerMode(InSceneOutliner)
-		, bHideComponents(bInHideComponents)
-		, SpecifiedWorldToDisplay(InSpecifiedWorldToDisplay)
-	{
-	}
+FActorMode::FActorMode(SSceneOutliner* InSceneOutliner, bool bInHideComponents, TWeakObjectPtr<UWorld> InSpecifiedWorldToDisplay)
+	: ISceneOutlinerMode(InSceneOutliner)
+	, bHideComponents(bInHideComponents)
+	, SpecifiedWorldToDisplay(InSpecifiedWorldToDisplay)
+{
+}
 
-	FActorMode::~FActorMode()
+FActorMode::~FActorMode()
+{
+	if (Hierarchy)
 	{
-		if (Hierarchy)
-		{
-			Hierarchy->OnHierarchyChanged().RemoveAll(this);
-		}
+		Hierarchy->OnHierarchyChanged().RemoveAll(this);
 	}
+}
 
-	TUniquePtr<ISceneOutlinerHierarchy> FActorMode::CreateHierarchy()
-	{
-		TUniquePtr<FActorHierarchy> ActorHierarchy = FActorHierarchy::Create(this, RepresentingWorld, [this](const AActor* InActor) {return IsActorDisplayable(InActor); });
-		ActorHierarchy->SetShowingComponents(!bHideComponents);
-
-		return ActorHierarchy;
-	}
+TUniquePtr<ISceneOutlinerHierarchy> FActorMode::CreateHierarchy()
+{
+	TUniquePtr<FActorHierarchy> ActorHierarchy = FActorHierarchy::Create(this, RepresentingWorld, [this](const AActor* InActor) {return IsActorDisplayable(InActor); });
+	ActorHierarchy->SetShowingComponents(!bHideComponents);
 
-	void FActorMode::Rebuild()
-	{
-		ChooseRepresentingWorld();
+	return ActorHierarchy;
+}
 
-		if (Hierarchy)
-		{
-			Hierarchy->OnHierarchyChanged().Clear();
-		}
+void FActorMode::Rebuild()
+{
+	ChooseRepresentingWorld();
 
-		Hierarchy = CreateHierarchy();
+	if (Hierarchy)
+	{
+		Hierarchy->OnHierarchyChanged().Clear();
 	}
 
-	void FActorMode::ChooseRepresentingWorld()
-	{
-		// Select a world to represent
+	Hierarchy = CreateHierarchy();
+}
 
-		RepresentingWorld = nullptr;
+void FActorMode::ChooseRepresentingWorld()
+{
+	// Select a world to represent
 
-		// If a specified world was provided, represent it
-		if (SpecifiedWorldToDisplay.IsValid())
-		{
-			RepresentingWorld = SpecifiedWorldToDisplay.Get();
-		}
+	RepresentingWorld = nullptr;
 
-		// check if the user-chosen world is valid and in the editor contexts
+	// If a specified world was provided, represent it
+	if (SpecifiedWorldToDisplay.IsValid())
+	{
+		RepresentingWorld = SpecifiedWorldToDisplay.Get();
+	}
 
-		if (!RepresentingWorld.IsValid() && UserChosenWorld.IsValid())
-		{
-			for (const FWorldContext& Context : GEngine->GetWorldContexts())
-			{
-				if (UserChosenWorld.Get() == Context.World())
-				{
-					RepresentingWorld = UserChosenWorld.Get();
-					break;
-				}
-			}
-		}
+	// check if the user-chosen world is valid and in the editor contexts
 
-		// If the user did not manually select a world, try to pick the most suitable world context
-		if (!RepresentingWorld.IsValid())
+	if (!RepresentingWorld.IsValid() && UserChosenWorld.IsValid())
+	{
+		for (const FWorldContext& Context : GEngine->GetWorldContexts())
 		{
-			// ideally we want a PIE world that is standalone or the first client
-			for (const FWorldContext& Context : GEngine->GetWorldContexts())
+			if (UserChosenWorld.Get() == Context.World())
 			{
-				UWorld* World = Context.World();
-				if (World && Context.WorldType == EWorldType::PIE)
-				{
-					if (World->GetNetMode() == NM_Standalone)
-					{
-						RepresentingWorld = World;
-						break;
-					}
-					else if (World->GetNetMode() == NM_Client && Context.PIEInstance == 2)	// Slightly dangerous: assumes server is always PIEInstance = 1;
-					{
-						RepresentingWorld = World;
-						break;
-					}
-				}
+				RepresentingWorld = UserChosenWorld.Get();
+				break;
 			}
 		}
+	}
 
-		if (RepresentingWorld == nullptr)
+	// If the user did not manually select a world, try to pick the most suitable world context
+	if (!RepresentingWorld.IsValid())
+	{
+		// ideally we want a PIE world that is standalone or the first client
+		for (const FWorldContext& Context : GEngine->GetWorldContexts())
 		{
-			// still not world so fallback to old logic where we just prefer PIE over Editor
-
-			for (const FWorldContext& Context : GEngine->GetWorldContexts())
+			UWorld* World = Context.World();
+			if (World && Context.WorldType == EWorldType::PIE)
 			{
-				if (Context.WorldType == EWorldType::PIE)
+				if (World->GetNetMode() == NM_Standalone)
 				{
-					RepresentingWorld = Context.World();
+					RepresentingWorld = World;
 					break;
 				}
-				else if (Context.WorldType == EWorldType::Editor)
+				else if (World->GetNetMode() == NM_Client && Context.PIEInstance == 2)	// Slightly dangerous: assumes server is always PIEInstance = 1;
 				{
-					RepresentingWorld = Context.World();
+					RepresentingWorld = World;
+					break;
 				}
 			}
 		}
 	}
 
-	void FActorMode::BuildWorldPickerMenu(FMenuBuilder& MenuBuilder)
+	if (RepresentingWorld == nullptr)
 	{
-		MenuBuilder.BeginSection("Worlds", LOCTEXT("WorldsHeading", "Worlds"));
+		// still not world so fallback to old logic where we just prefer PIE over Editor
+
+		for (const FWorldContext& Context : GEngine->GetWorldContexts())
 		{
-			MenuBuilder.AddMenuEntry(
-				LOCTEXT("AutoWorld", "Auto"),
-				LOCTEXT("AutoWorldToolTip", "Automatically pick the world to display based on context."),
-				FSlateIcon(),
-				FUIAction(
-					FExecuteAction::CreateRaw(this, &FActorMode::OnSelectWorld, TWeakObjectPtr<UWorld>()),
-					FCanExecuteAction(),
-					FIsActionChecked::CreateRaw(this, &FActorMode::IsWorldChecked, TWeakObjectPtr<UWorld>())
-				),
-				NAME_None,
-				EUserInterfaceActionType::RadioButton
-			);
-
-			for (const FWorldContext& Context : GEngine->GetWorldContexts())
+			if (Context.WorldType == EWorldType::PIE)
 			{
-				UWorld* World = Context.World();
-				if (World && (World->WorldType == EWorldType::PIE || Context.WorldType == EWorldType::Editor))
-				{
-					MenuBuilder.AddMenuEntry(
-						GetWorldDescription(World),
-						LOCTEXT("ChooseWorldToolTip", "Display actors for this world."),
-						FSlateIcon(),
-						FUIAction(
-							FExecuteAction::CreateRaw(this, &FActorMode::OnSelectWorld, MakeWeakObjectPtr(World)),
-							FCanExecuteAction(),
-							FIsActionChecked::CreateRaw(this, &FActorMode::IsWorldChecked, MakeWeakObjectPtr(World))
-						),
-						NAME_None,
-						EUserInterfaceActionType::RadioButton
-					);
-				}
+				RepresentingWorld = Context.World();
+				break;
+			}
+			else if (Context.WorldType == EWorldType::Editor)
+			{
+				RepresentingWorld = Context.World();
 			}
 		}
-		MenuBuilder.EndSection();
 	}
+}
 
-	void FActorMode::OnSelectWorld(TWeakObjectPtr<UWorld> World)
+void FActorMode::BuildWorldPickerMenu(FMenuBuilder& MenuBuilder)
+{
+	MenuBuilder.BeginSection("Worlds", LOCTEXT("WorldsHeading", "Worlds"));
 	{
-		UserChosenWorld = World;
-		SceneOutliner->FullRefresh();
+		MenuBuilder.AddMenuEntry(
+			LOCTEXT("AutoWorld", "Auto"),
+			LOCTEXT("AutoWorldToolTip", "Automatically pick the world to display based on context."),
+			FSlateIcon(),
+			FUIAction(
+				FExecuteAction::CreateRaw(this, &FActorMode::OnSelectWorld, TWeakObjectPtr<UWorld>()),
+				FCanExecuteAction(),
+				FIsActionChecked::CreateRaw(this, &FActorMode::IsWorldChecked, TWeakObjectPtr<UWorld>())
+			),
+			NAME_None,
+			EUserInterfaceActionType::RadioButton
+		);
+
+		for (const FWorldContext& Context : GEngine->GetWorldContexts())
+		{
+			UWorld* World = Context.World();
+			if (World && (World->WorldType == EWorldType::PIE || Context.WorldType == EWorldType::Editor))
+			{
+				MenuBuilder.AddMenuEntry(
+					SceneOutliner::GetWorldDescription(World),
+					LOCTEXT("ChooseWorldToolTip", "Display actors for this world."),
+					FSlateIcon(),
+					FUIAction(
+						FExecuteAction::CreateRaw(this, &FActorMode::OnSelectWorld, MakeWeakObjectPtr(World)),
+						FCanExecuteAction(),
+						FIsActionChecked::CreateRaw(this, &FActorMode::IsWorldChecked, MakeWeakObjectPtr(World))
+					),
+					NAME_None,
+					EUserInterfaceActionType::RadioButton
+				);
+			}
+		}
 	}
+	MenuBuilder.EndSection();
+}
 
-	bool FActorMode::IsWorldChecked(TWeakObjectPtr<UWorld> World) const
-	{
-		return (UserChosenWorld == World) || (World.IsExplicitlyNull() && !UserChosenWorld.IsValid());
-	}
+void FActorMode::OnSelectWorld(TWeakObjectPtr<UWorld> World)
+{
+	UserChosenWorld = World;
+	SceneOutliner->FullRefresh();
+}
 
-	void FActorMode::SynchronizeActorSelection()
-	{
-		USelection* SelectedActors = GEditor->GetSelectedActors();
+bool FActorMode::IsWorldChecked(TWeakObjectPtr<UWorld> World) const
+{
+	return (UserChosenWorld == World) || (World.IsExplicitlyNull() && !UserChosenWorld.IsValid());
+}
 
-		// Deselect actors in the tree that are no longer selected in the world
-		const FItemSelection Selection(SceneOutliner->GetSelection());
-		auto DeselectActors = [this](FActorTreeItem& Item)
-		{
-			if (!Item.Actor.IsValid() || !Item.Actor.Get()->IsSelected())
-			{
-				SceneOutliner->RemoveFromSelection(Item.AsShared());
-			}
-		};
-		Selection.ForEachItem<FActorTreeItem>(DeselectActors);
+void FActorMode::SynchronizeActorSelection()
+{
+	USelection* SelectedActors = GEditor->GetSelectedActors();
 
-		// Show actor selection but only if sub objects are not selected
-		if (!Selection.Has<FComponentTreeItem>())
+	// Deselect actors in the tree that are no longer selected in the world
+	const FSceneOutlinerItemSelection Selection(SceneOutliner->GetSelection());
+	auto DeselectActors = [this](FActorTreeItem& Item)
+	{
+		if (!Item.Actor.IsValid() || !Item.Actor.Get()->IsSelected())
 		{
-			// See if the tree view selector is pointing at a selected item
-			bool bSelectorInSelectionSet = false;
+			SceneOutliner->SetItemSelection(Item.AsShared(), false);
+		}
+	};
+	Selection.ForEachItem<FActorTreeItem>(DeselectActors);
+
+	// Show actor selection but only if sub objects are not selected
+	if (!Selection.Has<FComponentTreeItem>())
+	{
+		// See if the tree view selector is pointing at a selected item
+		bool bSelectorInSelectionSet = false;
 
-			TArray<FTreeItemPtr> ActorItems;
-			for (FSelectionIterator SelectionIt(*SelectedActors); SelectionIt; ++SelectionIt)
+		TArray<FSceneOutlinerTreeItemPtr> ActorItems;
+		for (FSelectionIterator SelectionIt(*SelectedActors); SelectionIt; ++SelectionIt)
+		{
+			AActor* Actor = CastChecked< AActor >(*SelectionIt);
+			if (FSceneOutlinerTreeItemPtr ActorItem = SceneOutliner->GetTreeItem(Actor))
 			{
-				AActor* Actor = CastChecked< AActor >(*SelectionIt);
-				if (FTreeItemPtr ActorItem = SceneOutliner->GetTreeItem(Actor))
+				if (!bSelectorInSelectionSet && SceneOutliner->HasSelectorFocus(ActorItem))
 				{
-					if (!bSelectorInSelectionSet && SceneOutliner->HasSelectorFocus(ActorItem))
-					{
-						bSelectorInSelectionSet = true;
-					}
-
-					ActorItems.Add(ActorItem);
+					bSelectorInSelectionSet = true;
 				}
-			}
 
-			// If NOT bSelectorInSelectionSet then we want to just move the selector to the first selected item.
-			ESelectInfo::Type SelectInfo = bSelectorInSelectionSet ? ESelectInfo::Direct : ESelectInfo::OnMouseClick;
-			SceneOutliner->AddToSelection(ActorItems, SelectInfo);
+				ActorItems.Add(ActorItem);
+			}
 		}
-		FSceneOutlinerDelegates::Get().SelectionChanged.Broadcast();
-	}
 
-	bool FActorMode::IsActorDisplayable(const AActor* Actor) const
-	{
-		static const FName SequencerActorTag(TEXT("SequencerActor"));
-
-		return Actor && 
-			!SceneOutliner->GetSharedData().bOnlyShowFolders && 												// Don't show actors if we're only showing folders
-			Actor->IsEditable() &&															// Only show actors that are allowed to be selected and drawn in editor
-			Actor->IsListedInSceneOutliner() &&
-			((Actor->GetWorld()->IsPlayInEditor() || !Actor->HasAnyFlags(RF_Transient)) ||
-				(SceneOutliner->GetSharedData().bShowTransient && Actor->HasAnyFlags(RF_Transient)) ||			// Don't show transient actors in non-play worlds
-				(Actor->ActorHasTag(SequencerActorTag))) &&
-			!Actor->IsTemplate() &&															// Should never happen, but we never want CDOs displayed
-			!FActorEditorUtils::IsABuilderBrush(Actor) &&									// Don't show the builder brush
-			!Actor->IsA(AWorldSettings::StaticClass()) &&									// Don't show the WorldSettings actor, even though it is technically editable
-			!Actor->IsPendingKill() &&														// We don't want to show actors that are about to go away
-			FLevelUtils::IsLevelVisible(Actor->GetLevel());								// Only show Actors whose level is visible
+		// If NOT bSelectorInSelectionSet then we want to just move the selector to the first selected item.
+		ESelectInfo::Type SelectInfo = bSelectorInSelectionSet ? ESelectInfo::Direct : ESelectInfo::OnMouseClick;
+		SceneOutliner->AddToSelection(ActorItems, SelectInfo);
 	}
+	FSceneOutlinerDelegates::Get().SelectionChanged.Broadcast();
+}
+
+bool FActorMode::IsActorDisplayable(const AActor* Actor) const
+{
+	static const FName SequencerActorTag(TEXT("SequencerActor"));
+
+	return Actor && 
+		!SceneOutliner->GetSharedData().bOnlyShowFolders && 												// Don't show actors if we're only showing folders
+		Actor->IsEditable() &&															// Only show actors that are allowed to be selected and drawn in editor
+		Actor->IsListedInSceneOutliner() &&
+		((Actor->GetWorld()->IsPlayInEditor() || !Actor->HasAnyFlags(RF_Transient)) ||
+			(SceneOutliner->GetSharedData().bShowTransient && Actor->HasAnyFlags(RF_Transient)) ||			// Don't show transient actors in non-play worlds
+			(Actor->ActorHasTag(SequencerActorTag))) &&
+		!Actor->IsTemplate() &&															// Should never happen, but we never want CDOs displayed
+		!FActorEditorUtils::IsABuilderBrush(Actor) &&									// Don't show the builder brush
+		!Actor->IsA(AWorldSettings::StaticClass()) &&									// Don't show the WorldSettings actor, even though it is technically editable
+		!Actor->IsPendingKill() &&														// We don't want to show actors that are about to go away
+		FLevelUtils::IsLevelVisible(Actor->GetLevel());								// Only show Actors whose level is visible
+}
 
-	void FActorMode::OnFilterTextChanged(const FText& InFilterText)
+void FActorMode::OnFilterTextChanged(const FText& InFilterText)
+{
+	// Scroll last item (if it passes the filter) into view - this means if we are multi-selecting, we show newest selection that passes the filter
+	if (const AActor* LastSelectedActor = GEditor->GetSelectedActors()->GetBottom<AActor>())
 	{
-		// Scroll last item (if it passes the filter) into view - this means if we are multi-selecting, we show newest selection that passes the filter
-		if (const AActor* LastSelectedActor = GEditor->GetSelectedActors()->GetBottom<AActor>())
-		{
-			// This part is different than that of OnLevelSelectionChanged(nullptr) because IsItemVisible(TreeItem) & ScrollItemIntoView(TreeItem) are applied to
-			// the current visual state, not to the one after applying the filter. Thus, the scroll would go to the place where the object was located
-			// before applying the FilterText
+		// This part is different than that of OnLevelSelectionChanged(nullptr) because IsItemVisible(TreeItem) & ScrollItemIntoView(TreeItem) are applied to
+		// the current visual state, not to the one after applying the filter. Thus, the scroll would go to the place where the object was located
+		// before applying the FilterText
 
-			// If the object is already in the list, but it does not passes the filter, then we do not want to re-add it, because it will be removed by the filter
-			const FTreeItemPtr TreeItem = SceneOutliner->GetTreeItem(LastSelectedActor);
-			if (!TreeItem.IsValid() || !SceneOutliner->PassesTextFilter(TreeItem))
-			{
-				return;
-			}
+		// If the object is already in the list, but it does not passes the filter, then we do not want to re-add it, because it will be removed by the filter
+		const FSceneOutlinerTreeItemPtr TreeItem = SceneOutliner->GetTreeItem(LastSelectedActor);
+		if (!TreeItem.IsValid() || !SceneOutliner->PassesTextFilter(TreeItem))
+		{
+			return;
+		}
 
-			// If the object is not in the list, and it does not passes the filter, then we should not re-add it, because it would be removed by the filter again. Unfortunately,
-			// there is no code to check if a future element (i.e., one that is currently not in the TreeItemMap list) will pass the filter. Therefore, we kind of overkill it
-			// by re-adding that element (even though it will be removed). However, AddItemToTree(FTreeItemRef Item) and similar functions already check the element before
-			// adding it. So this solution is fine.
-			// This solution might affect the performance of the World Outliner when a key is pressed, but it will still work properly when the remove/del keys are pressed. Not
-			// updating the filter when !TreeItem.IsValid() would result in the focus not being updated when the remove/del keys are pressed.
+		// If the object is not in the list, and it does not passes the filter, then we should not re-add it, because it would be removed by the filter again. Unfortunately,
+		// there is no code to check if a future element (i.e., one that is currently not in the TreeItemMap list) will pass the filter. Therefore, we kind of overkill it
+		// by re-adding that element (even though it will be removed). However, AddItemToTree(FSceneOutlinerTreeItemRef Item) and similar functions already check the element before
+		// adding it. So this solution is fine.
+		// This solution might affect the performance of the World Outliner when a key is pressed, but it will still work properly when the remove/del keys are pressed. Not
+		// updating the filter when !TreeItem.IsValid() would result in the focus not being updated when the remove/del keys are pressed.
 
-			// In any other case (i.e., if the object passes the current filter), re-add it
-			SceneOutliner->ScrollItemIntoView(TreeItem);
-		}
+		// In any other case (i.e., if the object passes the current filter), re-add it
+		SceneOutliner->ScrollItemIntoView(TreeItem);
 	}
+}
 
-	int32 FActorMode::GetTypeSortPriority(const ITreeItem& Item) const
+int32 FActorMode::GetTypeSortPriority(const ISceneOutlinerTreeItem& Item) const
+{
+	if (Item.IsA<FWorldTreeItem>())
 	{
-		if (Item.IsA<FWorldTreeItem>())
-		{
-			return EItemSortOrder::World;
-		}
-		else if (Item.IsA<FFolderTreeItem>())
-		{
-			return EItemSortOrder::Folder;
-		}
-		else if (Item.IsA<FActorTreeItem>() || Item.IsA<FComponentTreeItem>())
-		{
-			return EItemSortOrder::Actor;
-		}
-		// Warning: using actor mode with an unsupported item type!
-		check(false);
-		return -1;
+		return EItemSortOrder::World;
+	}
+	else if (Item.IsA<FFolderTreeItem>())
+	{
+		return EItemSortOrder::Folder;
 	}
+	else if (Item.IsA<FActorTreeItem>() || Item.IsA<FComponentTreeItem>())
+	{
+		return EItemSortOrder::Actor;
+	}
+	// Warning: using actor mode with an unsupported item type!
+	check(false);
+	return -1;
 }
+
 #undef LOCTEXT_NAMESPACE
\ No newline at end of file
diff --git a/Engine/Source/Editor/SceneOutliner/Private/ActorModeInteractive.cpp b/Engine/Source/Editor/SceneOutliner/Private/ActorModeInteractive.cpp
--- a/Engine/Source/Editor/SceneOutliner/Private/ActorModeInteractive.cpp
+++ b/Engine/Source/Editor/SceneOutliner/Private/ActorModeInteractive.cpp
@@ -5,103 +5,100 @@
 #include "Editor.h"
 #include "ActorTreeItem.h"
 
-namespace SceneOutliner
+FActorModeInteractive::FActorModeInteractive(SSceneOutliner* InSceneOutliner, bool bInHideComponents, TWeakObjectPtr<UWorld> InSpecifiedWorldToDisplay)
+	: FActorMode(InSceneOutliner, bHideComponents, SpecifiedWorldToDisplay)
 {
-	FActorModeInteractive::FActorModeInteractive(SSceneOutliner* InSceneOutliner, bool bInHideComponents, TWeakObjectPtr<UWorld> InSpecifiedWorldToDisplay)
-		: FActorMode(InSceneOutliner, bInHideComponents, InSpecifiedWorldToDisplay)
-	{
-		USelection::SelectionChangedEvent.AddRaw(this, &FActorModeInteractive::OnLevelSelectionChanged);
-		USelection::SelectObjectEvent.AddRaw(this, &FActorModeInteractive::OnLevelSelectionChanged);
+	USelection::SelectionChangedEvent.AddRaw(this, &FActorModeInteractive::OnLevelSelectionChanged);
+	USelection::SelectObjectEvent.AddRaw(this, &FActorModeInteractive::OnLevelSelectionChanged);
 
-		FEditorDelegates::MapChange.AddRaw(this, &FActorModeInteractive::OnMapChange);
-		FEditorDelegates::NewCurrentLevel.AddRaw(this, &FActorModeInteractive::OnNewCurrentLevel);
+	FEditorDelegates::MapChange.AddRaw(this, &FActorModeInteractive::OnMapChange);
+	FEditorDelegates::NewCurrentLevel.AddRaw(this, &FActorModeInteractive::OnNewCurrentLevel);
 
-		FCoreDelegates::OnActorLabelChanged.AddRaw(this, &FActorModeInteractive::OnActorLabelChanged);
-		GEngine->OnLevelActorRequestRename().AddRaw(this, &FActorModeInteractive::OnLevelActorRequestsRename);
-	}
+	FCoreDelegates::OnActorLabelChanged.AddRaw(this, &FActorModeInteractive::OnActorLabelChanged);
+	GEngine->OnLevelActorRequestRename().AddRaw(this, &FActorModeInteractive::OnLevelActorRequestsRename);
+}
 
-	FActorModeInteractive::~FActorModeInteractive()
-	{
-		USelection::SelectionChangedEvent.RemoveAll(this);
-		USelection::SelectObjectEvent.RemoveAll(this);
+FActorModeInteractive::~FActorModeInteractive()
+{
+	USelection::SelectionChangedEvent.RemoveAll(this);
+	USelection::SelectObjectEvent.RemoveAll(this);
 
-		FEditorDelegates::MapChange.RemoveAll(this);
-		FEditorDelegates::NewCurrentLevel.RemoveAll(this);
+	FEditorDelegates::MapChange.RemoveAll(this);
+	FEditorDelegates::NewCurrentLevel.RemoveAll(this);
 
-		FCoreDelegates::OnActorLabelChanged.RemoveAll(this);
-		GEngine->OnLevelActorRequestRename().RemoveAll(this);
-	}
+	FCoreDelegates::OnActorLabelChanged.RemoveAll(this);
+	GEngine->OnLevelActorRequestRename().RemoveAll(this);
+}
 
-	void FActorModeInteractive::OnMapChange(uint32 MapFlags)
-	{
-		// Instruct the scene outliner to generate a new hierarchy
-		SceneOutliner->FullRefresh();
-	}
+void FActorModeInteractive::OnMapChange(uint32 MapFlags)
+{
+	// Instruct the scene outliner to generate a new hierarchy
+	SceneOutliner->FullRefresh();
+}
+
+void FActorModeInteractive::OnNewCurrentLevel()
+{
+	// Instruct the scene outliner to generate a new hierarchy
+	SceneOutliner->FullRefresh();
+}
 
-	void FActorModeInteractive::OnNewCurrentLevel()
+void FActorModeInteractive::OnLevelSelectionChanged(UObject* Obj)
+{
+	const FSceneOutlinerFilterInfo* ShowOnlySelectedActorsFilter = FilterInfoMap.Find(TEXT("ShowOnlySelectedActors"));
+
+	// Since there is no way to know which items were removed/added to a selection, we must force a full refresh to handle this
+	if (ShowOnlySelectedActorsFilter && ShowOnlySelectedActorsFilter->IsFilterActive())
 	{
-		// Instruct the scene outliner to generate a new hierarchy
 		SceneOutliner->FullRefresh();
 	}
-
-	void FActorModeInteractive::OnLevelSelectionChanged(UObject* Obj)
+	else
 	{
-		const FOutlinerFilterInfo* ShowOnlySelectedActorsFilter = FilterInfoMap.Find(TEXT("ShowOnlySelectedActors"));
+		SceneOutliner->ClearSelection();
+		SceneOutliner->RefreshSelection();
 
-		// Since there is no way to know which items were removed/added to a selection, we must force a full refresh to handle this
-		if (ShowOnlySelectedActorsFilter && ShowOnlySelectedActorsFilter->IsFilterActive())
+		// Scroll last item into view - this means if we are multi-selecting, we show newest selection. @TODO Not perfect though
+		if (const AActor* LastSelectedActor = GEditor->GetSelectedActors()->GetBottom<AActor>())
 		{
-			SceneOutliner->FullRefresh();
-		}
-		else
-		{
-			SceneOutliner->ClearSelection();
-			SceneOutliner->RefreshSelection();
-
-			// Scroll last item into view - this means if we are multi-selecting, we show newest selection. @TODO Not perfect though
-			if (const AActor* LastSelectedActor = GEditor->GetSelectedActors()->GetBottom<AActor>())
+			if (FSceneOutlinerTreeItemPtr TreeItem = SceneOutliner->GetTreeItem(LastSelectedActor, false))
+			{
+				SceneOutliner->ScrollItemIntoView(TreeItem);
+			}
+			else
 			{
-				if (FTreeItemPtr TreeItem = SceneOutliner->GetTreeItem(LastSelectedActor, false))
-				{
-					SceneOutliner->ScrollItemIntoView(TreeItem);
-				}
-				else
-				{
-					SceneOutliner->OnItemAdded(LastSelectedActor, ENewItemAction::ScrollIntoView);
-				}
+				SceneOutliner->OnItemAdded(LastSelectedActor, SceneOutliner::ENewItemAction::ScrollIntoView);
 			}
 		}
 	}
+}
 
-	void FActorModeInteractive::OnLevelActorRequestsRename(const AActor* Actor)
+void FActorModeInteractive::OnLevelActorRequestsRename(const AActor* Actor)
+{
+	const auto& SelectedItems = SceneOutliner->GetSelectedItems();
+	if (SelectedItems.Num() > 0)
 	{
-		const auto& SelectedItems = SceneOutliner->GetSelectedItems();
-		if (SelectedItems.Num() > 0)
+		// Ensure that the item we want to rename is visible in the tree
+		FSceneOutlinerTreeItemPtr ItemToRename = SelectedItems[SelectedItems.Num() - 1];
+		if (SceneOutliner->CanExecuteRenameRequest(*ItemToRename) && ItemToRename->CanInteract())
 		{
-			// Ensure that the item we want to rename is visible in the tree
-			FTreeItemPtr ItemToRename = SelectedItems[SelectedItems.Num() - 1];
-			if (SceneOutliner->CanExecuteRenameRequest(*ItemToRename) && ItemToRename->CanInteract())
-			{
-				SceneOutliner->SetPendingRenameItem(ItemToRename);
-				SceneOutliner->ScrollItemIntoView(ItemToRename);
-			}
+			SceneOutliner->SetPendingRenameItem(ItemToRename);
+			SceneOutliner->ScrollItemIntoView(ItemToRename);
 		}
 	}
+}
 
-	void FActorModeInteractive::OnActorLabelChanged(AActor* ChangedActor)
+void FActorModeInteractive::OnActorLabelChanged(AActor* ChangedActor)
+{
+	if (!ensure(ChangedActor))
 	{
-		if (!ensure(ChangedActor))
-		{
-			return;
-		}
+		return;
+	}
 
-		if (IsActorDisplayable(ChangedActor) && RepresentingWorld.Get() == ChangedActor->GetWorld())
+	if (IsActorDisplayable(ChangedActor) && RepresentingWorld.Get() == ChangedActor->GetWorld())
+	{
+		// Force create the item otherwise the outliner may not be notified of a change to the item if it is filtered out
+		if (FSceneOutlinerTreeItemPtr Item = CreateItemFor<FActorTreeItem>(ChangedActor, true))
 		{
-			// Force create the item otherwise the outliner may not be notified of a change to the item if it is filtered out
-			if (FTreeItemPtr Item = CreateItemFor<FActorTreeItem>(ChangedActor, true))
-			{
-				SceneOutliner->OnItemLabelChanged(Item);
-			}
+			SceneOutliner->OnItemLabelChanged(Item);
 		}
 	}
-}
\ No newline at end of file
+}
diff --git a/Engine/Source/Editor/SceneOutliner/Private/ActorPickingMode.cpp b/Engine/Source/Editor/SceneOutliner/Private/ActorPickingMode.cpp
--- a/Engine/Source/Editor/SceneOutliner/Private/ActorPickingMode.cpp
+++ b/Engine/Source/Editor/SceneOutliner/Private/ActorPickingMode.cpp
@@ -4,45 +4,42 @@
 #include "SSceneOutliner.h"
 #include "ActorTreeItem.h"
 
-namespace SceneOutliner
+FActorPickingMode::FActorPickingMode(SSceneOutliner* InSceneOutliner, bool bInHideComponents, FOnSceneOutlinerItemPicked OnItemPickedDelegate, TWeakObjectPtr<UWorld> InSpecifiedWorldToDisplay)
+	: FActorMode(InSceneOutliner, bInHideComponents, InSpecifiedWorldToDisplay)
+	, OnItemPicked(OnItemPickedDelegate)
 {
-	FActorPickingMode::FActorPickingMode(SSceneOutliner* InSceneOutliner, bool bInHideComponents, FOnSceneOutlinerItemPicked OnItemPickedDelegate, TWeakObjectPtr<UWorld> InSpecifiedWorldToDisplay)
-		: FActorMode(InSceneOutliner, bInHideComponents, InSpecifiedWorldToDisplay)
-		, OnItemPicked(OnItemPickedDelegate)
-	{
-	}
+}
 
-	void FActorPickingMode::OnItemSelectionChanged(FTreeItemPtr Item, ESelectInfo::Type SelectionType, const FItemSelection& Selection)
+void FActorPickingMode::OnItemSelectionChanged(FSceneOutlinerTreeItemPtr Item, ESelectInfo::Type SelectionType, const FSceneOutlinerItemSelection& Selection)
+{
+	// In actor picking mode, we fire off the notification to whoever is listening.
+	// This may often cause the widget itself to be enqueued for destruction
+	auto SelectedItems = SceneOutliner->GetSelectedItems();
+	if (SelectedItems.Num() > 0)
 	{
-		// In actor picking mode, we fire off the notification to whoever is listening.
-		// This may often cause the widget itself to be enqueued for destruction
-		auto SelectedItems = SceneOutliner->GetSelectedItems();
-		if (SelectedItems.Num() > 0)
+		auto FirstItem = SelectedItems[0];
+		if (FirstItem->CanInteract())
 		{
-			auto FirstItem = SelectedItems[0];
-			if (FirstItem->CanInteract())
-			{
-				OnItemPicked.ExecuteIfBound(FirstItem.ToSharedRef());
-			}
+			OnItemPicked.ExecuteIfBound(FirstItem.ToSharedRef());
 		}
 	}
+}
 
-	void FActorPickingMode::OnFilterTextCommited(FItemSelection& Selection, ETextCommit::Type CommitType)
+void FActorPickingMode::OnFilterTextCommited(FSceneOutlinerItemSelection& Selection, ETextCommit::Type CommitType)
+{
+	// In actor picking mode, we check to see if we have a selected actor, and if so, fire
+	// off the notification to whoever is listening.  This may often cause the widget itself
+	// to be enqueued for destruction
+	TArray<FActorTreeItem*> Actors;
+	Selection.Get(Actors);
+	if (Actors.Num() == 1 && Actors[0])
 	{
-		// In actor picking mode, we check to see if we have a selected actor, and if so, fire
-		// off the notification to whoever is listening.  This may often cause the widget itself
-		// to be enqueued for destruction
-		TArray<FActorTreeItem*> Actors;
-		Selection.Get(Actors);
-		if (Actors.Num() == 1 && Actors[0])
-		{
-			// Signal that an actor was selected. We assume it is valid as it won't have been added to ActorsToSelect if not.
-			SceneOutliner->SetItemSelection(Actors[0]->AsShared(), true, ESelectInfo::OnKeyPress);
-		}
+		// Signal that an actor was selected. We assume it is valid as it won't have been added to ActorsToSelect if not.
+		SceneOutliner->SetItemSelection(Actors[0]->AsShared(), true, ESelectInfo::OnKeyPress);
 	}
+}
 
-	void FActorPickingMode::CreateViewContent(FMenuBuilder& MenuBuilder)
-	{
-		FActorMode::BuildWorldPickerMenu(MenuBuilder);
-	}
-}
\ No newline at end of file
+void FActorPickingMode::CreateViewContent(FMenuBuilder& MenuBuilder)
+{
+	FActorMode::BuildWorldPickerMenu(MenuBuilder);
+}
diff --git a/Engine/Source/Editor/SceneOutliner/Private/ActorTreeItem.cpp b/Engine/Source/Editor/SceneOutliner/Private/ActorTreeItem.cpp
--- a/Engine/Source/Editor/SceneOutliner/Private/ActorTreeItem.cpp
+++ b/Engine/Source/Editor/SceneOutliner/Private/ActorTreeItem.cpp
@@ -20,16 +20,14 @@
 
 #define LOCTEXT_NAMESPACE "SceneOutliner_ActorTreeItem"
 
-namespace SceneOutliner
-{
-const FTreeItemType FActorTreeItem::Type(&ITreeItem::Type);
+const FSceneOutlinerTreeItemType FActorTreeItem::Type(&ISceneOutlinerTreeItem::Type);
 
-struct SActorTreeLabel : FCommonLabelData, public SCompoundWidget
+struct SActorTreeLabel : FSceneOutlinerCommonLabelData, public SCompoundWidget
 {
 	SLATE_BEGIN_ARGS(SActorTreeLabel) {}
 	SLATE_END_ARGS()
 
-	void Construct(const FArguments& InArgs, FActorTreeItem& ActorItem, ISceneOutliner& SceneOutliner, const STableRow<FTreeItemPtr>& InRow)
+	void Construct(const FArguments& InArgs, FActorTreeItem& ActorItem, ISceneOutliner& SceneOutliner, const STableRow<FSceneOutlinerTreeItemPtr>& InRow)
 	{
 		WeakSceneOutliner = StaticCastSharedRef<ISceneOutliner>(SceneOutliner.AsShared());
 
@@ -53,7 +51,7 @@ struct SActorTreeLabel : FCommonLabelData, public SCompoundWidget
 				.ColorAndOpacity(this, &SActorTreeLabel::GetForegroundColor)
 				.OnTextCommitted(this, &SActorTreeLabel::OnLabelCommitted)
 				.OnVerifyTextChanged(this, &SActorTreeLabel::OnVerifyItemLabelChanged)
-				.IsSelected(FIsSelected::CreateSP(&InRow, &STableRow<FTreeItemPtr>::IsSelectedExclusively))
+				.IsSelected(FIsSelected::CreateSP(&InRow, &STableRow<FSceneOutlinerTreeItemPtr>::IsSelectedExclusively))
 				.IsReadOnly_Lambda([Item = ActorItem.AsShared(), this]()
 				{
 					return !CanExecuteRenameRequest(Item.Get());
@@ -83,11 +81,11 @@ struct SActorTreeLabel : FCommonLabelData, public SCompoundWidget
 				+ SHorizontalBox::Slot()
 				.AutoWidth()
 				.VAlign(VAlign_Center)
-				.Padding(FDefaultTreeItemMetrics::IconPadding())
+				.Padding(FSceneOutlinerDefaultTreeItemMetrics::IconPadding())
 				[
 					SNew(SBox)
-					.WidthOverride(FDefaultTreeItemMetrics::IconSize())
-					.HeightOverride(FDefaultTreeItemMetrics::IconSize())
+					.WidthOverride(FSceneOutlinerDefaultTreeItemMetrics::IconSize())
+					.HeightOverride(FSceneOutlinerDefaultTreeItemMetrics::IconSize())
 					[
 						SNew(SImage)
 						.Image(this, &SActorTreeLabel::GetIcon)
@@ -236,8 +234,7 @@ private:
 	FSlateColor GetForegroundColor() const
 	{
 		auto TreeItem = TreeItemPtr.Pin();
-
-		if (auto BaseColor = FCommonLabelData::GetForegroundColor(*TreeItem))
+		if (auto BaseColor = FSceneOutlinerCommonLabelData::GetForegroundColor(*TreeItem))
 		{
 			return BaseColor.GetValue();
 		}
@@ -249,14 +246,14 @@ private:
 			return FLinearColor(0.2f, 0.2f, 0.25f);
 		}
 
-		UWorld* OwningWorld = TreeItem->Actor->GetWorld();
+		UWorld* OwningWorld = Actor->GetWorld();
 		if (!OwningWorld)
 		{
-			// Invalid world!
+			// Deleted world!
 			return FLinearColor(0.2f, 0.2f, 0.25f);
 		}
 
-		const bool bRepresentingPIEWorld = OwningWorld->IsPlayInEditor();
+		const bool bRepresentingPIEWorld = TreeItem->Actor->GetWorld()->IsPlayInEditor();
 		if (bRepresentingPIEWorld && !TreeItem->bExistsInCurrentWorldAndPIE)
 		{
 			// Highlight actors that are exclusive to PlayWorld
@@ -268,7 +265,7 @@ private:
 		const bool bSelectEvenIfHidden = true;		// @todo outliner: Is this actually OK?
 		if (!GEditor->CanSelectActor(Actor, bInSelected, bSelectEvenIfHidden))
 		{
-			return FCommonLabelData::DarkColor;
+			return FSceneOutlinerCommonLabelData::DarkColor;
 		}
 
 		return FSlateColor::UseForeground();
@@ -297,14 +294,14 @@ private:
 };
 
 FActorTreeItem::FActorTreeItem(AActor* InActor)
-	: ITreeItem(Type)
+	: ISceneOutlinerTreeItem(Type)
 	, Actor(InActor)
 	, ID(InActor)
 {
 	bExistsInCurrentWorldAndPIE = GEditor->ObjectsThatExistInEditorWorld.Get(InActor);
 }
 
-FTreeItemID FActorTreeItem::GetID() const
+FSceneOutlinerTreeItemID FActorTreeItem::GetID() const
 {
 	return ID;
 }
@@ -333,7 +330,7 @@ bool FActorTreeItem::CanInteract() const
 	return true;
 }
 
-TSharedRef<SWidget> FActorTreeItem::GenerateLabelWidget(ISceneOutliner& Outliner, const STableRow<FTreeItemPtr>& InRow)
+TSharedRef<SWidget> FActorTreeItem::GenerateLabelWidget(ISceneOutliner& Outliner, const STableRow<FSceneOutlinerTreeItemPtr>& InRow)
 {
 	return SNew(SActorTreeLabel, *this, Outliner, InRow);
 }
@@ -352,6 +349,4 @@ bool FActorTreeItem::GetVisibility() const
 	return Actor.IsValid() && !Actor->IsTemporarilyHiddenInEditor(true);
 }
 
-}	// namespace SceneOutliner
-
 #undef LOCTEXT_NAMESPACE
diff --git a/Engine/Source/Editor/SceneOutliner/Private/ComponentTreeItem.cpp b/Engine/Source/Editor/SceneOutliner/Private/ComponentTreeItem.cpp
--- a/Engine/Source/Editor/SceneOutliner/Private/ComponentTreeItem.cpp
+++ b/Engine/Source/Editor/SceneOutliner/Private/ComponentTreeItem.cpp
@@ -15,16 +15,12 @@
 
 #define LOCTEXT_NAMESPACE "SceneOutliner_ComponentTreeItem"
 
-
-namespace SceneOutliner
-{
-	
-struct SComponentTreeLabel : FCommonLabelData, public SCompoundWidget
+struct SComponentTreeLabel : FSceneOutlinerCommonLabelData, public SCompoundWidget
 {
 	SLATE_BEGIN_ARGS(SComponentTreeLabel) {}
 	SLATE_END_ARGS()
 
-	void Construct(const FArguments& InArgs, FComponentTreeItem& ComponentItem, ISceneOutliner& SceneOutliner, const STableRow<FTreeItemPtr>& InRow)
+	void Construct(const FArguments& InArgs, FComponentTreeItem& ComponentItem, ISceneOutliner& SceneOutliner, const STableRow<FSceneOutlinerTreeItemPtr>& InRow)
 	{
 		TreeItemPtr = StaticCastSharedRef<FComponentTreeItem>(ComponentItem.AsShared());
 		WeakSceneOutliner = StaticCastSharedRef<ISceneOutliner>(SceneOutliner.AsShared());
@@ -51,7 +47,7 @@ struct SComponentTreeLabel : FCommonLabelData, public SCompoundWidget
 				{
 					return !CanExecuteRenameRequest(Item.Get());
 				})
-				.IsSelected(FIsSelected::CreateSP(&InRow, &STableRow<FTreeItemPtr>::IsSelectedExclusively))
+				.IsSelected(FIsSelected::CreateSP(&InRow, &STableRow<FSceneOutlinerTreeItemPtr>::IsSelectedExclusively))
 			];
 
 		if (WeakSceneOutliner.Pin()->GetMode()->IsInteractive())
@@ -66,11 +62,11 @@ struct SComponentTreeLabel : FCommonLabelData, public SCompoundWidget
 				+ SHorizontalBox::Slot()
 				.AutoWidth()
 				.VAlign(VAlign_Center)
-				.Padding(FDefaultTreeItemMetrics::IconPadding())
+				.Padding(FSceneOutlinerDefaultTreeItemMetrics::IconPadding())
 				[
 					SNew(SBox)
-					.WidthOverride(FDefaultTreeItemMetrics::IconSize())
-					.HeightOverride(FDefaultTreeItemMetrics::IconSize())
+					.WidthOverride(FSceneOutlinerDefaultTreeItemMetrics::IconSize())
+					.HeightOverride(FSceneOutlinerDefaultTreeItemMetrics::IconSize())
 					[
 						SNew(SImage)
 						.Image(this, &SComponentTreeLabel::GetIcon)
@@ -166,7 +162,7 @@ private:
 
 	FSlateColor GetForegroundColor() const
 	{
-		if (auto BaseColor = FCommonLabelData::GetForegroundColor(*TreeItemPtr.Pin()))
+		if (auto BaseColor = FSceneOutlinerCommonLabelData::GetForegroundColor(*TreeItemPtr.Pin()))
 		{
 			return BaseColor.GetValue();
 		}
@@ -175,11 +171,10 @@ private:
 	}
 };
 
-
-const FTreeItemType FComponentTreeItem::Type(&ITreeItem::Type);
+const FSceneOutlinerTreeItemType FComponentTreeItem::Type(&ISceneOutlinerTreeItem::Type);
 
 FComponentTreeItem::FComponentTreeItem(UActorComponent* InComponent)
-	: ITreeItem(Type)
+	: ISceneOutlinerTreeItem(Type)
 	, Component(InComponent)
 	, ID(InComponent)
 {
@@ -195,7 +190,7 @@ FComponentTreeItem::FComponentTreeItem(UActorComponent* InComponent)
 	}
 }
 
-FTreeItemID FComponentTreeItem::GetID() const
+FSceneOutlinerTreeItemID FComponentTreeItem::GetID() const
 {
 	return ID;
 }
@@ -237,10 +232,9 @@ bool FComponentTreeItem::CanInteract() const
 	return true;
 }
 
-TSharedRef<SWidget> FComponentTreeItem::GenerateLabelWidget(ISceneOutliner& Outliner, const STableRow<FTreeItemPtr>& InRow)
+TSharedRef<SWidget> FComponentTreeItem::GenerateLabelWidget(ISceneOutliner& Outliner, const STableRow<FSceneOutlinerTreeItemPtr>& InRow)
 {
 	return SNew(SComponentTreeLabel, *this, Outliner, InRow);
 }
-} // namespace SceneOutliner
 
 #undef LOCTEXT_NAMESPACE
diff --git a/Engine/Source/Editor/SceneOutliner/Private/FolderTreeItem.cpp b/Engine/Source/Editor/SceneOutliner/Private/FolderTreeItem.cpp
--- a/Engine/Source/Editor/SceneOutliner/Private/FolderTreeItem.cpp
+++ b/Engine/Source/Editor/SceneOutliner/Private/FolderTreeItem.cpp
@@ -16,39 +16,39 @@
 
 namespace SceneOutliner
 {
-const FTreeItemType FFolderTreeItem::Type(&ITreeItem::Type);
-
-bool FFolderPathSelector::operator()(TWeakPtr<ITreeItem> Item, FName& DataOut) const
-{
-	if (FFolderTreeItem* FolderItem = Item.Pin()->CastTo<FFolderTreeItem>())
+	bool FFolderPathSelector::operator()(TWeakPtr<ISceneOutlinerTreeItem> Item, FName& DataOut) const
 	{
-		if (FolderItem->IsValid())
+		if (FFolderTreeItem* FolderItem = Item.Pin()->CastTo<FFolderTreeItem>())
 		{
-			DataOut = FolderItem->Path;
-			return true;
+			if (FolderItem->IsValid())
+			{
+				DataOut = FolderItem->Path;
+				return true;
+			}
 		}
+		return false;
 	}
-	return false;
 }
 
+const FSceneOutlinerTreeItemType FFolderTreeItem::Type(&ISceneOutlinerTreeItem::Type);
 
 FFolderTreeItem::FFolderTreeItem(FName InPath)
-	: ITreeItem(Type)
+	: ISceneOutlinerTreeItem(Type)
 	, Path(InPath)
-	, LeafName(GetFolderLeafName(InPath))
+	, LeafName(SceneOutliner::GetFolderLeafName(InPath))
 {
 }
 
-FFolderTreeItem::FFolderTreeItem(FName InPath, FTreeItemType InType)
-	: ITreeItem(InType)
+FFolderTreeItem::FFolderTreeItem(FName InPath, FSceneOutlinerTreeItemType InType)
+	: ISceneOutlinerTreeItem(InType)
 	, Path(InPath)
-	, LeafName(GetFolderLeafName(InPath))
+	, LeafName(SceneOutliner::GetFolderLeafName(InPath))
 {
 }
 
-FTreeItemID FFolderTreeItem::GetID() const
+FSceneOutlinerTreeItemID FFolderTreeItem::GetID() const
 {
-	return FTreeItemID(Path);
+	return FSceneOutlinerTreeItemID(Path);
 }
 
 FString FFolderTreeItem::GetDisplayString() const
@@ -83,6 +83,4 @@ void FFolderTreeItem::GenerateContextMenu(UToolMenu* Menu, SSceneOutliner& Outli
 	Section.AddMenuEntry("DuplicateFolderHierarchy", LOCTEXT("DuplicateFolderHierarchy", "Duplicate Hierarchy"), FText(), FSlateIcon(), FUIAction(FExecuteAction::CreateSP(&Outliner, &SSceneOutliner::DuplicateFoldersHierarchy)));
 }
 
-}	// namespace SceneOutliner
-
 #undef LOCTEXT_NAMESPACE
