Removed SceneOutliner namespace

diff --git a/Engine/Source/Editor/SceneOutliner/Private/SOutlinerTreeView.cpp b/Engine/Source/Editor/SceneOutliner/Private/SOutlinerTreeView.cpp
--- a/Engine/Source/Editor/SceneOutliner/Private/SOutlinerTreeView.cpp
+++ b/Engine/Source/Editor/SceneOutliner/Private/SOutlinerTreeView.cpp
@@ -13,385 +13,382 @@
 
 #define LOCTEXT_NAMESPACE "SSceneOutliner"
 
-namespace SceneOutliner
+static void UpdateOperationDecorator(const FDragDropEvent& Event, const FSceneOutlinerDragValidationInfo& ValidationInfo)
 {
-	static void UpdateOperationDecorator(const FDragDropEvent& Event, const FDragValidationInfo& ValidationInfo)
-	{
-		const FSlateBrush* Icon = ValidationInfo.IsValid() ? FEditorStyle::GetBrush(TEXT("Graph.ConnectorFeedback.OK")) : FEditorStyle::GetBrush(TEXT("Graph.ConnectorFeedback.Error"));
-
-		FDragDropOperation* Operation = Event.GetOperation().Get();
-		if (Operation)
-		{
-			if (Operation->IsOfType<FSceneOutlinerDragDropOp>())
-			{
-				auto* OutlinerOp = static_cast<FSceneOutlinerDragDropOp*>(Operation);
-				OutlinerOp->SetTooltip(ValidationInfo.ValidationText, Icon);
-			}
-			else if (Operation->IsOfType<FDecoratedDragDropOp>())
-			{
-				auto* DecoratedOp = static_cast<FDecoratedDragDropOp*>(Operation);
-				DecoratedOp->SetToolTip(LOCTEXT("SceneOutlinerInvalidDragDropOp", "Invalid drag and drop operation, Drag into the Viewport."), Icon);
-			}
-		}
-	}
+	const FSlateBrush* Icon = ValidationInfo.IsValid() ? FEditorStyle::GetBrush(TEXT("Graph.ConnectorFeedback.OK")) : FEditorStyle::GetBrush(TEXT("Graph.ConnectorFeedback.Error"));
 
-	static void ResetOperationDecorator(const FDragDropEvent& Event)
+	FDragDropOperation* Operation = Event.GetOperation().Get();
+	if (Operation)
 	{
-		FDragDropOperation* Operation = Event.GetOperation().Get();
-		if (Operation)
+		if (Operation->IsOfType<FSceneOutlinerDragDropOp>())
 		{
-			if (Operation->IsOfType<FSceneOutlinerDragDropOp>())
-			{
-				static_cast<FSceneOutlinerDragDropOp*>(Operation)->ResetTooltip();
-			}
-			else if (Operation->IsOfType<FDecoratedDragDropOp>())
-			{
-				static_cast<FDecoratedDragDropOp*>(Operation)->ResetToDefaultToolTip();
-			}
+			auto* OutlinerOp = static_cast<FSceneOutlinerDragDropOp*>(Operation);
+			OutlinerOp->SetTooltip(ValidationInfo.ValidationText, Icon);
 		}
-	}
-
-	static FReply OnDragDetected( const FGeometry& MyGeometry, const FPointerEvent& MouseEvent, TWeakPtr<SOutlinerTreeView> Table )
-	{
-		auto TablePtr = Table.Pin();
-		if (TablePtr.IsValid() && MouseEvent.IsMouseButtonDown( EKeys::LeftMouseButton ))
+		else if (Operation->IsOfType<FDecoratedDragDropOp>())
 		{
-			auto Operation = TablePtr->GetOutlinerPtr().Pin()->CreateDragDropOperation(TablePtr->GetSelectedItems());
-
-			if (Operation.IsValid())
-			{
-				return FReply::Handled().BeginDragDrop(Operation.ToSharedRef());
-			}
+			auto* DecoratedOp = static_cast<FDecoratedDragDropOp*>(Operation);
+			DecoratedOp->SetToolTip(LOCTEXT("SceneOutlinerInvalidDragDropOp", "Invalid drag and drop operation, Drag into the Viewport."), Icon);
 		}
-
-		return FReply::Unhandled();
 	}
+}
 
-	FReply HandleDrop(TSharedPtr<SSceneOutliner> SceneOutlinerPtr, const FDragDropEvent& DragDropEvent, ITreeItem& DropTarget, FDragValidationInfo& ValidationInfo, bool bApplyDrop = false)
+static void ResetOperationDecorator(const FDragDropEvent& Event)
+{
+	FDragDropOperation* Operation = Event.GetOperation().Get();
+	if (Operation)
 	{
-		if (!SceneOutlinerPtr.IsValid())
-		{
-			return FReply::Unhandled();
-		}
-
-		// Don't handle this if the scene outliner is not in a mode which supports drag and drop
-		if (!SceneOutlinerPtr->CanSupportDragAndDrop())
+		if (Operation->IsOfType<FSceneOutlinerDragDropOp>())
 		{
-			return FReply::Unhandled();
+			static_cast<FSceneOutlinerDragDropOp*>(Operation)->ResetTooltip();
 		}
-
-		FDragDropPayload DraggedObjects;
-		// Validate now to make sure we don't doing anything we shouldn't
-		if (!SceneOutlinerPtr->ParseDragDrop(DraggedObjects, *DragDropEvent.GetOperation()))
+		else if (Operation->IsOfType<FDecoratedDragDropOp>())
 		{
-			return FReply::Unhandled();
+			static_cast<FDecoratedDragDropOp*>(Operation)->ResetToDefaultToolTip();
 		}
+	}
+}
 
-		ValidationInfo = SceneOutlinerPtr->ValidateDrop(StaticCast<ITreeItem&>(DropTarget), DraggedObjects);
+static FReply HandleOnDragDetected( const FGeometry& MyGeometry, const FPointerEvent& MouseEvent, TWeakPtr<SSceneOutlinerTreeView> Table )
+{
+	auto TablePtr = Table.Pin();
+	if (TablePtr.IsValid() && MouseEvent.IsMouseButtonDown( EKeys::LeftMouseButton ))
+	{
+		auto Operation = TablePtr->GetOutlinerPtr().Pin()->CreateDragDropOperation(TablePtr->GetSelectedItems());
 
-		if (!ValidationInfo.IsValid())
+		if (Operation.IsValid())
 		{
-			// Return handled here to stop anything else trying to handle it - the operation is invalid as far as we're concerned
-			return FReply::Handled();
+			return FReply::Handled().BeginDragDrop(Operation.ToSharedRef());
 		}
+	}
 
-		if (bApplyDrop)
-		{
-			SceneOutlinerPtr->OnDropPayload(DropTarget, DraggedObjects, ValidationInfo);
-		}
+	return FReply::Unhandled();
+}
 
-		return FReply::Handled();
+FReply HandleDrop(TSharedPtr<SSceneOutliner> SceneOutlinerPtr, const FDragDropEvent& DragDropEvent, ISceneOutlinerTreeItem& DropTarget, FSceneOutlinerDragValidationInfo& ValidationInfo, bool bApplyDrop = false)
+{
+	if (!SceneOutlinerPtr.IsValid())
+	{
+		return FReply::Unhandled();
 	}
 
-	FReply HandleDropFromWeak(TWeakPtr<SSceneOutliner> SceneOutlinerWeak, const FDragDropEvent& DragDropEvent, ITreeItem& DropTarget, FDragValidationInfo& ValidationInfo, bool bApplyDrop = false)
+	// Don't handle this if the scene outliner is not in a mode which supports drag and drop
+	if (!SceneOutlinerPtr->CanSupportDragAndDrop())
 	{
-		return HandleDrop(SceneOutlinerWeak.Pin(), DragDropEvent, DropTarget, ValidationInfo, bApplyDrop);
+		return FReply::Unhandled();
 	}
 
-	void SOutlinerTreeView::Construct(const SOutlinerTreeView::FArguments& InArgs, TSharedRef<SSceneOutliner> Owner)
+	FSceneOutlinerDragDropPayload DraggedObjects;
+	// Validate now to make sure we don't doing anything we shouldn't
+	if (!SceneOutlinerPtr->ParseDragDrop(DraggedObjects, *DragDropEvent.GetOperation()))
 	{
-		SceneOutlinerWeak = Owner;
-		STreeView::Construct(InArgs);
+		return FReply::Unhandled();
 	}
 
-	void SOutlinerTreeView::FlashHighlightOnItem( FTreeItemPtr FlashHighlightOnItem )
+	ValidationInfo = SceneOutlinerPtr->ValidateDrop(StaticCast<ISceneOutlinerTreeItem&>(DropTarget), DraggedObjects);
+
+	if (!ValidationInfo.IsValid())
 	{
-		TSharedPtr< SSceneOutlinerTreeRow > RowWidget = StaticCastSharedPtr< SSceneOutlinerTreeRow >( WidgetGenerator.GetWidgetForItem( FlashHighlightOnItem ) );
-		if( RowWidget.IsValid() )
-		{
-			RowWidget->FlashHighlight();
-		}
+		// Return handled here to stop anything else trying to handle it - the operation is invalid as far as we're concerned
+		return FReply::Handled();
 	}
 
-	FReply SOutlinerTreeView::OnDragOver(const FGeometry& MyGeometry, const FDragDropEvent& DragDropEvent)
+	if (bApplyDrop)
 	{
-		FDragValidationInfo ValidationInfo = FDragValidationInfo::Invalid();
+		SceneOutlinerPtr->OnDropPayload(DropTarget, DraggedObjects, ValidationInfo);
+	}
 
-		FFolderTreeItem DropTarget(NAME_None);
+	return FReply::Handled();
+}
 
-		auto Reply = HandleDropFromWeak(SceneOutlinerWeak, DragDropEvent, DropTarget, ValidationInfo);
-		UpdateOperationDecorator(DragDropEvent, ValidationInfo);
+FReply HandleDropFromWeak(TWeakPtr<SSceneOutliner> SceneOutlinerWeak, const FDragDropEvent& DragDropEvent, ISceneOutlinerTreeItem& DropTarget, FSceneOutlinerDragValidationInfo& ValidationInfo, bool bApplyDrop = false)
+{
+	return HandleDrop(SceneOutlinerWeak.Pin(), DragDropEvent, DropTarget, ValidationInfo, bApplyDrop);
+}
 
-		return Reply;
-	}
+void SSceneOutlinerTreeView::Construct(const SSceneOutlinerTreeView::FArguments& InArgs, TSharedRef<SSceneOutliner> Owner)
+{
+	SceneOutlinerWeak = Owner;
+	STreeView::Construct(InArgs);
+}
 
-	void SOutlinerTreeView::OnDragLeave(const FDragDropEvent& DragDropEvent)
+void SSceneOutlinerTreeView::FlashHighlightOnItem( FSceneOutlinerTreeItemPtr FlashHighlightOnItem )
+{
+	TSharedPtr< SSceneOutlinerTreeRow > RowWidget = StaticCastSharedPtr< SSceneOutlinerTreeRow >( WidgetGenerator.GetWidgetForItem( FlashHighlightOnItem ) );
+	if( RowWidget.IsValid() )
 	{
-		if( SceneOutlinerWeak.Pin()->GetSharedData().bShowParentTree )
-		{
-			ResetOperationDecorator(DragDropEvent);
-		}
+		RowWidget->FlashHighlight();
 	}
+}
 
-	FReply SOutlinerTreeView::OnDrop(const FGeometry& MyGeometry, const FDragDropEvent& DragDropEvent)
-	{
-		FDragValidationInfo ValidationInfo = FDragValidationInfo::Invalid();
-		FFolderTreeItem DropTarget(NAME_None);
+FReply SSceneOutlinerTreeView::OnDragOver(const FGeometry& MyGeometry, const FDragDropEvent& DragDropEvent)
+{
+	FSceneOutlinerDragValidationInfo ValidationInfo = FSceneOutlinerDragValidationInfo::Invalid();
 
-		return HandleDropFromWeak(SceneOutlinerWeak, DragDropEvent, DropTarget, ValidationInfo, true);
-	}
+	FFolderTreeItem DropTarget(NAME_None);
 
-	FReply SSceneOutlinerTreeRow::OnDrop( const FGeometry& MyGeometry, const FDragDropEvent& DragDropEvent )
-	{
-		auto ItemPtr = Item.Pin();
-		auto SceneOutlinerPtr = SceneOutlinerWeak.Pin();
-		if (ItemPtr.IsValid() && SceneOutlinerPtr.IsValid())
-		{
-			FDragValidationInfo ValidationInfo = FDragValidationInfo::Invalid();
-			return HandleDrop(SceneOutlinerPtr, DragDropEvent, *ItemPtr, ValidationInfo, true);
-		}
+	auto Reply = HandleDropFromWeak(SceneOutlinerWeak, DragDropEvent, DropTarget, ValidationInfo);
+	UpdateOperationDecorator(DragDropEvent, ValidationInfo);
 
-		return FReply::Unhandled();
-	}
+	return Reply;
+}
 
-	void SSceneOutlinerTreeRow::OnDragEnter( const FGeometry& MyGeometry, const FDragDropEvent& DragDropEvent )
+void SSceneOutlinerTreeView::OnDragLeave(const FDragDropEvent& DragDropEvent)
+{
+	if( SceneOutlinerWeak.Pin()->GetSharedData().bShowParentTree )
 	{
-		auto ItemPtr = Item.Pin();
-		auto SceneOutlinerPtr = SceneOutlinerWeak.Pin();
-		if (ItemPtr.IsValid() && SceneOutlinerPtr.IsValid())
-		{
-			FDragValidationInfo ValidationInfo = FDragValidationInfo::Invalid();
+		ResetOperationDecorator(DragDropEvent);
+	}
+}
 
-			HandleDrop(SceneOutlinerPtr, DragDropEvent, *ItemPtr, ValidationInfo, false);
-			UpdateOperationDecorator(DragDropEvent, ValidationInfo);
-		}
+FReply SSceneOutlinerTreeView::OnDrop(const FGeometry& MyGeometry, const FDragDropEvent& DragDropEvent)
+{
+	FSceneOutlinerDragValidationInfo ValidationInfo = FSceneOutlinerDragValidationInfo::Invalid();
+	FFolderTreeItem DropTarget(NAME_None);
+
+	return HandleDropFromWeak(SceneOutlinerWeak, DragDropEvent, DropTarget, ValidationInfo, true);
+}
+
+FReply SSceneOutlinerTreeRow::OnDrop( const FGeometry& MyGeometry, const FDragDropEvent& DragDropEvent )
+{
+	auto ItemPtr = Item.Pin();
+	auto SceneOutlinerPtr = SceneOutlinerWeak.Pin();
+	if (ItemPtr.IsValid() && SceneOutlinerPtr.IsValid())
+	{
+		FSceneOutlinerDragValidationInfo ValidationInfo = FSceneOutlinerDragValidationInfo::Invalid();
+		return HandleDrop(SceneOutlinerPtr, DragDropEvent, *ItemPtr, ValidationInfo, true);
 	}
 
-	void SSceneOutlinerTreeRow::OnDragLeave( const FDragDropEvent& DragDropEvent )
+	return FReply::Unhandled();
+}
+
+void SSceneOutlinerTreeRow::OnDragEnter( const FGeometry& MyGeometry, const FDragDropEvent& DragDropEvent )
+{
+	auto ItemPtr = Item.Pin();
+	auto SceneOutlinerPtr = SceneOutlinerWeak.Pin();
+	if (ItemPtr.IsValid() && SceneOutlinerPtr.IsValid())
 	{
-		auto ItemPtr = Item.Pin();
-		auto SceneOutlinerPtr = SceneOutlinerWeak.Pin();
+		FSceneOutlinerDragValidationInfo ValidationInfo = FSceneOutlinerDragValidationInfo::Invalid();
 
-		ResetOperationDecorator(DragDropEvent);
+		HandleDrop(SceneOutlinerPtr, DragDropEvent, *ItemPtr, ValidationInfo, false);
+		UpdateOperationDecorator(DragDropEvent, ValidationInfo);
 	}
+}
 
-	FReply SSceneOutlinerTreeRow::OnDragOver(const FGeometry& MyGeometry, const FDragDropEvent& DragDropEvent)
+void SSceneOutlinerTreeRow::OnDragLeave( const FDragDropEvent& DragDropEvent )
+{
+	auto ItemPtr = Item.Pin();
+	auto SceneOutlinerPtr = SceneOutlinerWeak.Pin();
+
+	ResetOperationDecorator(DragDropEvent);
+}
+
+FReply SSceneOutlinerTreeRow::OnDragOver(const FGeometry& MyGeometry, const FDragDropEvent& DragDropEvent)
+{
+	auto SceneOutlinerPtr = SceneOutlinerWeak.Pin();
+	if (SSceneOutliner* SceneOutliner = SceneOutlinerPtr.Get())
 	{
-		auto SceneOutlinerPtr = SceneOutlinerWeak.Pin();
-		if (SSceneOutliner* SceneOutliner = SceneOutlinerPtr.Get())
+		if (const auto* ItemPtr = Item.Pin().Get())
 		{
-			if (const auto* ItemPtr = Item.Pin().Get())
-			{
-				return SceneOutliner->OnDragOverItem(DragDropEvent, *ItemPtr);
-			}
-			return FReply::Unhandled();
+			return SceneOutliner->OnDragOverItem(DragDropEvent, *ItemPtr);
 		}
-
-		return FReply::Handled();
+		return FReply::Unhandled();
 	}
 
-	FReply SSceneOutlinerTreeRow::OnMouseButtonDown( const FGeometry& MyGeometry, const FPointerEvent& MouseEvent )
+	return FReply::Handled();
+}
+
+FReply SSceneOutlinerTreeRow::OnMouseButtonDown( const FGeometry& MyGeometry, const FPointerEvent& MouseEvent )
+{
+	auto ItemPtr = Item.Pin();
+	if (ItemPtr.IsValid() && ItemPtr->CanInteract())
 	{
-		auto ItemPtr = Item.Pin();
-		if (ItemPtr.IsValid() && ItemPtr->CanInteract())
+		if (MouseEvent.GetEffectingButton() == EKeys::LeftMouseButton)
 		{
-			if (MouseEvent.GetEffectingButton() == EKeys::LeftMouseButton)
-			{
-				FReply Reply = SMultiColumnTableRow<FTreeItemPtr>::OnMouseButtonDown( MyGeometry, MouseEvent );
-
-				if (SceneOutlinerWeak.Pin()->CanSupportDragAndDrop())
-				{
-					return Reply.DetectDrag( SharedThis(this) , EKeys::LeftMouseButton );
-				}
+			FReply Reply = SMultiColumnTableRow<FSceneOutlinerTreeItemPtr>::OnMouseButtonDown( MyGeometry, MouseEvent );
 
-				return Reply.PreventThrottling();
+			if (SceneOutlinerWeak.Pin()->CanSupportDragAndDrop())
+			{
+				return Reply.DetectDrag( SharedThis(this) , EKeys::LeftMouseButton );
 			}
-		}
 
-		return FReply::Handled();
+			return Reply.PreventThrottling();
+		}
 	}
 
-	FReply SSceneOutlinerTreeRow::OnMouseButtonUp(const FGeometry& MyGeometry, const FPointerEvent& MouseEvent)
-	{
-		auto ItemPtr = Item.Pin();
-		// We don't to change the selection when it is a left click since this was handle in the on mouse down
-		if (ItemPtr.IsValid() && ItemPtr->CanInteract())
-		{
-			return SMultiColumnTableRow<FTreeItemPtr>::OnMouseButtonUp(MyGeometry, MouseEvent);
-		}
+	return FReply::Handled();
+}
 
-		return FReply::Handled();
+FReply SSceneOutlinerTreeRow::OnMouseButtonUp(const FGeometry& MyGeometry, const FPointerEvent& MouseEvent)
+{
+	auto ItemPtr = Item.Pin();
+	// We don't to change the selection when it is a left click since this was handle in the on mouse down
+	if (ItemPtr.IsValid() && ItemPtr->CanInteract())
+	{
+		return SMultiColumnTableRow<FSceneOutlinerTreeItemPtr>::OnMouseButtonUp(MyGeometry, MouseEvent);
 	}
 
-	TSharedRef<SWidget> SSceneOutlinerTreeRow::GenerateWidgetForColumn( const FName& ColumnName )
-	{
-		auto ItemPtr = Item.Pin();
-		if (!ItemPtr.IsValid())
-		{
-			return SNullWidget::NullWidget;
-		}
+	return FReply::Handled();
+}
 
+TSharedRef<SWidget> SSceneOutlinerTreeRow::GenerateWidgetForColumn( const FName& ColumnName )
+{
+	auto ItemPtr = Item.Pin();
+	if (!ItemPtr.IsValid())
+	{
+		return SNullWidget::NullWidget;
+	}
 
-		auto Outliner = SceneOutlinerWeak.Pin();
-		check(Outliner.IsValid());
 
-		// Create the widget for this item
-		TSharedRef<SWidget> NewItemWidget = SNullWidget::NullWidget;
+	auto Outliner = SceneOutlinerWeak.Pin();
+	check(Outliner.IsValid());
 
-		auto Column = Outliner->GetColumns().FindRef(ColumnName);
-		if (Column.IsValid())
-		{
-			NewItemWidget = Column->ConstructRowWidget(ItemPtr.ToSharedRef(), *this);
-		}
+	// Create the widget for this item
+	TSharedRef<SWidget> NewItemWidget = SNullWidget::NullWidget;
 
-		if( ColumnName == FBuiltInColumnTypes::Label() )
-		{
-			// The first column gets the tree expansion arrow for this row
-			return
-				SNew( SHorizontalBox )
+	auto Column = Outliner->GetColumns().FindRef(ColumnName);
+	if (Column.IsValid())
+	{
+		NewItemWidget = Column->ConstructRowWidget(ItemPtr.ToSharedRef(), *this);
+	}
 
-				+SHorizontalBox::Slot()
-				.AutoWidth()
-				.Padding(6, 0, 0, 0)
-				[
-					SNew( SExpanderArrow, SharedThis(this) ).IndentAmount(12)
-				]
+	if( ColumnName == FSceneOutlinerBuiltInColumnTypes::Label() )
+	{
+		// The first column gets the tree expansion arrow for this row
+		return
+			SNew( SHorizontalBox )
 
 			+SHorizontalBox::Slot()
-				.FillWidth(1.0f)
-				[
-					NewItemWidget
-				];
-		}
-		else
-		{
-			// Other columns just get widget content -- no expansion arrow needed
-			return NewItemWidget;
-		}
+			.AutoWidth()
+			.Padding(6, 0, 0, 0)
+			[
+				SNew( SExpanderArrow, SharedThis(this) ).IndentAmount(12)
+			]
+
+		+SHorizontalBox::Slot()
+			.FillWidth(1.0f)
+			[
+				NewItemWidget
+			];
 	}
-
-	void SSceneOutlinerTreeRow::Construct( const FArguments& InArgs, const TSharedRef<SOutlinerTreeView>& OutlinerTreeView, TSharedRef<SSceneOutliner> SceneOutliner )
+	else
 	{
-		Item = InArgs._Item->AsShared();
-		SceneOutlinerWeak = SceneOutliner;
-		LastHighlightInteractionTime = 0.0;
+		// Other columns just get widget content -- no expansion arrow needed
+		return NewItemWidget;
+	}
+}
 
-		auto Args = FSuperRowType::FArguments()
-			.Style(&FEditorStyle::Get().GetWidgetStyle<FTableRowStyle>("SceneOutliner.TableViewRow"));
+void SSceneOutlinerTreeRow::Construct( const FArguments& InArgs, const TSharedRef<SSceneOutlinerTreeView>& OutlinerTreeView, TSharedRef<SSceneOutliner> SceneOutliner )
+{
+	Item = InArgs._Item->AsShared();
+	SceneOutlinerWeak = SceneOutliner;
+	LastHighlightInteractionTime = 0.0;
 
+	auto Args = FSuperRowType::FArguments()
+		.Style(&FEditorStyle::Get().GetWidgetStyle<FTableRowStyle>("SceneOutliner.TableViewRow"));
 
-		Args.OnDragDetected_Static(SceneOutliner::OnDragDetected, TWeakPtr<SOutlinerTreeView>(OutlinerTreeView));
 
-		SMultiColumnTableRow<FTreeItemPtr>::Construct(Args, OutlinerTreeView);
-	}
+	Args.OnDragDetected_Static(HandleOnDragDetected, TWeakPtr<SSceneOutlinerTreeView>(OutlinerTreeView));
 
-	const float SSceneOutlinerTreeRow::HighlightRectLeftOffset = 0.0f;
-	const float SSceneOutlinerTreeRow::HighlightRectRightOffset = 0.0f;
-	const float SSceneOutlinerTreeRow::HighlightTargetSpringConstant = 25.0f;
-	const float SSceneOutlinerTreeRow::HighlightTargetEffectDuration = 0.5f;
-	const float SSceneOutlinerTreeRow::HighlightTargetOpacity = 0.8f;
-	const float SSceneOutlinerTreeRow::LabelChangedAnimOffsetPercent = 0.2f;
+	SMultiColumnTableRow<FSceneOutlinerTreeItemPtr>::Construct(Args, OutlinerTreeView);
+}
 
-    void SSceneOutlinerTreeRow::FlashHighlight()
-    {
-        LastHighlightInteractionTime = FSlateApplication::Get().GetCurrentTime();
-    }
+const float SSceneOutlinerTreeRow::HighlightRectLeftOffset = 0.0f;
+const float SSceneOutlinerTreeRow::HighlightRectRightOffset = 0.0f;
+const float SSceneOutlinerTreeRow::HighlightTargetSpringConstant = 25.0f;
+const float SSceneOutlinerTreeRow::HighlightTargetEffectDuration = 0.5f;
+const float SSceneOutlinerTreeRow::HighlightTargetOpacity = 0.8f;
+const float SSceneOutlinerTreeRow::LabelChangedAnimOffsetPercent = 0.2f;
 
-	void SSceneOutlinerTreeRow::Tick( const FGeometry& AllottedGeometry, const double InCurrentTime, const float InDeltaTime )
-	{
-		// Call parent implementation.
-		SCompoundWidget::Tick( AllottedGeometry, InCurrentTime, InDeltaTime );
+void SSceneOutlinerTreeRow::FlashHighlight()
+{
+    LastHighlightInteractionTime = FSlateApplication::Get().GetCurrentTime();
+}
 
-		// We'll draw with the 'focused' look if we're either focused or we have a context menu summoned
-		const bool bShouldAppearFocused = HasKeyboardFocus();
+void SSceneOutlinerTreeRow::Tick( const FGeometry& AllottedGeometry, const double InCurrentTime, const float InDeltaTime )
+{
+	// Call parent implementation.
+	SCompoundWidget::Tick( AllottedGeometry, InCurrentTime, InDeltaTime );
 
-		// Update highlight 'target' effect
-		{
-			const float HighlightLeftX = HighlightRectLeftOffset;
-			const float HighlightRightX = HighlightRectRightOffset + AllottedGeometry.GetLocalSize().X;
+	// We'll draw with the 'focused' look if we're either focused or we have a context menu summoned
+	const bool bShouldAppearFocused = HasKeyboardFocus();
 
-			HighlightTargetLeftSpring.SetTarget( HighlightLeftX );
-			HighlightTargetRightSpring.SetTarget( HighlightRightX );
+	// Update highlight 'target' effect
+	{
+		const float HighlightLeftX = HighlightRectLeftOffset;
+		const float HighlightRightX = HighlightRectRightOffset + AllottedGeometry.GetLocalSize().X;
 
-			float TimeSinceHighlightInteraction = (float)( InCurrentTime - LastHighlightInteractionTime );
-			if( TimeSinceHighlightInteraction <= HighlightTargetEffectDuration || bShouldAppearFocused )
-			{
-				HighlightTargetLeftSpring.Tick( InDeltaTime );
-				HighlightTargetRightSpring.Tick( InDeltaTime );
-			}
+		HighlightTargetLeftSpring.SetTarget( HighlightLeftX );
+		HighlightTargetRightSpring.SetTarget( HighlightRightX );
+
+		float TimeSinceHighlightInteraction = (float)( InCurrentTime - LastHighlightInteractionTime );
+		if( TimeSinceHighlightInteraction <= HighlightTargetEffectDuration || bShouldAppearFocused )
+		{
+			HighlightTargetLeftSpring.Tick( InDeltaTime );
+			HighlightTargetRightSpring.Tick( InDeltaTime );
 		}
 	}
+}
 
-	int32 SSceneOutlinerTreeRow::OnPaint( const FPaintArgs& Args, const FGeometry& AllottedGeometry, const FSlateRect& MyCullingRect, FSlateWindowElementList& OutDrawElements, int32 LayerId, const FWidgetStyle& InWidgetStyle, bool bParentEnabled ) const
-	{
-		int32 StartLayer = SMultiColumnTableRow::OnPaint( Args, AllottedGeometry, MyCullingRect, OutDrawElements, LayerId, InWidgetStyle, bParentEnabled );
+int32 SSceneOutlinerTreeRow::OnPaint( const FPaintArgs& Args, const FGeometry& AllottedGeometry, const FSlateRect& MyCullingRect, FSlateWindowElementList& OutDrawElements, int32 LayerId, const FWidgetStyle& InWidgetStyle, bool bParentEnabled ) const
+{
+	int32 StartLayer = SMultiColumnTableRow::OnPaint( Args, AllottedGeometry, MyCullingRect, OutDrawElements, LayerId, InWidgetStyle, bParentEnabled );
 
-		const int32 TextLayer = 1;
+	const int32 TextLayer = 1;
 
-		// See if a disabled effect should be used
-		bool bEnabled = ShouldBeEnabled( bParentEnabled );
-		ESlateDrawEffect DrawEffects = (bEnabled) ? ESlateDrawEffect::None : ESlateDrawEffect::DisabledEffect;
+	// See if a disabled effect should be used
+	bool bEnabled = ShouldBeEnabled( bParentEnabled );
+	ESlateDrawEffect DrawEffects = (bEnabled) ? ESlateDrawEffect::None : ESlateDrawEffect::DisabledEffect;
 
-		const double CurrentTime = FSlateApplication::Get().GetCurrentTime();
+	const double CurrentTime = FSlateApplication::Get().GetCurrentTime();
 
-		// We'll draw with the 'focused' look if we're either focused or we have a context menu summoned
-		const bool bShouldAppearFocused = HasKeyboardFocus();
+	// We'll draw with the 'focused' look if we're either focused or we have a context menu summoned
+	const bool bShouldAppearFocused = HasKeyboardFocus();
 
-		// Draw highlight targeting effect
-		const float TimeSinceHighlightInteraction = (float)( CurrentTime - LastHighlightInteractionTime );
-		if( TimeSinceHighlightInteraction <= HighlightTargetEffectDuration )
-		{
-
-			// Compute animation progress
-			float EffectAlpha = FMath::Clamp( TimeSinceHighlightInteraction / HighlightTargetEffectDuration, 0.0f, 1.0f );
-			EffectAlpha = 1.0f - EffectAlpha * EffectAlpha;  // Inverse square falloff (looks nicer!)
+	// Draw highlight targeting effect
+	const float TimeSinceHighlightInteraction = (float)( CurrentTime - LastHighlightInteractionTime );
+	if( TimeSinceHighlightInteraction <= HighlightTargetEffectDuration )
+	{
 
-			// Apply extra opacity falloff when dehighlighting
-			float EffectOpacity = EffectAlpha;
+		// Compute animation progress
+		float EffectAlpha = FMath::Clamp( TimeSinceHighlightInteraction / HighlightTargetEffectDuration, 0.0f, 1.0f );
+		EffectAlpha = 1.0f - EffectAlpha * EffectAlpha;  // Inverse square falloff (looks nicer!)
 
-			// Figure out a universally visible highlight color.
-			FLinearColor HighlightTargetColorAndOpacity = ( (FLinearColor::White - ColorAndOpacity.Get())*0.5f + FLinearColor(+0.4f, +0.1f, -0.2f)) * InWidgetStyle.GetColorAndOpacityTint();
-			HighlightTargetColorAndOpacity.A = HighlightTargetOpacity * EffectOpacity * 255.0f;
+		// Apply extra opacity falloff when dehighlighting
+		float EffectOpacity = EffectAlpha;
 
-			// Compute the bounds offset of the highlight target from where the highlight target spring
-			// extents currently lie.  This is used to "grow" or "shrink" the highlight as needed.
-			const float LabelChangedAnimOffset = LabelChangedAnimOffsetPercent * AllottedGeometry.GetLocalSize().Y;
+		// Figure out a universally visible highlight color.
+		FLinearColor HighlightTargetColorAndOpacity = ( (FLinearColor::White - ColorAndOpacity.Get())*0.5f + FLinearColor(+0.4f, +0.1f, -0.2f)) * InWidgetStyle.GetColorAndOpacityTint();
+		HighlightTargetColorAndOpacity.A = HighlightTargetOpacity * EffectOpacity * 255.0f;
 
-			// Choose an offset amount depending on whether we're highlighting, or clearing highlight
-			const float EffectOffset = EffectAlpha * LabelChangedAnimOffset;
+		// Compute the bounds offset of the highlight target from where the highlight target spring
+		// extents currently lie.  This is used to "grow" or "shrink" the highlight as needed.
+		const float LabelChangedAnimOffset = LabelChangedAnimOffsetPercent * AllottedGeometry.GetLocalSize().Y;
 
-			const float HighlightLeftX = HighlightTargetLeftSpring.GetPosition() - EffectOffset;
-			const float HighlightRightX = HighlightTargetRightSpring.GetPosition() + EffectOffset;
-			const float HighlightTopY = 0.0f - LabelChangedAnimOffset;
-			const float HighlightBottomY = AllottedGeometry.GetLocalSize().Y + EffectOffset;
+		// Choose an offset amount depending on whether we're highlighting, or clearing highlight
+		const float EffectOffset = EffectAlpha * LabelChangedAnimOffset;
 
-			const FVector2D DrawPosition = FVector2D( HighlightLeftX, HighlightTopY );
-			const FVector2D DrawSize = FVector2D( HighlightRightX - HighlightLeftX, HighlightBottomY - HighlightTopY );
+		const float HighlightLeftX = HighlightTargetLeftSpring.GetPosition() - EffectOffset;
+		const float HighlightRightX = HighlightTargetRightSpring.GetPosition() + EffectOffset;
+		const float HighlightTopY = 0.0f - LabelChangedAnimOffset;
+		const float HighlightBottomY = AllottedGeometry.GetLocalSize().Y + EffectOffset;
 
-			const FSlateBrush* StyleInfo = FEditorStyle::GetBrush("SceneOutliner.ChangedItemHighlight");
+		const FVector2D DrawPosition = FVector2D( HighlightLeftX, HighlightTopY );
+		const FVector2D DrawSize = FVector2D( HighlightRightX - HighlightLeftX, HighlightBottomY - HighlightTopY );
 
-			// NOTE: We rely on scissor clipping for the highlight rectangle
-			FSlateDrawElement::MakeBox(
-				OutDrawElements,
-				LayerId + TextLayer,
-				AllottedGeometry.ToPaintGeometry( DrawPosition, DrawSize ),	// Position, Size, Scale
-				StyleInfo,													// Style
-				DrawEffects,												// Effects to use
-				HighlightTargetColorAndOpacity );							// Color
-		}
+		const FSlateBrush* StyleInfo = FEditorStyle::GetBrush("SceneOutliner.ChangedItemHighlight");
 
-		return LayerId + TextLayer;
+		// NOTE: We rely on scissor clipping for the highlight rectangle
+		FSlateDrawElement::MakeBox(
+			OutDrawElements,
+			LayerId + TextLayer,
+			AllottedGeometry.ToPaintGeometry( DrawPosition, DrawSize ),	// Position, Size, Scale
+			StyleInfo,													// Style
+			DrawEffects,												// Effects to use
+			HighlightTargetColorAndOpacity );							// Color
 	}
+
+	return LayerId + TextLayer;
 }
 
 #undef LOCTEXT_NAMESPACE
