Removed SceneOutliner namespace

diff --git a/Engine/Source/Editor/SceneOutliner/Private/SSceneOutliner.cpp b/Engine/Source/Editor/SceneOutliner/Private/SSceneOutliner.cpp
--- a/Engine/Source/Editor/SceneOutliner/Private/SSceneOutliner.cpp
+++ b/Engine/Source/Editor/SceneOutliner/Private/SSceneOutliner.cpp
@@ -38,7 +38,6 @@
 #include "Widgets/Notifications/SNotificationList.h"
 #include "ISceneOutlinerMode.h"
 #include "FolderTreeItem.h"
-#include "UObject/PackageReload.h"
 
 DEFINE_LOG_CATEGORY_STATIC(LogSceneOutliner, Log, All);
 
@@ -47,1892 +46,1884 @@ DEFINE_LOG_CATEGORY_STATIC(LogSceneOutliner, Log, All);
 // The amount of time that must pass before the Scene Outliner will attempt a sort when in PIE/SIE.
 #define SCENE_OUTLINER_RESORT_TIMER 1.0f
 
-namespace SceneOutliner
+void SSceneOutliner::Construct(const FArguments& InArgs, const FSceneOutlinerInitializationOptions& InInitOptions)
 {
-	void SSceneOutliner::Construct(const FArguments& InArgs, const FInitializationOptions& InInitOptions)
-	{
-		// Copy over the shared data from the initialization options
-		static_cast<FSharedOutlinerData&>(*SharedData) = static_cast<const FSharedOutlinerData&>(InInitOptions);
+	// Copy over the shared data from the initialization options
+	static_cast<FSharedSceneOutlinerData&>(*SharedData) = static_cast<const FSharedSceneOutlinerData&>(InInitOptions);
 
-		// We use the filter collection provided, otherwise we create our own
-		Filters = InInitOptions.Filters.IsValid() ? InInitOptions.Filters : MakeShareable(new FOutlinerFilters);
+	// We use the filter collection provided, otherwise we create our own
+	Filters = InInitOptions.Filters.IsValid() ? InInitOptions.Filters : MakeShareable(new FSceneOutlinerFilters);
 
-		check(InInitOptions.ModeFactory.IsBound());
-		Mode = InInitOptions.ModeFactory.Execute(this);
-		check(Mode);
+	check(InInitOptions.ModeFactory.IsBound());
+	Mode = InInitOptions.ModeFactory.Execute(this);
+	check(Mode);
 
-		bFullRefresh = true;
-		bNeedsRefresh = true;
-		bNeedsColumRefresh = true;
-		bIsReentrant = false;
-		bSortDirty = true;
-		bSelectionDirty = true;
-		SortOutlinerTimer = 0.0f;
-		bPendingFocusNextFrame = InInitOptions.bFocusSearchBoxWhenOpened;
+	bFullRefresh = true;
+	bNeedsRefresh = true;
+	bNeedsColumRefresh = true;
+	bIsReentrant = false;
+	bSortDirty = true;
+	bSelectionDirty = true;
+	SortOutlinerTimer = 0.0f;
+	bPendingFocusNextFrame = InInitOptions.bFocusSearchBoxWhenOpened;
 
-		SortByColumn = FBuiltInColumnTypes::Label();
-		SortMode = EColumnSortMode::Ascending;
+	SortByColumn = FSceneOutlinerBuiltInColumnTypes::Label();
+	SortMode = EColumnSortMode::Ascending;
 
-		// @todo outliner: Should probably save this in layout!
-		// @todo outliner: Should save spacing for list view in layout
+	// @todo outliner: Should probably save this in layout!
+	// @todo outliner: Should save spacing for list view in layout
 
-		//Setup the SearchBox filter
-		{
-			auto Delegate = TreeItemTextFilter::FItemToStringArray::CreateSP( this, &SSceneOutliner::PopulateSearchStrings );
-			SearchBoxFilter = MakeShareable( new TreeItemTextFilter( Delegate ) );
-		}
+	//Setup the SearchBox filter
+	{
+		auto Delegate = SceneOutliner::TreeItemTextFilter::FItemToStringArray::CreateSP( this, &SSceneOutliner::PopulateSearchStrings );
+		SearchBoxFilter = MakeShareable( new SceneOutliner::TreeItemTextFilter( Delegate ) );
+	}
 
-		TSharedRef<SVerticalBox> VerticalBox = SNew(SVerticalBox);
+	TSharedRef<SVerticalBox> VerticalBox = SNew(SVerticalBox);
 
-		for (auto& ModeFilterInfo : Mode->GetFilterInfos())
-		{
-			ModeFilterInfo.Value.InitFilter(Filters);
-		}
+	for (auto& ModeFilterInfo : Mode->GetFilterInfos())
+	{
+		ModeFilterInfo.Value.InitFilter(Filters);
+	}
 
-		SearchBoxFilter->OnChanged().AddSP( this, &SSceneOutliner::FullRefresh );
-		Filters->OnChanged().AddSP(this, &SSceneOutliner::FullRefresh);
+	SearchBoxFilter->OnChanged().AddSP( this, &SSceneOutliner::FullRefresh );
+	Filters->OnChanged().AddSP(this, &SSceneOutliner::FullRefresh);
 
-		HeaderRowWidget =
-			SNew( SHeaderRow )
-				// Only show the list header if the user configured the outliner for that
-				.Visibility(InInitOptions.bShowHeaderRow ? EVisibility::Visible : EVisibility::Collapsed);
+	HeaderRowWidget =
+		SNew( SHeaderRow )
+			// Only show the list header if the user configured the outliner for that
+			.Visibility(InInitOptions.bShowHeaderRow ? EVisibility::Visible : EVisibility::Collapsed);
 
-		SetupColumns(*HeaderRowWidget);
+	SetupColumns(*HeaderRowWidget);
 
-		ChildSlot
-		[
-			VerticalBox
-		];
+	ChildSlot
+	[
+		VerticalBox
+	];
 
-		auto Toolbar = SNew(SHorizontalBox);
+	auto Toolbar = SNew(SHorizontalBox);
 
-		Toolbar->AddSlot()
-		.VAlign(VAlign_Center)
-		[
-			SAssignNew( FilterTextBoxWidget, SSearchBox )
-			.Visibility( InInitOptions.bShowSearchBox ? EVisibility::Visible : EVisibility::Collapsed )
-			.HintText( LOCTEXT( "FilterSearch", "Search..." ) )
-			.ToolTipText( LOCTEXT("FilterSearchHint", "Type here to search (pressing enter selects the results)") )
-			.OnTextChanged( this, &SSceneOutliner::OnFilterTextChanged )
-			.OnTextCommitted( this, &SSceneOutliner::OnFilterTextCommitted )
-		];
+	Toolbar->AddSlot()
+	.VAlign(VAlign_Center)
+	[
+		SAssignNew( FilterTextBoxWidget, SSearchBox )
+		.Visibility( InInitOptions.bShowSearchBox ? EVisibility::Visible : EVisibility::Collapsed )
+		.HintText( LOCTEXT( "FilterSearch", "Search..." ) )
+		.ToolTipText( LOCTEXT("FilterSearchHint", "Type here to search (pressing enter selects the results)") )
+		.OnTextChanged( this, &SSceneOutliner::OnFilterTextChanged )
+		.OnTextCommitted( this, &SSceneOutliner::OnFilterTextCommitted )
+	];
 
 		
-		if (Mode->SupportsCreateNewFolder() && InInitOptions.bShowCreateNewFolder)
-		{
-			Toolbar->AddSlot()
-				.VAlign(VAlign_Center)
-				.AutoWidth()
-				.Padding(4.f, 0.f, 0.f, 0.f)
-				[
-					SNew(SButton)
-					.ButtonStyle(FAppStyle::Get(), "SimpleButton")
-					.ToolTipText(LOCTEXT("CreateFolderToolTip", "Create a new folder containing the current selection"))
-					.OnClicked(this, &SSceneOutliner::OnCreateFolderClicked)
-					[
-						SNew(SImage)
-	                	.ColorAndOpacity(FSlateColor::UseForeground())
-						.Image(FAppStyle::Get().GetBrush("SceneOutliner.NewFolderIcon"))
-					]
-				];
-		}
-
-
-		if (Mode->ShowViewButton())
-		{
-			// View mode combo button
-			Toolbar->AddSlot()
+	if (Mode->SupportsCreateNewFolder() && InInitOptions.bShowCreateNewFolder)
+	{
+		Toolbar->AddSlot()
 			.VAlign(VAlign_Center)
 			.AutoWidth()
+			.Padding(4.f, 0.f, 0.f, 0.f)
 			[
-				SAssignNew( ViewOptionsComboButton, SComboButton )
-				.ComboButtonStyle( FAppStyle::Get(), "SimpleComboButton" ) // Use the tool bar item style for this button
-				.OnGetMenuContent( this, &SSceneOutliner::GetViewButtonContent, Mode->ShowFilterOptions())
-				.HasDownArrow(false)
-				.ButtonContent()
+				SNew(SButton)
+				.ButtonStyle(FAppStyle::Get(), "SimpleButton")
+				.ToolTipText(LOCTEXT("CreateFolderToolTip", "Create a new folder containing the current selection"))
+				.OnClicked(this, &SSceneOutliner::OnCreateFolderClicked)
 				[
 					SNew(SImage)
 	                .ColorAndOpacity(FSlateColor::UseForeground())
-					.Image( FAppStyle::Get().GetBrush("Icons.Settings") )
+					.Image(FAppStyle::Get().GetBrush("SceneOutliner.NewFolderIcon"))
 				]
 			];
-		}
+	}
 
 
-		VerticalBox->AddSlot()
-		.AutoHeight()
-		.Padding( 8.0f, 4.0f, 4.0f, 2.0f )
+	if (Mode->ShowViewButton())
+	{
+		// View mode combo button
+		Toolbar->AddSlot()
+		.VAlign(VAlign_Center)
+		.AutoWidth()
 		[
-			Toolbar
+			SAssignNew( ViewOptionsComboButton, SComboButton )
+			.ComboButtonStyle( FAppStyle::Get(), "SimpleComboButton" ) // Use the tool bar item style for this button
+			.OnGetMenuContent( this, &SSceneOutliner::GetViewButtonContent, Mode->ShowFilterOptions())
+			.HasDownArrow(false)
+			.ButtonContent()
+			[
+				SNew(SImage)
+	            .ColorAndOpacity(FSlateColor::UseForeground())
+				.Image( FAppStyle::Get().GetBrush("Icons.Settings") )
+			]
 		];
+	}
 
-		VerticalBox->AddSlot()
-		.FillHeight( 1.0f )
+
+	VerticalBox->AddSlot()
+	.AutoHeight()
+	.Padding( 8.0f, 4.0f, 4.0f, 2.0f )
+	[
+		Toolbar
+	];
+
+	VerticalBox->AddSlot()
+	.FillHeight( 1.0f )
+	[
+		SNew( SOverlay )
+		+SOverlay::Slot()
+		.HAlign( HAlign_Center )
 		[
-			SNew( SOverlay )
-			+SOverlay::Slot()
-			.HAlign( HAlign_Center )
-			[
-				SNew( STextBlock )
-				.Visibility( this, &SSceneOutliner::GetEmptyLabelVisibility )
-				.Text( LOCTEXT( "EmptyLabel", "Empty" ) )
-				.ColorAndOpacity( FLinearColor( 0.4f, 1.0f, 0.4f ) )
-			]
+			SNew( STextBlock )
+			.Visibility( this, &SSceneOutliner::GetEmptyLabelVisibility )
+			.Text( LOCTEXT( "EmptyLabel", "Empty" ) )
+			.ColorAndOpacity( FLinearColor( 0.4f, 1.0f, 0.4f ) )
+		]
 
-			+SOverlay::Slot()
-			[
-				SNew(SBorder).BorderImage( FAppStyle::Get().GetBrush("InputBrush") )
-			]
+		+SOverlay::Slot()
+		[
+			SNew(SBorder).BorderImage( FAppStyle::Get().GetBrush("InputBrush") )
+		]
 
-			+SOverlay::Slot()
-			[
-				SAssignNew( OutlinerTreeView, SOutlinerTreeView, StaticCastSharedRef<SSceneOutliner>(AsShared()) )
+		+SOverlay::Slot()
+		[
+			SAssignNew( OutlinerTreeView, SSceneOutlinerTreeView, StaticCastSharedRef<SSceneOutliner>(AsShared()) )
 
-				// Determined by the mode
-				.SelectionMode( this, &SSceneOutliner::GetSelectionMode )
+			// Determined by the mode
+			.SelectionMode( this, &SSceneOutliner::GetSelectionMode )
 
-				// Point the tree to our array of root-level items.  Whenever this changes, we'll call RequestTreeRefresh()
-				.TreeItemsSource( &RootTreeItems )
+			// Point the tree to our array of root-level items.  Whenever this changes, we'll call RequestTreeRefresh()
+			.TreeItemsSource( &RootTreeItems )
 
-				// Find out when the user selects something in the tree
-				.OnSelectionChanged( this, &SSceneOutliner::OnOutlinerTreeSelectionChanged )
+			// Find out when the user selects something in the tree
+			.OnSelectionChanged( this, &SSceneOutliner::OnOutlinerTreeSelectionChanged )
 
-				// Called when the user double-clicks with LMB on an item in the list
-				.OnMouseButtonDoubleClick( this, &SSceneOutliner::OnOutlinerTreeDoubleClick )
+			// Called when the user double-clicks with LMB on an item in the list
+			.OnMouseButtonDoubleClick( this, &SSceneOutliner::OnOutlinerTreeDoubleClick )
 
-				// Called when an item is scrolled into view
-				.OnItemScrolledIntoView( this, &SSceneOutliner::OnOutlinerTreeItemScrolledIntoView )
+			// Called when an item is scrolled into view
+			.OnItemScrolledIntoView( this, &SSceneOutliner::OnOutlinerTreeItemScrolledIntoView )
 
-				// Called when an item is expanded or collapsed
-				.OnExpansionChanged(this, &SSceneOutliner::OnItemExpansionChanged)
+			// Called when an item is expanded or collapsed
+			.OnExpansionChanged(this, &SSceneOutliner::OnItemExpansionChanged)
 
-				// Called to child items for any given parent item
-				.OnGetChildren( this, &SSceneOutliner::OnGetChildrenForOutlinerTree )
+			// Called to child items for any given parent item
+			.OnGetChildren( this, &SSceneOutliner::OnGetChildrenForOutlinerTree )
 
-				// Generates the actual widget for a tree item
-				.OnGenerateRow( this, &SSceneOutliner::OnGenerateRowForOutlinerTree ) 
+			// Generates the actual widget for a tree item
+			.OnGenerateRow( this, &SSceneOutliner::OnGenerateRowForOutlinerTree ) 
 
-				// Use the level viewport context menu as the right click menu for tree items
-				.OnContextMenuOpening(this, &SSceneOutliner::OnOpenContextMenu)
+			// Use the level viewport context menu as the right click menu for tree items
+			.OnContextMenuOpening(this, &SSceneOutliner::OnOpenContextMenu)
 
-				// Header for the tree
-				.HeaderRow( HeaderRowWidget )
+			// Header for the tree
+			.HeaderRow( HeaderRowWidget )
 
-				// Called when an item is expanded or collapsed with the shift-key pressed down
-				.OnSetExpansionRecursive(this, &SSceneOutliner::SetItemExpansionRecursive)
+			// Called when an item is expanded or collapsed with the shift-key pressed down
+			.OnSetExpansionRecursive(this, &SSceneOutliner::SetItemExpansionRecursive)
 
-				// Make it easier to see hierarchies when there are a lot of items
-				.HighlightParentNodesForSelection(true)
-			]
-		];
+			// Make it easier to see hierarchies when there are a lot of items
+			.HighlightParentNodesForSelection(true)
+		]
+	];
 
 
-		// Bottom panel status bar, if enabled by the mode
-		if (Mode->ShowStatusBar())
-		{
-			VerticalBox->AddSlot()
-			.AutoHeight()
+	// Bottom panel status bar, if enabled by the mode
+	if (Mode->ShowStatusBar())
+	{
+		VerticalBox->AddSlot()
+		.AutoHeight()
+		[
+			SNew(SHorizontalBox)
+
+			// Asset count
+			+SHorizontalBox::Slot()
+			.FillWidth(1.f)
+			.VAlign(VAlign_Center)
+			.Padding(14, 9)
 			[
-				SNew(SHorizontalBox)
+				SNew( STextBlock )
+				.Text(this, &SSceneOutliner::GetFilterStatusText)
+				.ColorAndOpacity(this, &SSceneOutliner::GetFilterStatusTextColor)
+			]
+		];
+	}
 
-				// Asset count
-				+SHorizontalBox::Slot()
-				.FillWidth(1.f)
-				.VAlign(VAlign_Center)
-				.Padding(14, 9)
-				[
-					SNew( STextBlock )
-					.Text(this, &SSceneOutliner::GetFilterStatusText)
-					.ColorAndOpacity(this, &SSceneOutliner::GetFilterStatusTextColor)
-				]
-			];
-		}
+	// Don't allow tool-tips over the header
+	HeaderRowWidget->EnableToolTipForceField( true );
 
-		// Don't allow tool-tips over the header
-		HeaderRowWidget->EnableToolTipForceField( true );
+	// Populate our data set
+	Populate();
 
-		// Populate our data set
-		Populate();
+	// Register to update when an undo/redo operation has been called to update our list of items
+	GEditor->RegisterForUndo( this );
 
-		// Register to update when an undo/redo operation has been called to update our list of items
-		GEditor->RegisterForUndo( this );
+	// Register to be notified when properties are edited
+	FCoreUObjectDelegates::OnPackageReloaded.AddRaw(this, &SSceneOutliner::OnAssetReloaded);
+}
 
-		// Register to be notified when properties are edited
-		FCoreUObjectDelegates::OnPackageReloaded.AddRaw(this, &SSceneOutliner::OnAssetReloaded);
-	}
+void SSceneOutliner::SetupColumns(SHeaderRow& HeaderRow)
+{
+	FSceneOutlinerModule& SceneOutlinerModule = FModuleManager::LoadModuleChecked<FSceneOutlinerModule>("SceneOutliner");
 
-	void SSceneOutliner::SetupColumns(SHeaderRow& HeaderRow)
+	if (SharedData->ColumnMap.Num() == 0)
 	{
-		FSceneOutlinerModule& SceneOutlinerModule = FModuleManager::LoadModuleChecked<FSceneOutlinerModule>("SceneOutliner");
+		SharedData->UseDefaultColumns();
+	}
+
+	Columns.Empty(SharedData->ColumnMap.Num());
+	HeaderRow.ClearColumns();
+
+	// Get a list of sorted columns IDs to create
+	TArray<FName> SortedIDs;
+	SortedIDs.Reserve(SharedData->ColumnMap.Num());
+	SharedData->ColumnMap.GenerateKeyArray(SortedIDs);
 
-		if (SharedData->ColumnMap.Num() == 0)
+	SortedIDs.Sort([&](const FName& A, const FName& B){
+		return SharedData->ColumnMap[A].PriorityIndex < SharedData->ColumnMap[B].PriorityIndex;
+	});
+
+	for (const FName& ID : SortedIDs)
+	{
+		if (SharedData->ColumnMap[ID].Visibility == ESceneOutlinerColumnVisibility::Invisible)
 		{
-			SharedData->UseDefaultColumns();
+			continue;
 		}
 
-		Columns.Empty(SharedData->ColumnMap.Num());
-		HeaderRow.ClearColumns();
+		TSharedPtr<ISceneOutlinerColumn> Column;
 
-		// Get a list of sorted columns IDs to create
-		TArray<FName> SortedIDs;
-		SortedIDs.Reserve(SharedData->ColumnMap.Num());
-		SharedData->ColumnMap.GenerateKeyArray(SortedIDs);
-
-		SortedIDs.Sort([&](const FName& A, const FName& B){
-			return SharedData->ColumnMap[A].PriorityIndex < SharedData->ColumnMap[B].PriorityIndex;
-		});
+		if (SharedData->ColumnMap[ID].Factory.IsBound())
+		{
+			Column = SharedData->ColumnMap[ID].Factory.Execute(*this);
+		}
+		else
+		{
+			Column = SceneOutlinerModule.FactoryColumn(ID, *this);
+		}
 
-		for (const FName& ID : SortedIDs)
+		if (ensure(Column.IsValid()))
 		{
-			if (SharedData->ColumnMap[ID].Visibility == EColumnVisibility::Invisible)
-			{
-				continue;
-			}
+			check(Column->GetColumnID() == ID);
+			Columns.Add(Column->GetColumnID(), Column);
 
-			TSharedPtr<ISceneOutlinerColumn> Column;
+			auto ColumnArgs = Column->ConstructHeaderRowColumn();
 
-			if (SharedData->ColumnMap[ID].Factory.IsBound())
+			if (Column->SupportsSorting())
 			{
-				Column = SharedData->ColumnMap[ID].Factory.Execute(*this);
+				ColumnArgs
+					.SortMode(this, &SSceneOutliner::GetColumnSortMode, Column->GetColumnID())
+					.OnSort(this, &SSceneOutliner::OnColumnSortModeChanged);
 			}
-			else
-			{
-				Column = SceneOutlinerModule.FactoryColumn(ID, *this);
-			}
-
-			if (ensure(Column.IsValid()))
-			{
-				check(Column->GetColumnID() == ID);
-				Columns.Add(Column->GetColumnID(), Column);
 
-				auto ColumnArgs = Column->ConstructHeaderRowColumn();
+			HeaderRow.AddColumn(ColumnArgs);
+		}
+	}
 
-				if (Column->SupportsSorting())
-				{
-					ColumnArgs
-						.SortMode(this, &SSceneOutliner::GetColumnSortMode, Column->GetColumnID())
-						.OnSort(this, &SSceneOutliner::OnColumnSortModeChanged);
-				}
+	Columns.Shrink();
+	bNeedsColumRefresh = false;
+}
 
-				HeaderRow.AddColumn(ColumnArgs);
-			}
-		}
+void SSceneOutliner::RefreshColums()
+{
+	bNeedsColumRefresh = true;
+}
 
-		Columns.Shrink();
-		bNeedsColumRefresh = false;
-	}
+SSceneOutliner::~SSceneOutliner()
+{
+	Mode->GetHierarchy()->OnHierarchyChanged().RemoveAll(this);
+	delete Mode;
 
-	void SSceneOutliner::RefreshColums()
+	if(GEngine)
 	{
-		bNeedsColumRefresh = true;
+		GEditor->UnregisterForUndo(this);
 	}
 
-	SSceneOutliner::~SSceneOutliner()
-	{
-		Mode->GetHierarchy()->OnHierarchyChanged().RemoveAll(this);
-		delete Mode;
-
-		if(GEngine)
-		{
-			GEditor->UnregisterForUndo(this);
-		}
+	SearchBoxFilter->OnChanged().RemoveAll( this );
+	Filters->OnChanged().RemoveAll(this);
 
-		SearchBoxFilter->OnChanged().RemoveAll( this );
-		Filters->OnChanged().RemoveAll(this);
-
-		FCoreUObjectDelegates::OnPackageReloaded.RemoveAll(this);
-	}
+	FCoreUObjectDelegates::OnPackageReloaded.RemoveAll(this);
+}
 
-	void SSceneOutliner::OnItemAdded(const FTreeItemID& ItemID, uint8 ActionMask)
-	{
-		NewItemActions.Add(ItemID, ActionMask);
-	}
+void SSceneOutliner::OnItemAdded(const FSceneOutlinerTreeItemID& ItemID, uint8 ActionMask)
+{
+	NewItemActions.Add(ItemID, ActionMask);
+}
 
-	FSlateColor SSceneOutliner::GetViewButtonForegroundColor() const
-	{
-		static const FName InvertedForegroundName("InvertedForeground");
-		static const FName DefaultForegroundName("DefaultForeground");
+FSlateColor SSceneOutliner::GetViewButtonForegroundColor() const
+{
+	static const FName InvertedForegroundName("InvertedForeground");
+	static const FName DefaultForegroundName("DefaultForeground");
 
-		return ViewOptionsComboButton->IsHovered() ? FEditorStyle::GetSlateColor(InvertedForegroundName) : FEditorStyle::GetSlateColor(DefaultForegroundName);
-	}
+	return ViewOptionsComboButton->IsHovered() ? FEditorStyle::GetSlateColor(InvertedForegroundName) : FEditorStyle::GetSlateColor(DefaultForegroundName);
+}
 	 
-	TSharedRef<SWidget> SSceneOutliner::GetViewButtonContent(bool bShowFilters)
-	{
-		// Menu should stay open on selection if filters are not being shown
-		FMenuBuilder MenuBuilder(bShowFilters, NULL);
+TSharedRef<SWidget> SSceneOutliner::GetViewButtonContent(bool bShowFilters)
+{
+	// Menu should stay open on selection if filters are not being shown
+	FMenuBuilder MenuBuilder(bShowFilters, NULL);
 
-		if (bShowFilters)
+	if (bShowFilters)
+	{
+		MenuBuilder.BeginSection("AssetThumbnails", LOCTEXT("ShowHeading", "Show"));
 		{
-			MenuBuilder.BeginSection("AssetThumbnails", LOCTEXT("ShowHeading", "Show"));
+			// Add mode filters
+			for (auto& ModeFilterInfo : Mode->GetFilterInfos())
 			{
-				// Add mode filters
-				for (auto& ModeFilterInfo : Mode->GetFilterInfos())
-				{
-					ModeFilterInfo.Value.AddMenu(MenuBuilder);
-				}
+				ModeFilterInfo.Value.AddMenu(MenuBuilder);
 			}
-			MenuBuilder.EndSection();
 		}
-		Mode->CreateViewContent(MenuBuilder);
-
-		return MenuBuilder.MakeWidget();
+		MenuBuilder.EndSection();
 	}
+	Mode->CreateViewContent(MenuBuilder);
 
-	ESelectionMode::Type SSceneOutliner::GetSelectionMode() const
-	{
-		return Mode->GetSelectionMode();
-	}
+	return MenuBuilder.MakeWidget();
+}
 
-	void SSceneOutliner::Refresh()
-	{
-		bNeedsRefresh = true;
-	}
+ESelectionMode::Type SSceneOutliner::GetSelectionMode() const
+{
+	return Mode->GetSelectionMode();
+}
 
-	void SSceneOutliner::FullRefresh()
-	{
-		UE_LOG(LogSceneOutliner, Verbose, TEXT("Full Refresh"));
-		bFullRefresh = true;
-		Refresh();
-	}
+void SSceneOutliner::Refresh()
+{
+	bNeedsRefresh = true;
+}
 
-	void SSceneOutliner::RefreshSelection()
-	{
-		bSelectionDirty = true;
-	}
+void SSceneOutliner::FullRefresh()
+{
+	UE_LOG(LogSceneOutliner, Verbose, TEXT("Full Refresh"));
+	bFullRefresh = true;
+	Refresh();
+}
 
-	void SSceneOutliner::Populate()
-	{
-		// Block events while we clear out the list
- 		TGuardValue<bool> ReentrantGuard(bIsReentrant, true);
+void SSceneOutliner::RefreshSelection()
+{
+	bSelectionDirty = true;
+}
 
-		// Get a collection of items and folders which were formerly collapsed
-		const FParentsExpansionState ExpansionStateInfo = GetParentsExpansionState();
+void SSceneOutliner::Populate()
+{
+	// Block events while we clear out the list
+ 	TGuardValue<bool> ReentrantGuard(bIsReentrant, true);
 
-		bool bMadeAnySignificantChanges = false;
-		if (bFullRefresh)
-		{
-			// Clear the selection here - RepopulateEntireTree will reconstruct it.
-			OutlinerTreeView->ClearSelection();
+	// Get a collection of items and folders which were formerly collapsed
+	const FParentsExpansionState ExpansionStateInfo = GetParentsExpansionState();
 
-			RepopulateEntireTree();
+	bool bMadeAnySignificantChanges = false;
+	if (bFullRefresh)
+	{
+		// Clear the selection here - RepopulateEntireTree will reconstruct it.
+		OutlinerTreeView->ClearSelection();
 
-			bMadeAnySignificantChanges = true;
-			bFullRefresh = false;
-		}
+		RepopulateEntireTree();
 
-		// Only deal with 500 at a time
-		const int32 End = FMath::Min(PendingOperations.Num(), 500);
-		for (int32 Index = 0; Index < End; ++Index)
+		bMadeAnySignificantChanges = true;
+		bFullRefresh = false;
+	}
+
+	// Only deal with 500 at a time
+	const int32 End = FMath::Min(PendingOperations.Num(), 500);
+	for (int32 Index = 0; Index < End; ++Index)
+	{
+		auto& PendingOp = PendingOperations[Index];
+		switch (PendingOp.Type)
 		{
-			auto& PendingOp = PendingOperations[Index];
-			switch (PendingOp.Type)
-			{
-			case FPendingTreeOperation::Added:
-				bMadeAnySignificantChanges = AddItemToTree(PendingOp.Item) || bMadeAnySignificantChanges;
-				break;
+		case SceneOutliner::FPendingTreeOperation::Added:
+			bMadeAnySignificantChanges = AddItemToTree(PendingOp.Item) || bMadeAnySignificantChanges;
+			break;
 
-			case FPendingTreeOperation::Moved:
-				bMadeAnySignificantChanges = true;
-				OnItemMoved(PendingOp.Item);
-				break;
+		case SceneOutliner::FPendingTreeOperation::Moved:
+			bMadeAnySignificantChanges = true;
+			OnItemMoved(PendingOp.Item);
+			break;
 
-			case FPendingTreeOperation::Removed:
-				bMadeAnySignificantChanges = true;
-				RemoveItemFromTree(PendingOp.Item);
-				break;
+		case SceneOutliner::FPendingTreeOperation::Removed:
+			bMadeAnySignificantChanges = true;
+			RemoveItemFromTree(PendingOp.Item);
+			break;
 
-			default:
-				check(false);
-				break;
-			}
+		default:
+			check(false);
+			break;
 		}
+	}
 
-		PendingOperations.RemoveAt(0, End);
-		SetParentsExpansionState(ExpansionStateInfo);
-
+	PendingOperations.RemoveAt(0, End);
+	SetParentsExpansionState(ExpansionStateInfo);
 
-		for (FName Folder : PendingFoldersSelect)
-		{
-			if (FTreeItemPtr* Item = TreeItemMap.Find(Folder))
-			{
-				OutlinerTreeView->SetItemSelection(*Item, true);
-			}
-		}
-		PendingFoldersSelect.Empty();
 
-		// Check if we need to sort because we are finished with the populating operations
-		bool bFinalSort = false;
-		if (PendingOperations.Num() == 0)
+	for (FName Folder : PendingFoldersSelect)
+	{
+		if (FSceneOutlinerTreeItemPtr* Item = TreeItemMap.Find(Folder))
 		{
-			// We're fully refreshed now.
-			NewItemActions.Empty();
-			bNeedsRefresh = false;
-			if (bDisableIntermediateSorting)
-			{
-				bDisableIntermediateSorting = false;
-				bFinalSort = true;
-			}
+			OutlinerTreeView->SetItemSelection(*Item, true);
 		}
+	}
+	PendingFoldersSelect.Empty();
 
-		// If we are allowing intermediate sorts and met the conditions, or this is the final sort after all ops are complete
-		if ((bMadeAnySignificantChanges && !bDisableIntermediateSorting) || bFinalSort)
+	// Check if we need to sort because we are finished with the populating operations
+	bool bFinalSort = false;
+	if (PendingOperations.Num() == 0)
+	{
+		// We're fully refreshed now.
+		NewItemActions.Empty();
+		bNeedsRefresh = false;
+		if (bDisableIntermediateSorting)
 		{
-			RequestSort();
+			bDisableIntermediateSorting = false;
+			bFinalSort = true;
 		}
 	}
 
-	bool SSceneOutliner::ShouldShowFolders() const
+	// If we are allowing intermediate sorts and met the conditions, or this is the final sort after all ops are complete
+	if ((bMadeAnySignificantChanges && !bDisableIntermediateSorting) || bFinalSort)
 	{
-		return Mode->ShouldShowFolders();
+		RequestSort();
 	}
+}
 
-	void SSceneOutliner::EmptyTreeItems()
-	{
-		PendingOperations.Empty();
-		TreeItemMap.Reset();
-		PendingTreeItemMap.Empty();
+bool SSceneOutliner::ShouldShowFolders() const
+{
+	return Mode->ShouldShowFolders();
+}
+
+void SSceneOutliner::EmptyTreeItems()
+{
+	PendingOperations.Empty();
+	TreeItemMap.Reset();
+	PendingTreeItemMap.Empty();
 
-		FolderCount = 0;
+	FolderCount = 0;
 
-		RootTreeItems.Empty();
-	}
+	RootTreeItems.Empty();
+}
 
-	void SSceneOutliner::AddPendingItem(FTreeItemPtr Item)
-	{
-		PendingTreeItemMap.Add(Item->GetID(), Item);
-		PendingOperations.Emplace(FPendingTreeOperation::Added, Item.ToSharedRef());
-	}
+void SSceneOutliner::AddPendingItem(FSceneOutlinerTreeItemPtr Item)
+{
+	PendingTreeItemMap.Add(Item->GetID(), Item);
+	PendingOperations.Emplace(SceneOutliner::FPendingTreeOperation::Added, Item.ToSharedRef());
+}
 
-	void SSceneOutliner::AddPendingItemAndChildren(FTreeItemPtr Item)
+void SSceneOutliner::AddPendingItemAndChildren(FSceneOutlinerTreeItemPtr Item)
+{
+	if (Item)
 	{
-		if (Item)
-		{
-			AddPendingItem(Item);
+		AddPendingItem(Item);
 
-			TArray<FTreeItemPtr> Children;
-			Mode->GetHierarchy()->CreateChildren(Item, Children);
-			for (auto& Child : Children)
-			{
-				AddPendingItem(Child);
-			}
-
-			Refresh();
+		TArray<FSceneOutlinerTreeItemPtr> Children;
+		Mode->GetHierarchy()->CreateChildren(Item, Children);
+		for (auto& Child : Children)
+		{
+			AddPendingItem(Child);
 		}
+
+		Refresh();
 	}
+}
+
+void SSceneOutliner::RepopulateEntireTree()
+{
+	EmptyTreeItems();
+
+	// Rebuild the hierarchy
+	Mode->Rebuild();
+	Mode->GetHierarchy()->OnHierarchyChanged().AddSP(this, &SSceneOutliner::OnHierarchyChangedEvent);
+
+	// Create all the items which match the filters, parent-child relationships are handled when each item is actually added to the tree
+
+	TArray<FSceneOutlinerTreeItemPtr> Items;
+	Mode->GetHierarchy()->CreateItems(Items);
 
-	void SSceneOutliner::RepopulateEntireTree()
+	for (FSceneOutlinerTreeItemPtr& Item : Items)
 	{
-		EmptyTreeItems();
+		AddPendingItem(Item);
+	}
 
-		// Rebuild the hierarchy
-		Mode->Rebuild();
-		Mode->GetHierarchy()->OnHierarchyChanged().AddSP(this, &SSceneOutliner::OnHierarchyChangedEvent);
+	Refresh();
+}
 
-		// Create all the items which match the filters, parent-child relationships are handled when each item is actually added to the tree
+void SSceneOutliner::OnChildRemovedFromParent(ISceneOutlinerTreeItem& Parent)
+{
+	if (Parent.Flags.bIsFilteredOut && !Parent.GetChildren().Num())
+	{
+		// The parent no longer has any children that match the current search terms. Remove it.
+		RemoveItemFromTree(Parent.AsShared());
+	}
+}
 
-		TArray<FTreeItemPtr> Items;
-		Mode->GetHierarchy()->CreateItems(Items);
+void SSceneOutliner::OnItemMoved(const FSceneOutlinerTreeItemRef& Item)
+{
+	// Just remove the item if it no longer matches the filters
+	if (!Item->Flags.bIsFilteredOut && !SearchBoxFilter->PassesFilter(*Item))
+	{
+		// This will potentially remove any non-matching, empty parents as well
+		RemoveItemFromTree(Item);
+	}
+	else
+	{
+		// The item still matches the filters (or has children that do)
+		// When an item has been asked to move, it will still reside under its old parent
+		FSceneOutlinerTreeItemPtr Parent = Item->GetParent();
+		if (Parent.IsValid())
+		{
+			Parent->RemoveChild(Item);
+			OnChildRemovedFromParent(*Parent);
+		}
+		else
+		{
+			RootTreeItems.Remove(Item);
+		}
 
-		for (FTreeItemPtr& Item : Items)
+		Parent = EnsureParentForItem(Item);
+		if (Parent.IsValid())
+		{
+			Parent->AddChild(Item);
+			OutlinerTreeView->SetItemExpansion(Parent, true);
+		}
+		else
 		{
-			AddPendingItem(Item);
+			RootTreeItems.Add(Item);
 		}
+	}
+}
 
-		Refresh();
+FSceneOutlinerTreeItemPtr SSceneOutliner::GetTreeItem(FSceneOutlinerTreeItemID ItemID, bool bIncludePending)
+{
+	FSceneOutlinerTreeItemPtr Result = TreeItemMap.FindRef(ItemID);
+	if (bIncludePending && !Result.IsValid())
+	{
+		Result = PendingTreeItemMap.FindRef(ItemID);
 	}
+	return Result;
+}
 
-	void SSceneOutliner::OnChildRemovedFromParent(ITreeItem& Parent)
+void SSceneOutliner::RemoveItemFromTree(FSceneOutlinerTreeItemRef InItem)
+{
+	if (TreeItemMap.Contains(InItem->GetID()))
 	{
-		if (Parent.Flags.bIsFilteredOut && !Parent.GetChildren().Num())
+		auto Parent = InItem->GetParent();
+
+		if (Parent.IsValid())
+		{
+			Parent->RemoveChild(InItem);
+			OnChildRemovedFromParent(*Parent);
+		}
+		else
 		{
-			// The parent no longer has any children that match the current search terms. Remove it.
-			RemoveItemFromTree(Parent.AsShared());
+			RootTreeItems.Remove(InItem);
 		}
+
+		TreeItemMap.Remove(InItem->GetID());
+
+		Mode->OnItemRemoved(InItem);
 	}
+}
 
-	void SSceneOutliner::OnItemMoved(const FTreeItemRef& Item)
+FSceneOutlinerTreeItemPtr SSceneOutliner::EnsureParentForItem(FSceneOutlinerTreeItemRef Item)
+{
+	if (SharedData->bShowParentTree)
 	{
-		// Just remove the item if it no longer matches the filters
-		if (!Item->Flags.bIsFilteredOut && !SearchBoxFilter->PassesFilter(*Item))
+		FSceneOutlinerTreeItemPtr Parent = Mode->GetHierarchy()->FindParent(*Item, TreeItemMap);
+		if (Parent.IsValid())
 		{
-			// This will potentially remove any non-matching, empty parents as well
-			RemoveItemFromTree(Item);
+			return Parent;
 		}
 		else
 		{
-			// The item still matches the filters (or has children that do)
-			// When an item has been asked to move, it will still reside under its old parent
-			FTreeItemPtr Parent = Item->GetParent();
-			if (Parent.IsValid())
-			{
-				Parent->RemoveChild(Item);
-				OnChildRemovedFromParent(*Parent);
-			}
-			else
+			// Try to find the parent in the pending items
+			Parent = Mode->GetHierarchy()->FindParent(*Item, PendingTreeItemMap);
+			if (!Parent.IsValid())
 			{
-				RootTreeItems.Remove(Item);
+				// If there isn't already a parent for this item, try to create one for it
+				Parent = Mode->GetHierarchy()->CreateParentItem(Item);
 			}
-
-			Parent = EnsureParentForItem(Item);
+				
 			if (Parent.IsValid())
 			{
-				Parent->AddChild(Item);
-				OutlinerTreeView->SetItemExpansion(Parent, true);
-			}
-			else
-			{
-				RootTreeItems.Add(Item);
+				AddUnfilteredItemToTree(Parent.ToSharedRef());
+
+				return Parent;
 			}
 		}
 	}
 
-	FTreeItemPtr SSceneOutliner::GetTreeItem(FTreeItemID ItemID, bool bIncludePending)
+	return nullptr;
+}
+
+bool SSceneOutliner::AddItemToTree(FSceneOutlinerTreeItemRef Item)
+{
+	const auto ItemID = Item->GetID();
+
+	PendingTreeItemMap.Remove(ItemID);
+
+	// If a tree item already exists that represents the same data or if the item represents invalid data, bail
+	if (TreeItemMap.Find(ItemID)  || !Item->IsValid())
 	{
-		FTreeItemPtr Result = TreeItemMap.FindRef(ItemID);
-		if (bIncludePending && !Result.IsValid())
-		{
-			Result = PendingTreeItemMap.FindRef(ItemID);
-		}
-		return Result;
+		return false;
 	}
 
-	void SSceneOutliner::RemoveItemFromTree(FTreeItemRef InItem)
+	// Set the filtered out flag
+	Item->Flags.bIsFilteredOut = !SearchBoxFilter->PassesFilter(*Item);
+
+	if (!Item->Flags.bIsFilteredOut)
 	{
-		if (TreeItemMap.Contains(InItem->GetID()))
+		AddUnfilteredItemToTree(Item);
+
+		// Check if we need to do anything with this new item
+		if (uint8* ActionMask = NewItemActions.Find(ItemID))
 		{
-			auto Parent = InItem->GetParent();
+			if (*ActionMask & SceneOutliner::ENewItemAction::Select)
+			{
+				OutlinerTreeView->ClearSelection();
+				OutlinerTreeView->SetItemSelection(Item, true);
+			}
 
-			if (Parent.IsValid())
+			if (*ActionMask & SceneOutliner::ENewItemAction::Rename && CanExecuteRenameRequest(*Item))
 			{
-				Parent->RemoveChild(InItem);
-				OnChildRemovedFromParent(*Parent);
+				PendingRenameItem = Item;
 			}
-			else
+
+			if (*ActionMask & (SceneOutliner::ENewItemAction::ScrollIntoView | SceneOutliner::ENewItemAction::Rename))
 			{
-				RootTreeItems.Remove(InItem);
+				ScrollItemIntoView(Item);
 			}
+		}
+	}
+
+	return true;
+}
 
-			TreeItemMap.Remove(InItem->GetID());
+void SSceneOutliner::AddUnfilteredItemToTree(FSceneOutlinerTreeItemRef Item)
+{
+	FSceneOutlinerTreeItemPtr Parent = EnsureParentForItem(Item);
 
-			Mode->OnItemRemoved(InItem);
+	const FSceneOutlinerTreeItemID ItemID = Item->GetID();
+	if(TreeItemMap.Contains(ItemID))
+	{
+		UE_LOG(LogSceneOutliner, Error, TEXT("(%d | %s) already exists in tree.  Dumping map..."), GetTypeHash(ItemID), *Item->GetDisplayString() );
+		for(TPair<FSceneOutlinerTreeItemID, FSceneOutlinerTreeItemPtr>& Entry : TreeItemMap)
+		{
+			UE_LOG(LogSceneOutliner, Log, TEXT("(%d | %s)"), GetTypeHash(Entry.Key), *Entry.Value->GetDisplayString());
 		}
+
+		// this is a fatal error
+		check(false);
+	}
+
+	TreeItemMap.Add(ItemID, Item);
+
+	if (Parent.IsValid())
+	{
+		Parent->AddChild(Item);
+	}
+	else
+	{
+		RootTreeItems.Add(Item);
 	}
 
-	FTreeItemPtr SSceneOutliner::EnsureParentForItem(FTreeItemRef Item)
+	// keep track of the number of active folders
+	if (Item->IsA<FFolderTreeItem>())
+	{
+		++FolderCount;
+	}
+	Mode->OnItemAdded(Item);
+}
+
+SSceneOutliner::FParentsExpansionState SSceneOutliner::GetParentsExpansionState() const
+{
+	FParentsExpansionState States;
+	for (const auto& Pair : TreeItemMap)
 	{
-		if (SharedData->bShowParentTree)
+		if (Pair.Value->GetChildren().Num())
 		{
-			FTreeItemPtr Parent = Mode->GetHierarchy()->FindParent(*Item, TreeItemMap);
-			if (Parent.IsValid())
-			{
-				return Parent;
-			}
-			else
-			{
-				// Try to find the parent in the pending items
-				Parent = Mode->GetHierarchy()->FindParent(*Item, PendingTreeItemMap);
-				if (!Parent.IsValid())
-				{
-					// If there isn't already a parent for this item, try to create one for it
-					Parent = Mode->GetHierarchy()->CreateParentItem(Item);
-				}
-				
-				if (Parent.IsValid())
-				{
-					AddUnfilteredItemToTree(Parent.ToSharedRef());
-
-					return Parent;
-				}
-			}
+			States.Add(Pair.Key, Pair.Value->Flags.bIsExpanded);
 		}
-
-		return nullptr;
 	}
+	return States;
+}
 
-	bool SSceneOutliner::AddItemToTree(FTreeItemRef Item)
+void SSceneOutliner::SetParentsExpansionState(const FParentsExpansionState& ExpansionStateInfo) const
+{
+	for (const auto& Pair : TreeItemMap)
 	{
-		const auto ItemID = Item->GetID();
-
-		PendingTreeItemMap.Remove(ItemID);
-
-		// If a tree item already exists that represents the same data or if the item represents invalid data, bail
-		if (TreeItemMap.Find(ItemID)  || !Item->IsValid())
+		auto& Item = Pair.Value;
+		if (Item->GetChildren().Num())
 		{
-			return false;
-		}
-
-		// Set the filtered out flag
-		Item->Flags.bIsFilteredOut = !SearchBoxFilter->PassesFilter(*Item);
-
-		if (!Item->Flags.bIsFilteredOut)
-		{
-			AddUnfilteredItemToTree(Item);
-
-			// Check if we need to do anything with this new item
-			if (uint8* ActionMask = NewItemActions.Find(ItemID))
+			const bool* bIsExpanded = ExpansionStateInfo.Find(Pair.Key);
+			if (bIsExpanded)
 			{
-				if (*ActionMask & ENewItemAction::Select)
-				{
-					OutlinerTreeView->ClearSelection();
-					OutlinerTreeView->SetItemSelection(Item, true);
-				}
-
-				if (*ActionMask & ENewItemAction::Rename && CanExecuteRenameRequest(*Item))
-				{
-					PendingRenameItem = Item;
-				}
-
-				if (*ActionMask & (ENewItemAction::ScrollIntoView | ENewItemAction::Rename))
-				{
-					ScrollItemIntoView(Item);
-				}
+				OutlinerTreeView->SetItemExpansion(Item, *bIsExpanded);
 			}
-		}
-
-		return true;
-	}
-
-	void SSceneOutliner::AddUnfilteredItemToTree(FTreeItemRef Item)
-	{
-		FTreeItemPtr Parent = EnsureParentForItem(Item);
-
-		const FTreeItemID ItemID = Item->GetID();
-		if(TreeItemMap.Contains(ItemID))
-		{
-			UE_LOG(LogSceneOutliner, Error, TEXT("(%d | %s) already exists in tree.  Dumping map..."), GetTypeHash(ItemID), *Item->GetDisplayString() );
-			for(TPair<FTreeItemID, FTreeItemPtr>& Entry : TreeItemMap)
+			else
 			{
-				UE_LOG(LogSceneOutliner, Log, TEXT("(%d | %s)"), GetTypeHash(Entry.Key), *Entry.Value->GetDisplayString());
+				OutlinerTreeView->SetItemExpansion(Item, Item->Flags.bIsExpanded);
 			}
-
-			// this is a fatal error
-			check(false);
 		}
-
-		TreeItemMap.Add(ItemID, Item);
-
-		if (Parent.IsValid())
-		{
-			Parent->AddChild(Item);
-		}
-		else
-		{
-			RootTreeItems.Add(Item);
-		}
-
-		// keep track of the number of active folders
-		if (Item->IsA<FFolderTreeItem>())
-		{
-			++FolderCount;
-		}
-		Mode->OnItemAdded(Item);
 	}
+}
 
-	SSceneOutliner::FParentsExpansionState SSceneOutliner::GetParentsExpansionState() const
+void SSceneOutliner::PopulateSearchStrings(const ISceneOutlinerTreeItem& Item, TArray< FString >& OutSearchStrings) const
+{
+	for (const auto& Pair : Columns)
 	{
-		FParentsExpansionState States;
-		for (const auto& Pair : TreeItemMap)
-		{
-			if (Pair.Value->GetChildren().Num())
-			{
-				States.Add(Pair.Key, Pair.Value->Flags.bIsExpanded);
-			}
-		}
-		return States;
+		Pair.Value->PopulateSearchStrings(Item, OutSearchStrings);
 	}
+}
 
-	void SSceneOutliner::SetParentsExpansionState(const FParentsExpansionState& ExpansionStateInfo) const
-	{
-		for (const auto& Pair : TreeItemMap)
-		{
-			auto& Item = Pair.Value;
-			if (Item->GetChildren().Num())
-			{
-				const bool* bIsExpanded = ExpansionStateInfo.Find(Pair.Key);
-				if (bIsExpanded)
-				{
-					OutlinerTreeView->SetItemExpansion(Item, *bIsExpanded);
-				}
-				else
-				{
-					OutlinerTreeView->SetItemExpansion(Item, Item->Flags.bIsExpanded);
-				}
-			}
-		}
-	}
+void SSceneOutliner::GetSelectedFolders(TArray<FFolderTreeItem*>& OutFolders) const
+{
+	return FSceneOutlinerItemSelection(*OutlinerTreeView).Get<FFolderTreeItem>(OutFolders);
+}
 
-	void SSceneOutliner::PopulateSearchStrings(const ITreeItem& Item, TArray< FString >& OutSearchStrings) const
-	{
-		for (const auto& Pair : Columns)
-		{
-			Pair.Value->PopulateSearchStrings(Item, OutSearchStrings);
-		}
-	}
+TArray<FName> SSceneOutliner::GetSelectedFolderNames() const
+{
+	return GetSelection().GetData<FName>(SceneOutliner::FFolderPathSelector());
+}
 
-	void SSceneOutliner::GetSelectedFolders(TArray<FFolderTreeItem*>& OutFolders) const
-	{
-		return FItemSelection(*OutlinerTreeView).Get<FFolderTreeItem>(OutFolders);
-	}
+TSharedPtr<SWidget> SSceneOutliner::OnOpenContextMenu()
+{
+	return Mode->CreateContextMenu();
+}
+	
+bool SSceneOutliner::Delete_CanExecute()
+{
+	return Mode->CanDelete();
+}
+bool SSceneOutliner::Rename_CanExecute()
+{
+	return Mode->CanRename();
+}
 
-	TArray<FName> SSceneOutliner::GetSelectedFolderNames() const
-	{
-		return GetSelection().GetData<FName>(FFolderPathSelector());
-	}
+void SSceneOutliner::Rename_Execute()
+{
+	FSceneOutlinerItemSelection ItemSelection(*OutlinerTreeView);
+	FSceneOutlinerTreeItemPtr ItemToRename;
 
-	TSharedPtr<SWidget> SSceneOutliner::OnOpenContextMenu()
-	{
-		return Mode->CreateContextMenu();
-	}
-	
-	bool SSceneOutliner::Delete_CanExecute()
-	{
-		return Mode->CanDelete();
-	}
-	bool SSceneOutliner::Rename_CanExecute()
+	if (Mode->CanRename())
 	{
-		return Mode->CanRename();
+		ItemToRename = OutlinerTreeView->GetSelectedItems()[0];
 	}
 
-	void SSceneOutliner::Rename_Execute()
+	if (ItemToRename.IsValid() && CanExecuteRenameRequest(*ItemToRename) && ItemToRename->CanInteract())
 	{
-		FItemSelection ItemSelection(*OutlinerTreeView);
-		FTreeItemPtr ItemToRename;
-
-		if (Mode->CanRename())
-		{
-			ItemToRename = OutlinerTreeView->GetSelectedItems()[0];
-		}
-
-		if (ItemToRename.IsValid() && CanExecuteRenameRequest(*ItemToRename) && ItemToRename->CanInteract())
-		{
-			PendingRenameItem = ItemToRename->AsShared();
-			ScrollItemIntoView(ItemToRename);
-		}
+		PendingRenameItem = ItemToRename->AsShared();
+		ScrollItemIntoView(ItemToRename);
 	}
+}
 
-	bool SSceneOutliner::Cut_CanExecute()
-	{
-		return Mode->CanCut();
-	}
+bool SSceneOutliner::Cut_CanExecute()
+{
+	return Mode->CanCut();
+}
 
-	bool SSceneOutliner::Copy_CanExecute()
-	{
-		return Mode->CanCopy();
-	}
+bool SSceneOutliner::Copy_CanExecute()
+{
+	return Mode->CanCopy();
+}
 
-	bool SSceneOutliner::Paste_CanExecute()
-	{
-		return Mode->CanPaste();
-	}
+bool SSceneOutliner::Paste_CanExecute()
+{
+	return Mode->CanPaste();
+}
 
-	bool SSceneOutliner::CanSupportDragAndDrop() const
-	{
-		return Mode->CanSupportDragAndDrop();
-	}
+bool SSceneOutliner::CanSupportDragAndDrop() const
+{
+	return Mode->CanSupportDragAndDrop();
+}
 
-	bool SSceneOutliner::CanExecuteRenameRequest(const ITreeItem& ItemPtr) const
-	{
-		return Mode->CanRenameItem(ItemPtr);
-	}
+bool SSceneOutliner::CanExecuteRenameRequest(const ISceneOutlinerTreeItem& ItemPtr) const
+{
+	return Mode->CanRenameItem(ItemPtr);
+}
 
-	int32 SSceneOutliner::AddFilter(const TSharedRef<SceneOutliner::FOutlinerFilter>& Filter)
-	{
-		return Filters->Add(Filter);
-	}
+int32 SSceneOutliner::AddFilter(const TSharedRef<FSceneOutlinerFilter>& Filter)
+{
+	return Filters->Add(Filter);
+}
 
-	bool SSceneOutliner::RemoveFilter(const TSharedRef<SceneOutliner::FOutlinerFilter>& Filter)
-	{
-		return  Filters->Remove(Filter) > 0;
-	}
+bool SSceneOutliner::RemoveFilter(const TSharedRef<FSceneOutlinerFilter>& Filter)
+{
+	return  Filters->Remove(Filter) > 0;
+}
 
-	TSharedPtr<SceneOutliner::FOutlinerFilter> SSceneOutliner::GetFilterAtIndex(int32 Index)
-	{
-		return StaticCastSharedPtr<SceneOutliner::FOutlinerFilter>(Filters->GetFilterAtIndex(Index));
-	}
+TSharedPtr<FSceneOutlinerFilter> SSceneOutliner::GetFilterAtIndex(int32 Index)
+{
+	return StaticCastSharedPtr<FSceneOutlinerFilter>(Filters->GetFilterAtIndex(Index));
+}
 
-	int32 SSceneOutliner::GetFilterCount() const
-	{
-		return Filters->Num();
-	}
+int32 SSceneOutliner::GetFilterCount() const
+{
+	return Filters->Num();
+}
 
-	void SSceneOutliner::AddColumn(FName ColumId, const SceneOutliner::FColumnInfo& ColumInfo)
+void SSceneOutliner::AddColumn(FName ColumId, const FSceneOutlinerColumnInfo& ColumInfo)
+{
+	if (!SharedData->ColumnMap.Contains(ColumId))
 	{
-		if (!SharedData->ColumnMap.Contains(ColumId))
-		{
-			SharedData->ColumnMap.Add(ColumId, ColumInfo);
-			RefreshColums();
-		}
+		SharedData->ColumnMap.Add(ColumId, ColumInfo);
+		RefreshColums();
 	}
+}
 
-	void SSceneOutliner::RemoveColumn(FName ColumId)
+void SSceneOutliner::RemoveColumn(FName ColumId)
+{
+	if (SharedData->ColumnMap.Contains(ColumId))
 	{
-		if (SharedData->ColumnMap.Contains(ColumId))
-		{
-			SharedData->ColumnMap.Remove(ColumId);
-			RefreshColums();
-		}
+		SharedData->ColumnMap.Remove(ColumId);
+		RefreshColums();
 	}
+}
 
-	TArray<FName> SSceneOutliner::GetColumnIds() const
-	{
-		TArray<FName> ColumnsName;
-		SharedData->ColumnMap.GenerateKeyArray(ColumnsName);
-		return ColumnsName;
-	}
+TArray<FName> SSceneOutliner::GetColumnIds() const
+{
+	TArray<FName> ColumnsName;
+	SharedData->ColumnMap.GenerateKeyArray(ColumnsName);
+	return ColumnsName;
+}
 
-	void SSceneOutliner::SetSelection(const TFunctionRef<bool(SceneOutliner::ITreeItem&)> Selector)
+void SSceneOutliner::SetSelection(const TFunctionRef<bool(ISceneOutlinerTreeItem&)> Selector)
+{
+	TArray<FSceneOutlinerTreeItemPtr> ItemsToAdd;
+	for (const auto& Pair : TreeItemMap)
 	{
-		TArray<FTreeItemPtr> ItemsToAdd;
-		for (const auto& Pair : TreeItemMap)
+		FSceneOutlinerTreeItemPtr ItemPtr = Pair.Value;
+		if (ISceneOutlinerTreeItem* Item = ItemPtr.Get())
 		{
-			FTreeItemPtr ItemPtr = Pair.Value;
-			if (ITreeItem* Item = ItemPtr.Get())
+			if (Selector(*Item))
 			{
-				if (Selector(*Item))
-				{
-					ItemsToAdd.Add(ItemPtr);
-				}
+				ItemsToAdd.Add(ItemPtr);
 			}
 		}
-
-		SetItemSelection(ItemsToAdd, true);
 	}
 
-	void SSceneOutliner::SetItemSelection(const TArray<FTreeItemPtr>& InItems, bool bSelected, ESelectInfo::Type SelectInfo)
-	{
-		OutlinerTreeView->ClearSelection();
-		OutlinerTreeView->SetItemSelection(InItems, bSelected, SelectInfo);
-	}
+	SetItemSelection(ItemsToAdd, true);
+}
 
-	void SSceneOutliner::SetItemSelection(const FTreeItemPtr& InItem, bool bSelected, ESelectInfo::Type SelectInfo)
-	{
-		OutlinerTreeView->ClearSelection();
-		OutlinerTreeView->SetItemSelection(InItem, bSelected, SelectInfo);
-	}
+void SSceneOutliner::SetItemSelection(const TArray<FSceneOutlinerTreeItemPtr>& InItems, bool bSelected, ESelectInfo::Type SelectInfo)
+{
+	OutlinerTreeView->ClearSelection();
+	OutlinerTreeView->SetItemSelection(InItems, bSelected, SelectInfo);
+}
 
-	void SSceneOutliner::AddToSelection(const TArray<FTreeItemPtr>& InItems, ESelectInfo::Type SelectInfo)
-	{
-		OutlinerTreeView->SetItemSelection(InItems, true, SelectInfo);
-	}
+void SSceneOutliner::SetItemSelection(const FSceneOutlinerTreeItemPtr& InItem, bool bSelected, ESelectInfo::Type SelectInfo)
+{
+	OutlinerTreeView->ClearSelection();
+	OutlinerTreeView->SetItemSelection(InItem, bSelected, SelectInfo);
+}
 
-	void SSceneOutliner::RemoveFromSelection(const TArray<FTreeItemPtr>& InItems)
-	{
-		OutlinerTreeView->SetItemSelection(InItems, false);
-	}
+void SSceneOutliner::AddToSelection(const TArray<FSceneOutlinerTreeItemPtr>& InItems, ESelectInfo::Type SelectInfo)
+{
+	OutlinerTreeView->SetItemSelection(InItems, true, SelectInfo);
+}
 
-	void SSceneOutliner::RemoveFromSelection(const FTreeItemPtr& InItem)
-	{
-		OutlinerTreeView->SetItemSelection(InItem, false);
-	}
+void SSceneOutliner::RemoveFromSelection(const TArray<FSceneOutlinerTreeItemPtr>& InItems, ESelectInfo::Type SelectInfo)
+{
+	OutlinerTreeView->SetItemSelection(InItems, false, SelectInfo);
+}
 
-	void SSceneOutliner::AddFolderToSelection(const FName& FolderName)
+void SSceneOutliner::AddFolderToSelection(const FName& FolderName)
+{
+	if (FSceneOutlinerTreeItemPtr* ItemPtr = TreeItemMap.Find(FolderName))
 	{
-		if (FTreeItemPtr* ItemPtr = TreeItemMap.Find(FolderName))
-		{
-			OutlinerTreeView->SetItemSelection(*ItemPtr, true);
-		}
+		OutlinerTreeView->SetItemSelection(*ItemPtr, true);
 	}
+}
 
-	void SSceneOutliner::RemoveFolderFromSelection(const FName& FolderName)
+void SSceneOutliner::RemoveFolderFromSelection(const FName& FolderName)
+{
+	if (FSceneOutlinerTreeItemPtr* ItemPtr = TreeItemMap.Find(FolderName))
 	{
-		if (FTreeItemPtr* ItemPtr = TreeItemMap.Find(FolderName))
-		{
-			OutlinerTreeView->SetItemSelection(*ItemPtr, false);
-		}
+		OutlinerTreeView->SetItemSelection(*ItemPtr, false);
 	}
+}
 
-	void SSceneOutliner::ClearSelection()
+void SSceneOutliner::ClearSelection()
+{
+	if (!bIsReentrant)
 	{
-		if (!bIsReentrant)
-		{
-			OutlinerTreeView->ClearSelection();
-		}
+		OutlinerTreeView->ClearSelection();
 	}
+}
 
-	void SSceneOutliner::FillFoldersSubMenu(UToolMenu* Menu) const
-	{
-		FToolMenuSection& Section = Menu->AddSection("Section");
-		Section.AddMenuEntry("CreateNew", LOCTEXT( "CreateNew", "Create New Folder" ), LOCTEXT( "CreateNew_ToolTip", "Move to a new folder" ),
-			FSlateIcon(FEditorStyle::GetStyleSetName(), "SceneOutliner.NewFolderIcon"), FExecuteAction::CreateSP(const_cast<SSceneOutliner*>(this), &SSceneOutliner::CreateFolder));
+void SSceneOutliner::FillFoldersSubMenu(UToolMenu* Menu) const
+{
+	FToolMenuSection& Section = Menu->AddSection("Section");
+	Section.AddMenuEntry("CreateNew", LOCTEXT( "CreateNew", "Create New Folder" ), LOCTEXT( "CreateNew_ToolTip", "Move to a new folder" ),
+		FSlateIcon(FEditorStyle::GetStyleSetName(), "SceneOutliner.NewFolderIcon"), FExecuteAction::CreateSP(const_cast<SSceneOutliner*>(this), &SSceneOutliner::CreateFolder));
 
-		AddMoveToFolderOutliner(Menu);
-	}
+	AddMoveToFolderOutliner(Menu);
+}
 
-	TSharedRef<TSet<FName>> SSceneOutliner::GatherInvalidMoveToDestinations() const
-	{
-		// We use a pointer here to save copying the whole array for every invocation of the filter delegate
-		TSharedRef<TSet<FName>> ExcludedParents(new TSet<FName>());
+TSharedRef<TSet<FName>> SSceneOutliner::GatherInvalidMoveToDestinations() const
+{
+	// We use a pointer here to save copying the whole array for every invocation of the filter delegate
+	TSharedRef<TSet<FName>> ExcludedParents(new TSet<FName>());
 
-		for (const auto& Item : OutlinerTreeView->GetSelectedItems())
+	for (const auto& Item : OutlinerTreeView->GetSelectedItems())
+	{
+		if (FFolderTreeItem* ParentFolderItem = Item->GetParent()->CastTo<FFolderTreeItem>())
 		{
-			if (FFolderTreeItem* ParentFolderItem = Item->GetParent()->CastTo<FFolderTreeItem>())
+			auto FolderHasOtherSubFolders = [&Item](const TWeakPtr<ISceneOutlinerTreeItem>& WeakItem)
 			{
-				auto FolderHasOtherSubFolders = [&Item](const TWeakPtr<ITreeItem>& WeakItem)
+				if (WeakItem.Pin()->IsA<FFolderTreeItem>() && WeakItem.Pin() != Item)
 				{
-					if (WeakItem.Pin()->IsA<FFolderTreeItem>() && WeakItem.Pin() != Item)
-					{
-						return true;
-					}
-					return false;
-				};
-
-				// Exclude this items direct parent if it is a folder and has no other subfolders we can move to
-				if (!ParentFolderItem->GetChildren().ContainsByPredicate(FolderHasOtherSubFolders))
-				{
-					ExcludedParents->Add(ParentFolderItem->Path);
+					return true;
 				}
-			}
-			
-			if (FFolderTreeItem* FolderItem = Item->CastTo<FFolderTreeItem>())
+				return false;
+			};
+
+			// Exclude this items direct parent if it is a folder and has no other subfolders we can move to
+			if (!ParentFolderItem->GetChildren().ContainsByPredicate(FolderHasOtherSubFolders))
 			{
-				// Cannot move into itself, or any child
-				ExcludedParents->Add(FolderItem->Path);
+				ExcludedParents->Add(ParentFolderItem->Path);
 			}
 		}
-
-		return ExcludedParents;
+			
+		if (FFolderTreeItem* FolderItem = Item->CastTo<FFolderTreeItem>())
+		{
+			// Cannot move into itself, or any child
+			ExcludedParents->Add(FolderItem->Path);
+		}
 	}
+
+	return ExcludedParents;
+}
 	
-	void SSceneOutliner::AddMoveToFolderOutliner(UToolMenu* Menu) const
+void SSceneOutliner::AddMoveToFolderOutliner(UToolMenu* Menu) const
+{
+	// We don't show this if there aren't any folders in the world and if the mode is showing folders
+	if (!Mode->ShouldShowFolders() || FolderCount == 0)
 	{
-		// We don't show this if there aren't any folders in the world and if the mode is showing folders
-		if (!Mode->ShouldShowFolders() || FolderCount == 0)
-		{
-			return;
-		}
+		return;
+	}
 
-		// Add a mini scene outliner for choosing an existing folder
-		FInitializationOptions MiniSceneOutlinerInitOptions;
-		MiniSceneOutlinerInitOptions.bShowHeaderRow = false;
-		MiniSceneOutlinerInitOptions.bFocusSearchBoxWhenOpened = true;
+	// Add a mini scene outliner for choosing an existing folder
+	FSceneOutlinerInitializationOptions MiniSceneOutlinerInitOptions;
+	MiniSceneOutlinerInitOptions.bShowHeaderRow = false;
+	MiniSceneOutlinerInitOptions.bFocusSearchBoxWhenOpened = true;
 		
-		// Don't show any folders that are a child of any of the selected folders
-		auto ExcludedParents = GatherInvalidMoveToDestinations();
-		if (ExcludedParents->Num())
-		{
-			// Add a filter if necessary
-			auto FilterOutChildFolders = [](FName Path, TSharedRef<TSet<FName>> ExcludedParents){
-				for (const auto& Parent : *ExcludedParents)
+	// Don't show any folders that are a child of any of the selected folders
+	auto ExcludedParents = GatherInvalidMoveToDestinations();
+	if (ExcludedParents->Num())
+	{
+		// Add a filter if necessary
+		auto FilterOutChildFolders = [](FName Path, TSharedRef<TSet<FName>> ExcludedParents){
+			for (const auto& Parent : *ExcludedParents)
+			{
+				if (Path == Parent || SceneOutliner::PathIsChildOf(Path, Parent))
 				{
-					if (Path == Parent || PathIsChildOf(Path, Parent))
-					{
-						return false;
-					}
+					return false;
 				}
-				return true;
-			};
+			}
+			return true;
+		};
 
-			MiniSceneOutlinerInitOptions.Filters->AddFilterPredicate<FFolderTreeItem>(FFolderTreeItem::FFilterPredicate::CreateStatic(FilterOutChildFolders, ExcludedParents), EDefaultFilterBehaviour::Pass);
-		}
+		MiniSceneOutlinerInitOptions.Filters->AddFilterPredicate<FFolderTreeItem>(FFolderTreeItem::FFilterPredicate::CreateStatic(FilterOutChildFolders, ExcludedParents), FSceneOutlinerFilter::EDefaultBehaviour::Pass);
+	}
 
+	{
+		struct FFilterRoot : public FSceneOutlinerFilter 
 		{
-			struct FFilterRoot : public FOutlinerFilter 
+			FFilterRoot(TFunctionRef<bool(const ISceneOutlinerTreeItem&)> InFindParent) 
+				: FSceneOutlinerFilter(FSceneOutlinerFilter::EDefaultBehaviour::Pass) 
+				, FindParent(InFindParent)
+			{}
+			virtual bool PassesFilter(const ISceneOutlinerTreeItem& Item) const override
 			{
-				FFilterRoot(TFunctionRef<bool(const ITreeItem&)> InFindParent) 
-					: FOutlinerFilter(EDefaultFilterBehaviour::Pass) 
-					, FindParent(InFindParent)
-				{}
-				virtual bool PassesFilter(const ITreeItem& Item) const override
+				// if item has no parent, it is a root item and should be filtered out
+				if (!FindParent(Item))
 				{
-					// if item has no parent, it is a root item and should be filtered out
-					if (!FindParent(Item))
-					{
-						return false;
-					}
-					return DefaultBehaviour == EDefaultFilterBehaviour::Pass;
+					return false;
 				}
+				return DefaultBehaviour == FSceneOutlinerFilter::EDefaultBehaviour::Pass;
+			}
 
-				TFunctionRef<bool(const ITreeItem&)> FindParent;
-			};
-
-			auto FindParent = [this](const ITreeItem& Item)
-			{
-				FTreeItemPtr Parent = Mode->GetHierarchy()->FindParent(Item, TreeItemMap);
-				if (!Parent)
-				{
-					Parent = Mode->GetHierarchy()->FindParent(Item, PendingTreeItemMap);
-				}
-				return Parent.IsValid();
-			};
+			TFunctionRef<bool(const ISceneOutlinerTreeItem&)> FindParent;
+		};
 
-			// Filter in/out root items according to whether it is valid to move to/from the root
-			FDragDropPayload DraggedObjects(OutlinerTreeView->GetSelectedItems());
-			const bool bMoveToRootValid = Mode->ValidateDrop(FFolderTreeItem(NAME_None), DraggedObjects).IsValid();
-			if (!bMoveToRootValid)
+		auto FindParent = [this](const ISceneOutlinerTreeItem& Item)
+		{
+			FSceneOutlinerTreeItemPtr Parent = Mode->GetHierarchy()->FindParent(Item, TreeItemMap);
+			if (!Parent)
 			{
-				MiniSceneOutlinerInitOptions.Filters->Add(MakeShared<FFilterRoot>(FindParent));
+				Parent = Mode->GetHierarchy()->FindParent(Item, PendingTreeItemMap);
 			}
+			return Parent.IsValid();
+		};
+
+		// Filter in/out root items according to whether it is valid to move to/from the root
+		FSceneOutlinerDragDropPayload DraggedObjects(OutlinerTreeView->GetSelectedItems());
+		const bool bMoveToRootValid = Mode->ValidateDrop(FFolderTreeItem(NAME_None), DraggedObjects).IsValid();
+		if (!bMoveToRootValid)
+		{
+			MiniSceneOutlinerInitOptions.Filters->Add(MakeShared<FFilterRoot>(FindParent));
 		}
+	}
 
-		//Let the mode decide how folder selection is handled
+	//Let the mode decide how folder selection is handled
 
-		MiniSceneOutlinerInitOptions.ModeFactory = Mode->CreateFolderPickerMode();
+	MiniSceneOutlinerInitOptions.ModeFactory = Mode->CreateFolderPickerMode();
 
-		FSceneOutlinerModule& SceneOutlinerModule = FModuleManager::LoadModuleChecked<FSceneOutlinerModule>("SceneOutliner");
-		TSharedRef< SWidget > MiniSceneOutliner =
-			SNew(SVerticalBox)
-			+ SVerticalBox::Slot()
-			.MaxHeight(400.0f)
-			[
-				SNew(SSceneOutliner, MiniSceneOutlinerInitOptions)
-				.IsEnabled(FSlateApplication::Get().GetNormalExecutionAttribute())
-			];
+	FSceneOutlinerModule& SceneOutlinerModule = FModuleManager::LoadModuleChecked<FSceneOutlinerModule>("SceneOutliner");
+	TSharedRef< SWidget > MiniSceneOutliner =
+		SNew(SVerticalBox)
+		+ SVerticalBox::Slot()
+		.MaxHeight(400.0f)
+		[
+			SNew(SSceneOutliner, MiniSceneOutlinerInitOptions)
+			.IsEnabled(FSlateApplication::Get().GetNormalExecutionAttribute())
+		];
 
-		FToolMenuSection& Section = Menu->AddSection(FName(), LOCTEXT("ExistingFolders", "Existing:"));
-		Section.AddEntry(FToolMenuEntry::InitWidget(
-			"MiniSceneOutliner",
-			MiniSceneOutliner,
-			FText::GetEmpty(),
-			false));
-	}
+	FToolMenuSection& Section = Menu->AddSection(FName(), LOCTEXT("ExistingFolders", "Existing:"));
+	Section.AddEntry(FToolMenuEntry::InitWidget(
+		"MiniSceneOutliner",
+		MiniSceneOutliner,
+		FText::GetEmpty(),
+		false));
+}
 
-	void SSceneOutliner::FillSelectionSubMenu(UToolMenu* Menu) const
-	{
-		FToolMenuSection& Section = Menu->AddSection("Section");
-		Section.AddMenuEntry(
-			"AddChildrenToSelection",
-			LOCTEXT( "AddChildrenToSelection", "Immediate Children" ),
-			LOCTEXT( "AddChildrenToSelection_ToolTip", "Select all immediate children of the selected folders" ),
-			FSlateIcon(),
-			FExecuteAction::CreateSP(const_cast<SSceneOutliner*>(this), &SSceneOutliner::SelectFoldersDescendants, /*bSelectImmediateChildrenOnly=*/ true));
-		Section.AddMenuEntry(
-			"AddDescendantsToSelection",
-			LOCTEXT( "AddDescendantsToSelection", "All Descendants" ),
-			LOCTEXT( "AddDescendantsToSelection_ToolTip", "Select all descendants of the selected folders" ),
-			FSlateIcon(),
-			FExecuteAction::CreateSP(const_cast<SSceneOutliner*>(this), &SSceneOutliner::SelectFoldersDescendants, /*bSelectImmediateChildrenOnly=*/ false));
-	}
+void SSceneOutliner::FillSelectionSubMenu(UToolMenu* Menu) const
+{
+	FToolMenuSection& Section = Menu->AddSection("Section");
+	Section.AddMenuEntry(
+		"AddChildrenToSelection",
+		LOCTEXT( "AddChildrenToSelection", "Immediate Children" ),
+		LOCTEXT( "AddChildrenToSelection_ToolTip", "Select all immediate children of the selected folders" ),
+		FSlateIcon(),
+		FExecuteAction::CreateSP(const_cast<SSceneOutliner*>(this), &SSceneOutliner::SelectFoldersDescendants, /*bSelectImmediateChildrenOnly=*/ true));
+	Section.AddMenuEntry(
+		"AddDescendantsToSelection",
+		LOCTEXT( "AddDescendantsToSelection", "All Descendants" ),
+		LOCTEXT( "AddDescendantsToSelection_ToolTip", "Select all descendants of the selected folders" ),
+		FSlateIcon(),
+		FExecuteAction::CreateSP(const_cast<SSceneOutliner*>(this), &SSceneOutliner::SelectFoldersDescendants, /*bSelectImmediateChildrenOnly=*/ false));
+}
+
+void SSceneOutliner::SelectFoldersDescendants(bool bSelectImmediateChildrenOnly)
+{
+	TArray<FFolderTreeItem*> SelectedFolders;
+	GetSelectedFolders(SelectedFolders);
+	OutlinerTreeView->ClearSelection();
 
-	void SSceneOutliner::SelectFoldersDescendants(bool bSelectImmediateChildrenOnly)
+	if (SelectedFolders.Num())
 	{
-		TArray<FFolderTreeItem*> SelectedFolders;
-		GetSelectedFolders(SelectedFolders);
-		OutlinerTreeView->ClearSelection();
+		Mode->SelectFoldersDescendants(SelectedFolders, bSelectImmediateChildrenOnly);
+	}
 
-		if (SelectedFolders.Num())
-		{
-			Mode->SelectFoldersDescendants(SelectedFolders, bSelectImmediateChildrenOnly);
-		}
+	Refresh();
+}
 
-		Refresh();
-	}
+void SSceneOutliner::MoveSelectionTo(FName NewParent)
+{
+	FSlateApplication::Get().DismissAllMenus();
+		
+	FFolderTreeItem 	DropTarget(NewParent);
+	FSceneOutlinerDragDropPayload 	DraggedObjects(OutlinerTreeView->GetSelectedItems());
 
-	void SSceneOutliner::MoveSelectionTo(FName NewParent)
+	FSceneOutlinerDragValidationInfo Validation = Mode->ValidateDrop(DropTarget, DraggedObjects);
+	if (!Validation.IsValid())
 	{
-		FSlateApplication::Get().DismissAllMenus();
-		
-		FFolderTreeItem 	DropTarget(NewParent);
-		FDragDropPayload 	DraggedObjects(OutlinerTreeView->GetSelectedItems());
+		FNotificationInfo Info(Validation.ValidationText);
+		Info.ExpireDuration = 3.0f;
+		Info.bUseLargeFont = false;
+		Info.bFireAndForget = true;
+		Info.bUseSuccessFailIcons = true;
+		FSlateNotificationManager::Get().AddNotification(Info)->SetCompletionState(SNotificationItem::CS_Fail);
+		return;
+	}
 
-		FDragValidationInfo Validation = Mode->ValidateDrop(DropTarget, DraggedObjects);
-		if (!Validation.IsValid())
-		{
-			FNotificationInfo Info(Validation.ValidationText);
-			Info.ExpireDuration = 3.0f;
-			Info.bUseLargeFont = false;
-			Info.bFireAndForget = true;
-			Info.bUseSuccessFailIcons = true;
-			FSlateNotificationManager::Get().AddNotification(Info)->SetCompletionState(SNotificationItem::CS_Fail);
-			return;
-		}
+	const FScopedTransaction Transaction(LOCTEXT("MoveOutlinerItems", "Move Outliner Items"));
+	Mode->OnDrop(DropTarget, DraggedObjects, Validation);
+}
 
-		const FScopedTransaction Transaction(LOCTEXT("MoveOutlinerItems", "Move Outliner Items"));
-		Mode->OnDrop(DropTarget, DraggedObjects, Validation);
-	}
+FReply SSceneOutliner::OnCreateFolderClicked()
+{
+	CreateFolder();
+	return FReply::Handled();
+}
 
-	FReply SSceneOutliner::OnCreateFolderClicked()
-	{
-		CreateFolder();
-		return FReply::Handled();
-	}
+void SSceneOutliner::CreateFolder()
+{
+	const FName& NewFolderName = Mode->CreateNewFolder();
 
-	void SSceneOutliner::CreateFolder()
+	// Move any selected folders into the new folder
+	auto PreviouslySelectedItems = OutlinerTreeView->GetSelectedItems();
+	for (const auto& Item : PreviouslySelectedItems)
 	{
-		const FName& NewFolderName = Mode->CreateNewFolder();
-
-		// Move any selected folders into the new folder
-		auto PreviouslySelectedItems = OutlinerTreeView->GetSelectedItems();
-		for (const auto& Item : PreviouslySelectedItems)
+		if (FFolderTreeItem* FolderItem = Item->CastTo<FFolderTreeItem>())
 		{
-			if (FFolderTreeItem* FolderItem = Item->CastTo<FFolderTreeItem>())
-			{
-				FolderItem->MoveTo(NewFolderName);
-			}
+			FolderItem->MoveTo(NewFolderName);
 		}
-
-		// At this point the new folder will be in our newly added list, so select it and open a rename when it gets refreshed
-		NewItemActions.Add(NewFolderName, ENewItemAction::Select | ENewItemAction::Rename);
-		RequestSort();
 	}
 
-	void SSceneOutliner::CopyFoldersBegin()
-	{
-		CacheFoldersEdit = GetSelectedFolderNames();
-		FPlatformApplicationMisc::ClipboardPaste(CacheClipboardContents);
-	}
+	// At this point the new folder will be in our newly added list, so select it and open a rename when it gets refreshed
+	NewItemActions.Add(NewFolderName, SceneOutliner::ENewItemAction::Select | SceneOutliner::ENewItemAction::Rename);
+	RequestSort();
+}
+
+void SSceneOutliner::CopyFoldersBegin()
+{
+	CacheFoldersEdit = GetSelectedFolderNames();
+	FPlatformApplicationMisc::ClipboardPaste(CacheClipboardContents);
+}
 
-	void SSceneOutliner::CopyFoldersEnd()
+void SSceneOutliner::CopyFoldersEnd()
+{
+	if (CacheFoldersEdit.Num() > 0)
 	{
-		if (CacheFoldersEdit.Num() > 0)
-		{
-			CopyFoldersToClipboard(CacheFoldersEdit, CacheClipboardContents);
-			CacheFoldersEdit.Reset();
-		}
+		CopyFoldersToClipboard(CacheFoldersEdit, CacheClipboardContents);
+		CacheFoldersEdit.Reset();
 	}
+}
 
-	void SSceneOutliner::CopyFoldersToClipboard(const TArray<FName>& InFolders, const FString& InPrevClipboardContents)
+void SSceneOutliner::CopyFoldersToClipboard(const TArray<FName>& InFolders, const FString& InPrevClipboardContents)
+{
+	if (InFolders.Num() > 0)
 	{
-		if (InFolders.Num() > 0)
-		{
-			// If clipboard paste has changed since we cached it, items must have been cut 
-			// so folders need to appended to clipboard contents rather than replacing them.
-			FString CurrClipboardContents;
-			FPlatformApplicationMisc::ClipboardPaste(CurrClipboardContents);
-
-			FString Buffer = ExportFolderList(InFolders);
+		// If clipboard paste has changed since we cached it, items must have been cut 
+		// so folders need to appended to clipboard contents rather than replacing them.
+		FString CurrClipboardContents;
+		FPlatformApplicationMisc::ClipboardPaste(CurrClipboardContents);
 
-			FString* SourceData = (CurrClipboardContents != InPrevClipboardContents ? &CurrClipboardContents : nullptr);
+		FString Buffer = ExportFolderList(InFolders);
 
-			if (SourceData)
-			{
-				SourceData->Append(Buffer);
-			}
-			else
-			{
-				SourceData = &Buffer;
-			}
+		FString* SourceData = (CurrClipboardContents != InPrevClipboardContents ? &CurrClipboardContents : nullptr);
 
-			// Replace clipboard contents with original plus folders appended
-			FPlatformApplicationMisc::ClipboardCopy(**SourceData);
+		if (SourceData)
+		{
+			SourceData->Append(Buffer);
+		}
+		else
+		{
+			SourceData = &Buffer;
 		}
+
+		// Replace clipboard contents with original plus folders appended
+		FPlatformApplicationMisc::ClipboardCopy(**SourceData);
 	}
+}
 
-	void SSceneOutliner::PasteFoldersBegin(TArray<FName> InFolders)
+void SSceneOutliner::PasteFoldersBegin(TArray<FName> InFolders)
+{
+	auto CacheExistingChildrenAction = [this](const FSceneOutlinerTreeItemPtr& Item)
 	{
-		auto CacheExistingChildrenAction = [this](const FTreeItemPtr& Item)
+		if (FFolderTreeItem* FolderItem = Item->CastTo<FFolderTreeItem>())
 		{
-			if (FFolderTreeItem* FolderItem = Item->CastTo<FFolderTreeItem>())
+			TArray<FSceneOutlinerTreeItemID> ExistingChildren;
+			for (const TWeakPtr<ISceneOutlinerTreeItem>& Child : FolderItem->GetChildren())
 			{
-				TArray<FTreeItemID> ExistingChildren;
-				for (const TWeakPtr<ITreeItem>& Child : FolderItem->GetChildren())
+				if (Child.IsValid())
 				{
-					if (Child.IsValid())
-					{
-						ExistingChildren.Add(Child.Pin()->GetID());
-					}
+					ExistingChildren.Add(Child.Pin()->GetID());
 				}
-
-				CachePasteFolderExistingChildrenMap.Add(FolderItem->Path, ExistingChildren);
 			}
-		};
+
+			CachePasteFolderExistingChildrenMap.Add(FolderItem->Path, ExistingChildren);
+		}
+	};
 
 
-		CacheFoldersEdit.Reset();
-		CachePasteFolderExistingChildrenMap.Reset();
-		PendingFoldersSelect.Reset();
+	CacheFoldersEdit.Reset();
+	CachePasteFolderExistingChildrenMap.Reset();
+	PendingFoldersSelect.Reset();
 
-		CacheFoldersEdit = InFolders;
+	CacheFoldersEdit = InFolders;
 
-		// Sort folder names so parents appear before children
-		CacheFoldersEdit.Sort(FNameLexicalLess());
+	// Sort folder names so parents appear before children
+	CacheFoldersEdit.Sort(FNameLexicalLess());
 
-		// Cache existing children
-		for (FName Folder : CacheFoldersEdit)
+	// Cache existing children
+	for (FName Folder : CacheFoldersEdit)
+	{
+		if (FSceneOutlinerTreeItemPtr* TreeItem = TreeItemMap.Find(Folder))
 		{
-			if (FTreeItemPtr* TreeItem = TreeItemMap.Find(Folder))
-			{
-				CacheExistingChildrenAction(*TreeItem);
-			}
+			CacheExistingChildrenAction(*TreeItem);
 		}
 	}
+}
 
-	void SSceneOutliner::PasteFoldersEnd()
-	{
-		const FScopedTransaction Transaction(NSLOCTEXT("UnrealEd", "PasteItems", "Paste Items"));
+void SSceneOutliner::PasteFoldersEnd()
+{
+	const FScopedTransaction Transaction(NSLOCTEXT("UnrealEd", "PasteItems", "Paste Items"));
 
-		// Create new folder
-		TMap<FName, FName> FolderMap;
-		for (FName Folder : CacheFoldersEdit)
+	// Create new folder
+	TMap<FName, FName> FolderMap;
+	for (FName Folder : CacheFoldersEdit)
+	{
+		FName ParentPath = SceneOutliner::GetParentPath(Folder);
+		FName LeafName = SceneOutliner::GetFolderLeafName(Folder);
+		if (LeafName != TEXT(""))
 		{
-			FName ParentPath = GetParentPath(Folder);
-			FName LeafName = GetFolderLeafName(Folder);
-			if (LeafName != TEXT(""))
+			if (FName* NewParentPath = FolderMap.Find(ParentPath))
 			{
-				if (FName* NewParentPath = FolderMap.Find(ParentPath))
-				{
-					ParentPath = *NewParentPath;
-				}
-
-				FName NewFolderPath = Mode->CreateFolder(ParentPath, LeafName);
-				FolderMap.Add(Folder, NewFolderPath);
+				ParentPath = *NewParentPath;
 			}
+
+			FName NewFolderPath = Mode->CreateFolder(ParentPath, LeafName);
+			FolderMap.Add(Folder, NewFolderPath);
 		}
+	}
 
-		// Populate our data set
-		Populate();
+	// Populate our data set
+	Populate();
 
-		// Reparent duplicated items if the folder has been pasted/duplicated
-		for (FName OldFolderName : CacheFoldersEdit)
+	// Reparent duplicated items if the folder has been pasted/duplicated
+	for (FName OldFolderName : CacheFoldersEdit)
+	{
+		// Get the new folder that was created from this name
+		if (const FName* NewFolderName = FolderMap.Find(OldFolderName))
 		{
-			// Get the new folder that was created from this name
-			if (const FName* NewFolderName = FolderMap.Find(OldFolderName))
+			if (FSceneOutlinerTreeItemPtr* OldFolderItem = TreeItemMap.Find(OldFolderName))
 			{
-				if (FTreeItemPtr* OldFolderItem = TreeItemMap.Find(OldFolderName))
+				for (const TWeakPtr<ISceneOutlinerTreeItem>& Child : (*OldFolderItem)->GetChildren())
 				{
-					for (const TWeakPtr<ITreeItem>& Child : (*OldFolderItem)->GetChildren())
-					{
-						// If this child did not exist in the folder before the paste operation, it should be moved to the new folder
-						TArray<FTreeItemID>* ExistingChildren = CachePasteFolderExistingChildrenMap.Find(OldFolderName);
+					// If this child did not exist in the folder before the paste operation, it should be moved to the new folder
+					TArray<FSceneOutlinerTreeItemID>* ExistingChildren = CachePasteFolderExistingChildrenMap.Find(OldFolderName);
 
-						if (ExistingChildren && !ExistingChildren->Contains(Child.Pin()->GetID()))
-						{
-							Mode->ReparentItemToFolder(*NewFolderName, Child.Pin());
-						}
+					if (ExistingChildren && !ExistingChildren->Contains(Child.Pin()->GetID()))
+					{
+						Mode->ReparentItemToFolder(*NewFolderName, Child.Pin());
 					}
 				}
-				PendingFoldersSelect.Add(*NewFolderName);
 			}
+			PendingFoldersSelect.Add(*NewFolderName);
 		}
-
-		CacheFoldersEdit.Reset();
-		CachePasteFolderExistingChildrenMap.Reset();
-		FullRefresh();
 	}
 
-	void SSceneOutliner::DuplicateFoldersHierarchy()
+	CacheFoldersEdit.Reset();
+	CachePasteFolderExistingChildrenMap.Reset();
+	FullRefresh();
+}
+
+void SSceneOutliner::DuplicateFoldersHierarchy()
+{
+	TFunctionRef<void(const FSceneOutlinerTreeItemPtr&)> RecursiveFolderSelect = [&](const FSceneOutlinerTreeItemPtr& Item)
 	{
-		TFunctionRef<void(const FTreeItemPtr&)> RecursiveFolderSelect = [&](const FTreeItemPtr& Item)
+		if (Item->IsA<FFolderTreeItem>())
 		{
-			if (Item->IsA<FFolderTreeItem>())
-			{
-				OutlinerTreeView->SetItemSelection(Item, true);
-			}
-			for (const TWeakPtr<ITreeItem> Child : Item->GetChildren())
-			{
-				RecursiveFolderSelect(Child.Pin());
-			}
-		};
+			OutlinerTreeView->SetItemSelection(Item, true);
+		}
+		for (const TWeakPtr<ISceneOutlinerTreeItem> Child : Item->GetChildren())
+		{
+			RecursiveFolderSelect(Child.Pin());
+		}
+	};
+
+	const FScopedTransaction Transaction(NSLOCTEXT("UnrealEd", "DuplicateFoldersHierarchy", "Duplicate Folders Hierarchy"));
+
+	TArray<FFolderTreeItem*> SelectedFolders;
+	GetSelectedFolders(SelectedFolders);
+
+	if (SelectedFolders.Num() > 0)
+	{
+		// Select folder descendants
+		SelectFoldersDescendants();
+
+		// Select all sub-folders
+		for (FFolderTreeItem* Folder : SelectedFolders)
+		{
+			RecursiveFolderSelect(Folder->AsShared());
+		}
+
+		// Duplicate selected
+		Mode->OnDuplicateSelected();
+	}
+}
+
+void SSceneOutliner::DeleteFoldersBegin()
+{
+	CacheFoldersDelete.Empty();
+	GetSelectedFolders(CacheFoldersDelete);
+}
+
+void SSceneOutliner::DeleteFoldersEnd()
+{
+	struct FMatchName
+	{
+		FMatchName(const FName InPathName)
+			: PathName(InPathName) {}
 
-		const FScopedTransaction Transaction(NSLOCTEXT("UnrealEd", "DuplicateFoldersHierarchy", "Duplicate Folders Hierarchy"));
+		const FName PathName;
 
-		TArray<FFolderTreeItem*> SelectedFolders;
-		GetSelectedFolders(SelectedFolders);
+		bool operator()(const FFolderTreeItem *Entry)
+		{
+			return PathName == Entry->Path;
+		}
+	};
 
-		if (SelectedFolders.Num() > 0)
+	if (CacheFoldersDelete.Num() > 0)
+	{
+		// Sort in descending order so children will be deleted before parents
+		CacheFoldersDelete.Sort([](const FFolderTreeItem& FolderA, const FFolderTreeItem& FolderB)
 		{
-			// Select folder descendants
-			SelectFoldersDescendants();
+			return FolderB.Path.LexicalLess(FolderA.Path);
+		});
 
-			// Select all sub-folders
-			for (FFolderTreeItem* Folder : SelectedFolders)
+		for (FFolderTreeItem* Folder : CacheFoldersDelete)
+		{
+			if (Folder)
 			{
-				RecursiveFolderSelect(Folder->AsShared());
-			}
+				// Find lowest parent not being deleted, for reparenting children of current folder
+				FName NewParentPath = SceneOutliner::GetParentPath(Folder->Path);
+				while (!NewParentPath.IsNone() && CacheFoldersDelete.FindByPredicate(FMatchName(NewParentPath)))
+				{
+					NewParentPath = SceneOutliner::GetParentPath(NewParentPath);
+				}
 
-			// Duplicate selected
-			Mode->OnDuplicateSelected();
+				Folder->Delete(NewParentPath);
+			}
 		}
+
+		CacheFoldersDelete.Reset();
+		FullRefresh();
 	}
+}
+
+TArray<FName> SSceneOutliner::GetClipboardPasteFolders() const
+{
+	FString PasteString;
+	FPlatformApplicationMisc::ClipboardPaste(PasteString);
+	return ImportFolderList(*PasteString);
+}
+
+FString SSceneOutliner::ExportFolderList(TArray<FName> InFolders) const
+{
+	FString Buffer = FString(TEXT("Begin FolderList\n"));
 
-	void SSceneOutliner::DeleteFoldersBegin()
+	for (auto& FolderName : InFolders)
 	{
-		CacheFoldersDelete.Empty();
-		GetSelectedFolders(CacheFoldersDelete);
+		Buffer.Append(FString(TEXT("\tFolder=")) + FolderName.ToString() + FString(TEXT("\n")));
 	}
 
-	void SSceneOutliner::DeleteFoldersEnd()
+	Buffer += FString(TEXT("End FolderList\n"));
+
+	return Buffer;
+}
+
+TArray<FName> SSceneOutliner::ImportFolderList(const FString& InStrBuffer) const
+{
+	TArray<FName> Folders;
+
+	int32 Index = InStrBuffer.Find(TEXT("Begin FolderList"));
+	if (Index != INDEX_NONE)
 	{
-		struct FMatchName
-		{
-			FMatchName(const FName InPathName)
-				: PathName(InPathName) {}
+		FString TmpStr = InStrBuffer.RightChop(Index);
+		const TCHAR* Buffer = *TmpStr;
 
-			const FName PathName;
+		FString StrLine;
+		while (FParse::Line(&Buffer, StrLine))
+		{
+			const TCHAR* Str = *StrLine;				
+			FString FolderName;
 
-			bool operator()(const FFolderTreeItem *Entry)
+			if (FParse::Command(&Str, TEXT("Begin")) && FParse::Command(&Str, TEXT("FolderList")))
 			{
-				return PathName == Entry->Path;
+				continue;
 			}
-		};
-
-		if (CacheFoldersDelete.Num() > 0)
-		{
-			// Sort in descending order so children will be deleted before parents
-			CacheFoldersDelete.Sort([](const FFolderTreeItem& FolderA, const FFolderTreeItem& FolderB)
+			else if (FParse::Command(&Str, TEXT("End")) && FParse::Command(&Str, TEXT("FolderList")))
 			{
-				return FolderB.Path.LexicalLess(FolderA.Path);
-			});
-
-			for (FFolderTreeItem* Folder : CacheFoldersDelete)
+				break;
+			}
+			else if (FParse::Value(Str, TEXT("Folder="), FolderName))
 			{
-				if (Folder)
-				{
-					// Find lowest parent not being deleted, for reparenting children of current folder
-					FName NewParentPath = GetParentPath(Folder->Path);
-					while (!NewParentPath.IsNone() && CacheFoldersDelete.FindByPredicate(FMatchName(NewParentPath)))
-					{
-						NewParentPath = GetParentPath(NewParentPath);
-					}
-
-					Folder->Delete(NewParentPath);
-				}
+				Folders.Add(FName(*FolderName));
 			}
-
-			CacheFoldersDelete.Reset();
-			FullRefresh();
 		}
 	}
+	return Folders;
+}
 
-	TArray<FName> SSceneOutliner::GetClipboardPasteFolders() const
+void SSceneOutliner::ScrollItemIntoView(const FSceneOutlinerTreeItemPtr& Item)
+{
+	auto Parent = Item->GetParent();
+	while (Parent.IsValid())
 	{
-		FString PasteString;
-		FPlatformApplicationMisc::ClipboardPaste(PasteString);
-		return ImportFolderList(*PasteString);
+		OutlinerTreeView->SetItemExpansion(Parent->AsShared(), true);
+		Parent = Parent->GetParent();
 	}
+	OutlinerTreeView->RequestScrollIntoView(Item);
+}
 
-	FString SSceneOutliner::ExportFolderList(TArray<FName> InFolders) const
-	{
-		FString Buffer = FString(TEXT("Begin FolderList\n"));
-
-		for (auto& FolderName : InFolders)
-		{
-			Buffer.Append(FString(TEXT("\tFolder=")) + FolderName.ToString() + FString(TEXT("\n")));
-		}
+void SSceneOutliner::SetItemExpansion(const FSceneOutlinerTreeItemPtr& Item, bool bIsExpanded)
+{
+	OutlinerTreeView->SetItemExpansion(Item, bIsExpanded);
+}
 
-		Buffer += FString(TEXT("End FolderList\n"));
+bool SSceneOutliner::IsItemExpanded(const FSceneOutlinerTreeItemPtr& Item) const
+{
+	return OutlinerTreeView->IsItemExpanded(Item);
+}
 
-		return Buffer;
-	}
+TSharedRef< ITableRow > SSceneOutliner::OnGenerateRowForOutlinerTree( FSceneOutlinerTreeItemPtr Item, const TSharedRef< STableViewBase >& OwnerTable )
+{
+	return SNew( SSceneOutlinerTreeRow, OutlinerTreeView.ToSharedRef(), SharedThis(this) ).Item( Item );
+}
 
-	TArray<FName> SSceneOutliner::ImportFolderList(const FString& InStrBuffer) const
+void SSceneOutliner::OnGetChildrenForOutlinerTree( FSceneOutlinerTreeItemPtr InParent, TArray< FSceneOutlinerTreeItemPtr >& OutChildren )
+{
+	if( SharedData->bShowParentTree )
 	{
-		TArray<FName> Folders;
+		for (auto& WeakChild : InParent->GetChildren())
+		{
+			auto Child = WeakChild.Pin();
+			// Should never have bogus entries in this list
+			check(Child.IsValid());
+			OutChildren.Add(Child);
+		}
 
-		int32 Index = InStrBuffer.Find(TEXT("Begin FolderList"));
-		if (Index != INDEX_NONE)
+		// If the item needs it's children sorting, do that now
+		if (OutChildren.Num() && InParent->Flags.bChildrenRequireSort)
 		{
-			FString TmpStr = InStrBuffer.RightChop(Index);
-			const TCHAR* Buffer = *TmpStr;
+			// Sort the children we returned
+			SortItems(OutChildren);
 
-			FString StrLine;
-			while (FParse::Line(&Buffer, StrLine))
+			// Empty out the children and repopulate them in the correct order
+			InParent->Children.Empty();
+			for (auto& Child : OutChildren)
 			{
-				const TCHAR* Str = *StrLine;				
-				FString FolderName;
-
-				if (FParse::Command(&Str, TEXT("Begin")) && FParse::Command(&Str, TEXT("FolderList")))
-				{
-					continue;
-				}
-				else if (FParse::Command(&Str, TEXT("End")) && FParse::Command(&Str, TEXT("FolderList")))
-				{
-					break;
-				}
-				else if (FParse::Value(Str, TEXT("Folder="), FolderName))
-				{
-					Folders.Add(FName(*FolderName));
-				}
+				InParent->Children.Emplace(Child);
 			}
+				
+			// They no longer need sorting
+			InParent->Flags.bChildrenRequireSort = false;
 		}
-		return Folders;
 	}
+}
 
-	void SSceneOutliner::ScrollItemIntoView(const FTreeItemPtr& Item)
+void SSceneOutliner::OnOutlinerTreeSelectionChanged( FSceneOutlinerTreeItemPtr TreeItem, ESelectInfo::Type SelectInfo )
+{
+	if (SelectInfo == ESelectInfo::Direct)
 	{
-		auto Parent = Item->GetParent();
-		while (Parent.IsValid())
-		{
-			OutlinerTreeView->SetItemExpansion(Parent->AsShared(), true);
-			Parent = Parent->GetParent();
-		}
-		OutlinerTreeView->RequestScrollIntoView(Item);
+		return;
 	}
 
-	void SSceneOutliner::SetItemExpansion(const FTreeItemPtr& Item, bool bIsExpanded)
+	if (!bIsReentrant)
 	{
-		OutlinerTreeView->SetItemExpansion(Item, bIsExpanded);
-	}
+		TGuardValue<bool> ReentrantGuard(bIsReentrant, true);
+		Mode->OnItemSelectionChanged(TreeItem, SelectInfo, FSceneOutlinerItemSelection(*OutlinerTreeView));
 
-	bool SSceneOutliner::IsItemExpanded(const FTreeItemPtr& Item) const
-	{
-		return OutlinerTreeView->IsItemExpanded(Item);
+		OnItemSelectionChanged.Broadcast(TreeItem, SelectInfo);
 	}
+}
 
-	TSharedRef< ITableRow > SSceneOutliner::OnGenerateRowForOutlinerTree( FTreeItemPtr Item, const TSharedRef< STableViewBase >& OwnerTable )
+void SSceneOutliner::OnOutlinerTreeDoubleClick( FSceneOutlinerTreeItemPtr TreeItem )
+{
+	if (TreeItem->IsA<FFolderTreeItem>())
 	{
-		return SNew( SSceneOutlinerTreeRow, OutlinerTreeView.ToSharedRef(), SharedThis(this) ).Item( Item );
+		SetItemExpansion(TreeItem, !IsItemExpanded(TreeItem));
 	}
 
-	void SSceneOutliner::OnGetChildrenForOutlinerTree( FTreeItemPtr InParent, TArray< FTreeItemPtr >& OutChildren )
-	{
-		if( SharedData->bShowParentTree )
-		{
-			for (auto& WeakChild : InParent->GetChildren())
-			{
-				auto Child = WeakChild.Pin();
-				// Should never have bogus entries in this list
-				check(Child.IsValid());
-				OutChildren.Add(Child);
-			}
+	Mode->OnItemDoubleClick(TreeItem);
 
-			// If the item needs it's children sorting, do that now
-			if (OutChildren.Num() && InParent->Flags.bChildrenRequireSort)
-			{
-				// Sort the children we returned
-				SortItems(OutChildren);
+	OnDoubleClickOnTreeEvent.Broadcast(TreeItem);
+}
 
-				// Empty out the children and repopulate them in the correct order
-				InParent->Children.Empty();
-				for (auto& Child : OutChildren)
-				{
-					InParent->Children.Emplace(Child);
-				}
-				
-				// They no longer need sorting
-				InParent->Flags.bChildrenRequireSort = false;
-			}
-		}
+void SSceneOutliner::OnOutlinerTreeItemScrolledIntoView( FSceneOutlinerTreeItemPtr TreeItem, const TSharedPtr<ITableRow>& Widget )
+{
+	if (TreeItem == PendingRenameItem.Pin())
+	{
+		PendingRenameItem = nullptr;
+		TreeItem->RenameRequestEvent.ExecuteIfBound();
 	}
+}
+	
+void SSceneOutliner::OnItemExpansionChanged(FSceneOutlinerTreeItemPtr TreeItem, bool bIsExpanded) const
+{
+	TreeItem->Flags.bIsExpanded = bIsExpanded;
+	TreeItem->OnExpansionChanged();
 
-	void SSceneOutliner::OnOutlinerTreeSelectionChanged( FTreeItemPtr TreeItem, ESelectInfo::Type SelectInfo )
+	// Expand any children that are also expanded
+	for (auto WeakChild : TreeItem->GetChildren())
 	{
-		if (SelectInfo == ESelectInfo::Direct)
-		{
-			return;
-		}
-
-		if (!bIsReentrant)
+		auto Child = WeakChild.Pin();
+		if (Child->Flags.bIsExpanded)
 		{
-			TGuardValue<bool> ReentrantGuard(bIsReentrant, true);
-			Mode->OnItemSelectionChanged(TreeItem, SelectInfo, FItemSelection(*OutlinerTreeView));
-
-			OnItemSelectionChanged.Broadcast(TreeItem, SelectInfo);
+			OutlinerTreeView->SetItemExpansion(Child, true);
 		}
 	}
+}
 
-	void SSceneOutliner::OnOutlinerTreeDoubleClick( FTreeItemPtr TreeItem )
+void SSceneOutliner::OnHierarchyChangedEvent(FSceneOutlinerHierarchyChangedData Event)
+{
+	if (Event.Type == FSceneOutlinerHierarchyChangedData::Added)
 	{
-		if (TreeItem->IsA<FFolderTreeItem>())
+		if (Event.Item.IsValid() && !TreeItemMap.Find(Event.Item->GetID()))
 		{
-			SetItemExpansion(TreeItem, !IsItemExpanded(TreeItem));
+			AddPendingItemAndChildren(Event.Item);
+			if (Event.ItemActions)
+			{
+				NewItemActions.Add(Event.Item->GetID(), Event.ItemActions);
+			}
 		}
-
-		Mode->OnItemDoubleClick(TreeItem);
-
-		OnDoubleClickOnTreeEvent.Broadcast(TreeItem);
 	}
-
-	void SSceneOutliner::OnOutlinerTreeItemScrolledIntoView( FTreeItemPtr TreeItem, const TSharedPtr<ITableRow>& Widget )
+	else if (Event.Type == FSceneOutlinerHierarchyChangedData::Removed)
 	{
-		if (TreeItem == PendingRenameItem.Pin())
+		FSceneOutlinerTreeItemPtr* Item = TreeItemMap.Find(Event.ItemID);
+		if (!Item)
 		{
-			PendingRenameItem = nullptr;
-			TreeItem->RenameRequestEvent.ExecuteIfBound();
+			Item = PendingTreeItemMap.Find(Event.ItemID);
 		}
-	}
-	
-	void SSceneOutliner::OnItemExpansionChanged(FTreeItemPtr TreeItem, bool bIsExpanded) const
-	{
-		TreeItem->Flags.bIsExpanded = bIsExpanded;
-		TreeItem->OnExpansionChanged();
 
-		// Expand any children that are also expanded
-		for (auto WeakChild : TreeItem->GetChildren())
+		if (Item)
 		{
-			auto Child = WeakChild.Pin();
-			if (Child->Flags.bIsExpanded)
-			{
-				OutlinerTreeView->SetItemExpansion(Child, true);
-			}
+			PendingOperations.Emplace(SceneOutliner::FPendingTreeOperation::Removed, Item->ToSharedRef());
+			Refresh();
 		}
 	}
-
-	void SSceneOutliner::OnHierarchyChangedEvent(FHierarchyChangedData Event)
+	else if (Event.Type == FSceneOutlinerHierarchyChangedData::Moved)
 	{
-		if (Event.Type == FHierarchyChangedData::Added)
+		FSceneOutlinerTreeItemPtr* Item = Event.Item ? &Event.Item : TreeItemMap.Find(Event.ItemID);
+		if (!Item)
 		{
-			if (Event.Item.IsValid() && !TreeItemMap.Find(Event.Item->GetID()))
-			{
-				AddPendingItemAndChildren(Event.Item);
-				if (Event.ItemActions)
-				{
-					NewItemActions.Add(Event.Item->GetID(), Event.ItemActions);
-				}
-			}
+			Item = PendingTreeItemMap.Find(Event.ItemID);
 		}
-		else if (Event.Type == FHierarchyChangedData::Removed)
-		{
-			FTreeItemPtr* Item = TreeItemMap.Find(Event.ItemID);
-			if (!Item)
-			{
-				Item = PendingTreeItemMap.Find(Event.ItemID);
-			}
 
-			if (Item)
-			{
-				PendingOperations.Emplace(FPendingTreeOperation::Removed, Item->ToSharedRef());
-				Refresh();
-			}
-		}
-		else if (Event.Type == FHierarchyChangedData::Moved)
+		if (Item)
 		{
-			FTreeItemPtr* Item = Event.Item ? &Event.Item : TreeItemMap.Find(Event.ItemID);
-			if (!Item)
-			{
-				Item = PendingTreeItemMap.Find(Event.ItemID);
-			}
-
-			if (Item)
-			{
-				PendingOperations.Emplace(FPendingTreeOperation::Moved, Item->ToSharedRef());
-				Refresh();
-			}
+			PendingOperations.Emplace(SceneOutliner::FPendingTreeOperation::Moved, Item->ToSharedRef());
+			Refresh();
 		}
-		else if (Event.Type == FHierarchyChangedData::FolderMoved)
+	}
+	else if (Event.Type == FSceneOutlinerHierarchyChangedData::FolderMoved)
+	{
+		FSceneOutlinerTreeItemPtr Item = TreeItemMap.FindRef(Event.ItemID);
+		if (Item.IsValid())
 		{
-			FTreeItemPtr Item = TreeItemMap.FindRef(Event.ItemID);
-			if (Item.IsValid())
-			{
-				// Remove it from the map under the old ID (which is derived from the folder path)
-				TreeItemMap.Remove(Item->GetID());
+			// Remove it from the map under the old ID (which is derived from the folder path)
+			TreeItemMap.Remove(Item->GetID());
 
-				// Now change the path and put it back in the map with its new ID
-				auto Folder = StaticCastSharedPtr<FFolderTreeItem>(Item);
-				Folder->Path = Event.NewPath;
-				Folder->LeafName = GetFolderLeafName(Event.NewPath);
+			// Now change the path and put it back in the map with its new ID
+			auto Folder = StaticCastSharedPtr<FFolderTreeItem>(Item);
+			Folder->Path = Event.NewPath;
+			Folder->LeafName = SceneOutliner::GetFolderLeafName(Event.NewPath);
 
-				TreeItemMap.Add(Item->GetID(), Item);
+			TreeItemMap.Add(Item->GetID(), Item);
 
-				// Add an operation to move the item in the hierarchy
-				PendingOperations.Emplace(FPendingTreeOperation::Moved, Item.ToSharedRef());
-				Refresh();
-			}
-		}
-		else if (Event.Type == FHierarchyChangedData::FullRefresh)
-		{
-			FullRefresh();
+			// Add an operation to move the item in the hierarchy
+			PendingOperations.Emplace(SceneOutliner::FPendingTreeOperation::Moved, Item.ToSharedRef());
+			Refresh();
 		}
 	}
+	else if (Event.Type == FSceneOutlinerHierarchyChangedData::FullRefresh)
+	{
+		FullRefresh();
+	}
+}
 
-	void SSceneOutliner::PostUndo(bool bSuccess)
+void SSceneOutliner::PostUndo(bool bSuccess)
+{
+	// Refresh our tree in case any changes have been made to the scene that might effect our list
+	if( !bIsReentrant )
 	{
-		// Refresh our tree in case any changes have been made to the scene that might effect our list
-		if( !bIsReentrant )
-		{
-            bDisableIntermediateSorting = true;
-			FullRefresh();
-		}
+        bDisableIntermediateSorting = true;
+		FullRefresh();
 	}
+}
 
-	void SSceneOutliner::OnItemLabelChanged(FTreeItemPtr ChangedItem)
+void SSceneOutliner::OnItemLabelChanged(FSceneOutlinerTreeItemPtr ChangedItem)
+{
+	// If the item already exists
+	if (FSceneOutlinerTreeItemPtr* ExistingItem = TreeItemMap.Find(ChangedItem->GetID()))
 	{
-		// If the item already exists
-		if (FTreeItemPtr* ExistingItem = TreeItemMap.Find(ChangedItem->GetID()))
+		// The changed item flags will have been set already
+		if (!ChangedItem->Flags.bIsFilteredOut)
 		{
-			// The changed item flags will have been set already
-			if (!ChangedItem->Flags.bIsFilteredOut)
-			{
-				OutlinerTreeView->FlashHighlightOnItem(*ExistingItem);
-				RequestSort();
-			}
-			else
-			{
-				// No longer matches the filters, remove it
-				PendingOperations.Emplace(FPendingTreeOperation::Removed, ExistingItem->ToSharedRef());
-				Refresh();
-			}
+			OutlinerTreeView->FlashHighlightOnItem(*ExistingItem);
+			RequestSort();
 		}
 		else
 		{
-			// Attempt to add the item if we didn't find it - perhaps it now matches the filter?
-			if (ChangedItem && !ChangedItem->Flags.bIsFilteredOut)
-			{
-				AddPendingItemAndChildren(ChangedItem);
-			}
+			// No longer matches the filters, remove it
+			PendingOperations.Emplace(SceneOutliner::FPendingTreeOperation::Removed, ExistingItem->ToSharedRef());
+			Refresh();
 		}
 	}
-
-	void SSceneOutliner::OnAssetReloaded(const EPackageReloadPhase InPackageReloadPhase, FPackageReloadedEvent* InPackageReloadedEvent)
+	else
 	{
-		if (InPackageReloadPhase == EPackageReloadPhase::PostBatchPostGC)
+		// Attempt to add the item if we didn't find it - perhaps it now matches the filter?
+		if (ChangedItem && !ChangedItem->Flags.bIsFilteredOut)
 		{
-			// perhaps overkill but a simple Refresh() doesn't appear to work.
-			FullRefresh();
+			AddPendingItemAndChildren(ChangedItem);
 		}
 	}
+}
 
-	void SSceneOutliner::OnFilterTextChanged( const FText& InFilterText )
+void SSceneOutliner::OnAssetReloaded(const EPackageReloadPhase InPackageReloadPhase, FPackageReloadedEvent* InPackageReloadedEvent)
+{
+	if (InPackageReloadPhase == EPackageReloadPhase::PostBatchPostGC)
 	{
-		SearchBoxFilter->SetRawFilterText( InFilterText );
-		FilterTextBoxWidget->SetError( SearchBoxFilter->GetFilterErrorText() );
-
-		Mode->OnFilterTextChanged(InFilterText);
+		// perhaps overkill but a simple Refresh() doesn't appear to work.
+		FullRefresh();
 	}
+}
 
-	void SSceneOutliner::OnFilterTextCommitted( const FText& InFilterText, ETextCommit::Type CommitInfo )
+void SSceneOutliner::OnFilterTextChanged( const FText& InFilterText )
+{
+	SearchBoxFilter->SetRawFilterText( InFilterText );
+	FilterTextBoxWidget->SetError( SearchBoxFilter->GetFilterErrorText() );
+
+	Mode->OnFilterTextChanged(InFilterText);
+}
+
+void SSceneOutliner::OnFilterTextCommitted( const FText& InFilterText, ETextCommit::Type CommitInfo )
+{
+	const FString CurrentFilterText = InFilterText.ToString();
+	// We'll only select items if the user actually pressed the enter key. We don't want to change
+	// selection just because focus was lost from the search text field.
+	if( CommitInfo == ETextCommit::OnEnter )
 	{
-		const FString CurrentFilterText = InFilterText.ToString();
-		// We'll only select items if the user actually pressed the enter key. We don't want to change
-		// selection just because focus was lost from the search text field.
-		if( CommitInfo == ETextCommit::OnEnter )
+		// Any text in the filter?  If not, we won't bother doing anything
+		if( !CurrentFilterText.IsEmpty() )
 		{
-			// Any text in the filter?  If not, we won't bother doing anything
-			if( !CurrentFilterText.IsEmpty() )
-			{
-				FItemSelection Selection;
+			FSceneOutlinerItemSelection Selection;
 
-				// Gather all of the items that match the filter text
-				for (auto& Pair : TreeItemMap)
+			// Gather all of the items that match the filter text
+			for (auto& Pair : TreeItemMap)
+			{
+				if (!Pair.Value->Flags.bIsFilteredOut)
 				{
-					if (!Pair.Value->Flags.bIsFilteredOut)
-					{
-						Selection.Add(Pair.Value);
-					}
+					Selection.Add(Pair.Value);
 				}
-
-				Mode->OnFilterTextCommited(Selection, CommitInfo);
 			}
-		}
-		else if (CommitInfo == ETextCommit::OnCleared)
-		{
-			OnFilterTextChanged(InFilterText);
+
+			Mode->OnFilterTextCommited(Selection, CommitInfo);
 		}
 	}
-
-	EVisibility SSceneOutliner::GetFilterStatusVisibility() const
+	else if (CommitInfo == ETextCommit::OnCleared)
 	{
-		return IsTextFilterActive() ? EVisibility::Visible : EVisibility::Collapsed;
+		OnFilterTextChanged(InFilterText);
 	}
+}
 
-	EVisibility SSceneOutliner::GetEmptyLabelVisibility() const
-	{
-		return ( IsTextFilterActive() || RootTreeItems.Num() > 0 ) ? EVisibility::Collapsed : EVisibility::Visible;
-	}
+EVisibility SSceneOutliner::GetFilterStatusVisibility() const
+{
+	return IsTextFilterActive() ? EVisibility::Visible : EVisibility::Collapsed;
+}
 
-	FText SSceneOutliner::GetFilterStatusText() const
-	{
-		return Mode->GetStatusText();
-	}
+EVisibility SSceneOutliner::GetEmptyLabelVisibility() const
+{
+	return ( IsTextFilterActive() || RootTreeItems.Num() > 0 ) ? EVisibility::Collapsed : EVisibility::Visible;
+}
 
-	FSlateColor SSceneOutliner::GetFilterStatusTextColor() const
-	{
-		return Mode->GetStatusTextColor();
-	}
+FText SSceneOutliner::GetFilterStatusText() const
+{
+	return Mode->GetStatusText();
+}
+
+FSlateColor SSceneOutliner::GetFilterStatusTextColor() const
+{
+	return Mode->GetStatusTextColor();
+}
 
-	bool SSceneOutliner::IsTextFilterActive() const
+bool SSceneOutliner::IsTextFilterActive() const
+{
+	return FilterTextBoxWidget->GetText().ToString().Len() > 0;
+}
+
+const FSlateBrush* SSceneOutliner::GetFilterButtonGlyph() const
+{
+	if( IsTextFilterActive() )
 	{
-		return FilterTextBoxWidget->GetText().ToString().Len() > 0;
+		return FEditorStyle::GetBrush(TEXT("SceneOutliner.FilterCancel"));
 	}
-
-	const FSlateBrush* SSceneOutliner::GetFilterButtonGlyph() const
+	else
 	{
-		if( IsTextFilterActive() )
-		{
-			return FEditorStyle::GetBrush(TEXT("SceneOutliner.FilterCancel"));
-		}
-		else
-		{
-			return FEditorStyle::GetBrush(TEXT("SceneOutliner.FilterSearch"));
-		}
+		return FEditorStyle::GetBrush(TEXT("SceneOutliner.FilterSearch"));
 	}
+}
 
-	FString SSceneOutliner::GetFilterButtonToolTip() const
-	{
-		return IsTextFilterActive() ? LOCTEXT("ClearSearchFilter", "Clear search filter").ToString() : LOCTEXT("StartSearching", "Search").ToString();
+FString SSceneOutliner::GetFilterButtonToolTip() const
+{
+	return IsTextFilterActive() ? LOCTEXT("ClearSearchFilter", "Clear search filter").ToString() : LOCTEXT("StartSearching", "Search").ToString();
 
-	}
+}
 
-	TAttribute<FText> SSceneOutliner::GetFilterHighlightText() const
+TAttribute<FText> SSceneOutliner::GetFilterHighlightText() const
+{
+	return TAttribute<FText>::Create(TAttribute<FText>::FGetter::CreateStatic([](TWeakPtr<SceneOutliner::TreeItemTextFilter> Filter){
+		auto FilterPtr = Filter.Pin();
+		return FilterPtr.IsValid() ? FilterPtr->GetRawFilterText() : FText();
+	}, TWeakPtr<SceneOutliner::TreeItemTextFilter>(SearchBoxFilter)));
+}
+
+void SSceneOutliner::SetKeyboardFocus()
+{
+	if (SupportsKeyboardFocus())
 	{
-		return TAttribute<FText>::Create(TAttribute<FText>::FGetter::CreateStatic([](TWeakPtr<TreeItemTextFilter> Filter){
-			auto FilterPtr = Filter.Pin();
-			return FilterPtr.IsValid() ? FilterPtr->GetRawFilterText() : FText();
-		}, TWeakPtr<TreeItemTextFilter>(SearchBoxFilter)));
+		FWidgetPath OutlinerTreeViewWidgetPath;
+		// NOTE: Careful, GeneratePathToWidget can be reentrant in that it can call visibility delegates and such
+		FSlateApplication::Get().GeneratePathToWidgetUnchecked( OutlinerTreeView.ToSharedRef(), OutlinerTreeViewWidgetPath );
+		FSlateApplication::Get().SetKeyboardFocus( OutlinerTreeViewWidgetPath, EFocusCause::SetDirectly );
 	}
+}
 
-	void SSceneOutliner::SetKeyboardFocus()
+const FSlateBrush* SSceneOutliner::GetCachedIconForClass(FName InClassName) const
+{ 
+	if (CachedIcons.Find(InClassName))
 	{
-		if (SupportsKeyboardFocus())
-		{
-			FWidgetPath OutlinerTreeViewWidgetPath;
-			// NOTE: Careful, GeneratePathToWidget can be reentrant in that it can call visibility delegates and such
-			FSlateApplication::Get().GeneratePathToWidgetUnchecked( OutlinerTreeView.ToSharedRef(), OutlinerTreeViewWidgetPath );
-			FSlateApplication::Get().SetKeyboardFocus( OutlinerTreeViewWidgetPath, EFocusCause::SetDirectly );
-		}
+		return *CachedIcons.Find(InClassName);
 	}
-
-	const FSlateBrush* SSceneOutliner::GetCachedIconForClass(FName InClassName) const
-	{ 
-		if (CachedIcons.Find(InClassName))
-		{
-			return *CachedIcons.Find(InClassName);
-		}
-		else
-		{
-			return nullptr;
-		}
+	else
+	{
+		return nullptr;
 	}
+}
+
+void SSceneOutliner::CacheIconForClass(FName InClassName, const FSlateBrush* InSlateBrush)
+{
+	CachedIcons.Emplace(InClassName, InSlateBrush);
+}
+
+bool SSceneOutliner::SupportsKeyboardFocus() const
+{
+	return Mode->SupportsKeyboardFocus();
+}
+
+FReply SSceneOutliner::OnKeyDown( const FGeometry& MyGeometry, const FKeyEvent& InKeyEvent )
+{
+	// @todo outliner: Use command system for these for discoverability? (allow bindings?)
+	return Mode->OnKeyDown(InKeyEvent);
+}
 
-	void SSceneOutliner::CacheIconForClass(FName InClassName, const FSlateBrush* InSlateBrush)
+void SSceneOutliner::Tick(const FGeometry& AllottedGeometry, const double InCurrentTime, const float InDeltaTime)
+{	
+	for (auto& Pair : Columns)
 	{
-		CachedIcons.Emplace(InClassName, InSlateBrush);
+		Pair.Value->Tick(InCurrentTime, InDeltaTime);
 	}
 
-	bool SSceneOutliner::SupportsKeyboardFocus() const
+	if ( bPendingFocusNextFrame && FilterTextBoxWidget->GetVisibility() == EVisibility::Visible )
 	{
-		return Mode->SupportsKeyboardFocus();
+		FWidgetPath WidgetToFocusPath;
+		FSlateApplication::Get().GeneratePathToWidgetUnchecked( FilterTextBoxWidget.ToSharedRef(), WidgetToFocusPath );
+		FSlateApplication::Get().SetKeyboardFocus( WidgetToFocusPath, EFocusCause::SetDirectly );
+		bPendingFocusNextFrame = false;
 	}
 
-	FReply SSceneOutliner::OnKeyDown( const FGeometry& MyGeometry, const FKeyEvent& InKeyEvent )
+	if ( bNeedsColumRefresh )
 	{
-		// @todo outliner: Use command system for these for discoverability? (allow bindings?)
-		return Mode->OnKeyDown(InKeyEvent);
+		SetupColumns(*HeaderRowWidget);
 	}
 
-	void SSceneOutliner::Tick(const FGeometry& AllottedGeometry, const double InCurrentTime, const float InDeltaTime)
-	{	
-		for (auto& Pair : Columns)
-		{
-			Pair.Value->Tick(InCurrentTime, InDeltaTime);
-		}
-
-		if ( bPendingFocusNextFrame && FilterTextBoxWidget->GetVisibility() == EVisibility::Visible )
-		{
-			FWidgetPath WidgetToFocusPath;
-			FSlateApplication::Get().GeneratePathToWidgetUnchecked( FilterTextBoxWidget.ToSharedRef(), WidgetToFocusPath );
-			FSlateApplication::Get().SetKeyboardFocus( WidgetToFocusPath, EFocusCause::SetDirectly );
-			bPendingFocusNextFrame = false;
-		}
-
-		if ( bNeedsColumRefresh )
-		{
-			SetupColumns(*HeaderRowWidget);
-		}
-
-		if( bNeedsRefresh )
-		{
-			if( !bIsReentrant )
-			{
-				Populate();
-			}
-		}
-		SortOutlinerTimer -= InDeltaTime;
-
-		// Delay sorting when in PIE
-		if (bSortDirty && (GEditor->PlayWorld == nullptr || SortOutlinerTimer <= 0))
+	if( bNeedsRefresh )
+	{
+		if( !bIsReentrant )
 		{
-			SortItems(RootTreeItems);
-			for (const auto& Pair : TreeItemMap)
-			{
-				Pair.Value->Flags.bChildrenRequireSort = true;
-			}
-
-			OutlinerTreeView->RequestTreeRefresh();
-			bSortDirty = false;
+			Populate();
 		}
+	}
+	SortOutlinerTimer -= InDeltaTime;
 
-		if (SortOutlinerTimer <= 0)
+	// Delay sorting when in PIE
+	if (bSortDirty && (GEditor->PlayWorld == nullptr || SortOutlinerTimer <= 0))
+	{
+		SortItems(RootTreeItems);
+		for (const auto& Pair : TreeItemMap)
 		{
-			SortOutlinerTimer = SCENE_OUTLINER_RESORT_TIMER;
+			Pair.Value->Flags.bChildrenRequireSort = true;
 		}
 
-		if (bSelectionDirty)
-		{
-			Mode->SynchronizeSelection();
-			bSelectionDirty = false;
-		}
+		OutlinerTreeView->RequestTreeRefresh();
+		bSortDirty = false;
 	}
 
-	EColumnSortMode::Type SSceneOutliner::GetColumnSortMode( const FName ColumnId ) const
+	if (SortOutlinerTimer <= 0)
 	{
-		if (SortByColumn == ColumnId)
-		{
-			auto Column = Columns.FindRef(ColumnId);
-			if (Column.IsValid() && Column->SupportsSorting())
-			{
-				return SortMode;
-			}
-		}
+		SortOutlinerTimer = SCENE_OUTLINER_RESORT_TIMER;
+	}
 
-		return EColumnSortMode::None;
+	if (bSelectionDirty)
+	{
+		Mode->SynchronizeSelection();
+		bSelectionDirty = false;
 	}
+}
 
-	void SSceneOutliner::OnColumnSortModeChanged( const EColumnSortPriority::Type SortPriority, const FName& ColumnId, const EColumnSortMode::Type InSortMode )
+EColumnSortMode::Type SSceneOutliner::GetColumnSortMode( const FName ColumnId ) const
+{
+	if (SortByColumn == ColumnId)
 	{
 		auto Column = Columns.FindRef(ColumnId);
-		if (!Column.IsValid() || !Column->SupportsSorting())
+		if (Column.IsValid() && Column->SupportsSorting())
 		{
-			return;
+			return SortMode;
 		}
-
-		SortByColumn = ColumnId;
-		SortMode = InSortMode;
-
-		RequestSort();
 	}
 
-	void SSceneOutliner::RequestSort()
-	{
-		bSortDirty = true;
-	}
+	return EColumnSortMode::None;
+}
 
-	void SSceneOutliner::SortItems(TArray<FTreeItemPtr>& Items) const
+void SSceneOutliner::OnColumnSortModeChanged( const EColumnSortPriority::Type SortPriority, const FName& ColumnId, const EColumnSortMode::Type InSortMode )
+{
+	auto Column = Columns.FindRef(ColumnId);
+	if (!Column.IsValid() || !Column->SupportsSorting())
 	{
-		auto Column = Columns.FindRef(SortByColumn);
-		if (Column.IsValid())
-		{
-			Column->SortItems(Items, SortMode);
-		}
+		return;
 	}
 
-	uint32 SSceneOutliner::GetTypeSortPriority(const ITreeItem& Item) const 
+	SortByColumn = ColumnId;
+	SortMode = InSortMode;
+
+	RequestSort();
+}
+
+void SSceneOutliner::RequestSort()
+{
+	bSortDirty = true;
+}
+
+void SSceneOutliner::SortItems(TArray<FSceneOutlinerTreeItemPtr>& Items) const
+{
+	auto Column = Columns.FindRef(SortByColumn);
+	if (Column.IsValid())
 	{
-		return Mode->GetTypeSortPriority(Item);
+		Column->SortItems(Items, SortMode);
 	}
+}
 
-	void SSceneOutliner::SetItemExpansionRecursive(FTreeItemPtr Model, bool bInExpansionState)
+uint32 SSceneOutliner::GetTypeSortPriority(const ISceneOutlinerTreeItem& Item) const 
+{
+	return Mode->GetTypeSortPriority(Item);
+}
+
+void SSceneOutliner::SetItemExpansionRecursive(FSceneOutlinerTreeItemPtr Model, bool bInExpansionState)
+{
+	if (Model.IsValid())
 	{
-		if (Model.IsValid())
+		OutlinerTreeView->SetItemExpansion(Model, bInExpansionState);
+		for (auto& Child : Model->Children)
 		{
-			OutlinerTreeView->SetItemExpansion(Model, bInExpansionState);
-			for (auto& Child : Model->Children)
+			if (Child.IsValid())
 			{
-				if (Child.IsValid())
-				{
-					SetItemExpansionRecursive(Child.Pin(), bInExpansionState);
-				}
+				SetItemExpansionRecursive(Child.Pin(), bInExpansionState);
 			}
 		}
 	}
+}
 
-	TSharedPtr<FDragDropOperation> SSceneOutliner::CreateDragDropOperation(const TArray<FTreeItemPtr>& InTreeItems) const 
-	{ 
-		return Mode->CreateDragDropOperation(InTreeItems); 
-	}
-
-	/** Parse a drag drop operation into a payload */
-	bool SSceneOutliner::ParseDragDrop(FDragDropPayload& OutPayload, const FDragDropOperation& Operation) const 
-	{
-		return Mode->ParseDragDrop(OutPayload, Operation);
-	}
-
-	/** Validate a drag drop operation on a drop target */
-	FDragValidationInfo SSceneOutliner::ValidateDrop(const ITreeItem& DropTarget, const FDragDropPayload& Payload) const 
-	{ 
-		return Mode->ValidateDrop(DropTarget, Payload); 
-	}
+TSharedPtr<FDragDropOperation> SSceneOutliner::CreateDragDropOperation(const TArray<FSceneOutlinerTreeItemPtr>& InTreeItems) const 
+{ 
+	return Mode->CreateDragDropOperation(InTreeItems); 
+}
 
-	/** Called when a payload is dropped onto a target */
-	void SSceneOutliner::OnDropPayload(ITreeItem& DropTarget, const FDragDropPayload& Payload, const FDragValidationInfo& ValidationInfo) const 
-	{ 
-		return Mode->OnDrop(DropTarget, Payload, ValidationInfo); 
-	}
-
-	/** Called when a payload is dragged over an item */
-	FReply SSceneOutliner::OnDragOverItem(const FDragDropEvent& Event, const ITreeItem& Item) const 
-	{ 
-		return Mode->OnDragOverItem(Event, Item); 
-	}
-}	// namespace SceneOutliner
+/** Parse a drag drop operation into a payload */
+bool SSceneOutliner::ParseDragDrop(FSceneOutlinerDragDropPayload& OutPayload, const FDragDropOperation& Operation) const 
+{
+	return Mode->ParseDragDrop(OutPayload, Operation);
+}
+
+/** Validate a drag drop operation on a drop target */
+FSceneOutlinerDragValidationInfo SSceneOutliner::ValidateDrop(const ISceneOutlinerTreeItem& DropTarget, const FSceneOutlinerDragDropPayload& Payload) const 
+{ 
+	return Mode->ValidateDrop(DropTarget, Payload); 
+}
+
+/** Called when a payload is dropped onto a target */
+void SSceneOutliner::OnDropPayload(ISceneOutlinerTreeItem& DropTarget, const FSceneOutlinerDragDropPayload& Payload, const FSceneOutlinerDragValidationInfo& ValidationInfo) const 
+{ 
+	return Mode->OnDrop(DropTarget, Payload, ValidationInfo); 
+}
+
+/** Called when a payload is dragged over an item */
+FReply SSceneOutliner::OnDragOverItem(const FDragDropEvent& Event, const ISceneOutlinerTreeItem& Item) const 
+{ 
+	return Mode->OnDragOverItem(Event, Item); 
+}
 
 #undef LOCTEXT_NAMESPACE
