Removed SceneOutliner namespace

diff --git a/Engine/Source/Editor/SceneOutliner/Private/SceneOutlinerActorInfoColumn.cpp b/Engine/Source/Editor/SceneOutliner/Private/SceneOutlinerActorInfoColumn.cpp
--- a/Engine/Source/Editor/SceneOutliner/Private/SceneOutlinerActorInfoColumn.cpp
+++ b/Engine/Source/Editor/SceneOutliner/Private/SceneOutlinerActorInfoColumn.cpp
@@ -18,16 +18,13 @@
 
 #define LOCTEXT_NAMESPACE "SceneOutlinerActorInfoColumn"
 
-namespace SceneOutliner
-{
-
 struct FGetInfo
 {
-	ECustomColumnMode::Type CurrentMode;
+	SceneOutliner::ECustomColumnMode::Type CurrentMode;
 
-	FGetInfo(ECustomColumnMode::Type InCurrentMode) : CurrentMode(InCurrentMode) {}
+	FGetInfo(SceneOutliner::ECustomColumnMode::Type InCurrentMode) : CurrentMode(InCurrentMode) {}
 
-	FString operator()(const ITreeItem& Item) const
+	FString operator()(const ISceneOutlinerTreeItem& Item) const
 	{
 		if (const FActorTreeItem* ActorItem = Item.CastTo<FActorTreeItem>())
 		{
@@ -39,22 +36,22 @@ struct FGetInfo
 
 			switch (CurrentMode)
 			{
-			case ECustomColumnMode::Class:
+			case SceneOutliner::ECustomColumnMode::Class:
 				return Actor->GetClass()->GetName();
 
-			case ECustomColumnMode::Level:
+			case SceneOutliner::ECustomColumnMode::Level:
 				return FPackageName::GetShortName(Actor->GetOutermost()->GetName());
 
-			case ECustomColumnMode::Socket:
+			case SceneOutliner::ECustomColumnMode::Socket:
 				return Actor->GetAttachParentSocketName().ToString();
 
-			case ECustomColumnMode::InternalName:
+			case SceneOutliner::ECustomColumnMode::InternalName:
 				return Actor->GetFName().ToString();
 
-			case ECustomColumnMode::UncachedLights:
+			case SceneOutliner::ECustomColumnMode::UncachedLights:
 				return FString::Printf(TEXT("%7d"), Actor->GetNumUncachedStaticLightingInteractions());
 
-			case ECustomColumnMode::Layer:
+			case SceneOutliner::ECustomColumnMode::Layer:
 			{
 				FString Result;
 				for (const auto& Layer : Actor->Layers)
@@ -68,7 +65,7 @@ struct FGetInfo
 				}
 				return Result;
 			}
-		case ECustomColumnMode::Mobility:
+		case SceneOutliner::ECustomColumnMode::Mobility:
 			{
 				FString Result;
 				USceneComponent* RootComponent = Actor->GetRootComponent();
@@ -98,7 +95,7 @@ struct FGetInfo
 		{
 			switch (CurrentMode)
 			{
-			case ECustomColumnMode::Class:
+			case SceneOutliner::ECustomColumnMode::Class:
 				return LOCTEXT("FolderTypeName", "Folder").ToString();
 
 			default:
@@ -109,7 +106,7 @@ struct FGetInfo
 		{
 			switch (CurrentMode)
 			{
-			case ECustomColumnMode::Class:
+			case SceneOutliner::ECustomColumnMode::Class:
 				return LOCTEXT("WorldTypeName", "World").ToString();
 
 			default:
@@ -126,10 +123,10 @@ struct FGetInfo
 
 			switch (CurrentMode)
 			{
-			case ECustomColumnMode::Class:
+			case SceneOutliner::ECustomColumnMode::Class:
 				return LOCTEXT("ComponentTypeName", "Component").ToString();
 
-			case ECustomColumnMode::InternalName:
+			case SceneOutliner::ECustomColumnMode::InternalName:
 				return Component->GetFName().ToString();
 
 			default:
@@ -142,9 +139,9 @@ struct FGetInfo
 };
 
 
-TArray< TSharedPtr< ECustomColumnMode::Type > > FActorInfoColumn::ModeOptions;
+TArray< TSharedPtr< SceneOutliner::ECustomColumnMode::Type > > FActorInfoColumn::ModeOptions;
 
-FActorInfoColumn::FActorInfoColumn( ISceneOutliner& Outliner, ECustomColumnMode::Type InDefaultCustomColumnMode )
+FActorInfoColumn::FActorInfoColumn( ISceneOutliner& Outliner, SceneOutliner::ECustomColumnMode::Type InDefaultCustomColumnMode )
 	: CurrentMode( InDefaultCustomColumnMode )
 	, SceneOutlinerWeak( StaticCastSharedRef<ISceneOutliner>(Outliner.AsShared()) )
 {
@@ -162,9 +159,9 @@ SHeaderRow::FColumn::FArguments FActorInfoColumn::ConstructHeaderRowColumn()
 {
 	if( ModeOptions.Num() == 0 )
 	{
-		for( ECustomColumnMode::Type CurMode : TEnumRange<ECustomColumnMode::Type>() )
+		for(SceneOutliner::ECustomColumnMode::Type CurMode : TEnumRange<SceneOutliner::ECustomColumnMode::Type>() )
 		{
-			ModeOptions.Add( MakeShareable( new ECustomColumnMode::Type( CurMode ) ) );
+			ModeOptions.Add( MakeShareable( new SceneOutliner::ECustomColumnMode::Type( CurMode ) ) );
 		}
 	}
 
@@ -178,7 +175,7 @@ SHeaderRow::FColumn::FArguments FActorInfoColumn::ConstructHeaderRowColumn()
 			.Padding(FMargin(5))
 			.BorderImage(FEditorStyle::GetBrush("Menu.Background"))
 			[
-				SNew(SListView<TSharedPtr<ECustomColumnMode::Type>>)
+				SNew(SListView<TSharedPtr<SceneOutliner::ECustomColumnMode::Type>>)
 				.ListItemsSource(&ModeOptions)
 				.SelectionMode(ESelectionMode::Single)
 				.OnGenerateRow( this, &FActorInfoColumn::MakeComboButtonItemWidget )
@@ -200,7 +197,7 @@ SHeaderRow::FColumn::FArguments FActorInfoColumn::ConstructHeaderRowColumn()
 		];
 }
 
-const TSharedRef< SWidget > FActorInfoColumn::ConstructRowWidget( FTreeItemRef TreeItem, const STableRow<FTreeItemPtr>& Row )
+const TSharedRef< SWidget > FActorInfoColumn::ConstructRowWidget( FSceneOutlinerTreeItemRef TreeItem, const STableRow<FSceneOutlinerTreeItemPtr>& Row )
 {
 	auto SceneOutliner = SceneOutlinerWeak.Pin();
 	check(SceneOutliner.IsValid());
@@ -208,7 +205,7 @@ const TSharedRef< SWidget > FActorInfoColumn::ConstructRowWidget( FTreeItemRef T
 	TSharedRef<SHorizontalBox> HorizontalBox = SNew(SHorizontalBox);
 
 	TSharedRef<STextBlock> MainText = SNew( STextBlock )
-		.Text( this, &FActorInfoColumn::GetTextForItem, TWeakPtr<ITreeItem>(TreeItem) ) 
+		.Text( this, &FActorInfoColumn::GetTextForItem, TWeakPtr<ISceneOutlinerTreeItem>(TreeItem) ) 
 		.HighlightText( SceneOutliner->GetFilterHighlightText() )
 		.ColorAndOpacity( FSlateColor::UseSubduedForeground() );
 
@@ -235,10 +232,10 @@ const TSharedRef< SWidget > FActorInfoColumn::ConstructRowWidget( FTreeItemRef T
 			// Make sure that the hyperlink shows as black (by multiplying black * desired color) when selected so it is readable against the orange background even if blue/green/etc... normally
 			SNew(SBorder)
 			.BorderImage(FEditorStyle::GetBrush("NoBorder"))
-			.ColorAndOpacity_Static([](TWeakPtr<const STableRow<FTreeItemPtr>> WeakRow)->FLinearColor{
+			.ColorAndOpacity_Static([](TWeakPtr<const STableRow<FSceneOutlinerTreeItemPtr>> WeakRow)->FLinearColor{
 				auto TableRow = WeakRow.Pin();
 				return TableRow.IsValid() && TableRow->IsSelected() ? FLinearColor::Black : FLinearColor::White;
-			}, TWeakPtr<const STableRow<FTreeItemPtr>>(StaticCastSharedRef<const STableRow<FTreeItemPtr>>(Row.AsShared())))
+			}, TWeakPtr<const STableRow<FSceneOutlinerTreeItemPtr>>(StaticCastSharedRef<const STableRow<FSceneOutlinerTreeItemPtr>>(Row.AsShared())))
 			[
 				Hyperlink.ToSharedRef()
 			]
@@ -248,7 +245,7 @@ const TSharedRef< SWidget > FActorInfoColumn::ConstructRowWidget( FTreeItemRef T
 	return HorizontalBox;
 }
 
-TSharedPtr<SWidget> FActorInfoColumn::ConstructClassHyperlink( ITreeItem& TreeItem )
+TSharedPtr<SWidget> FActorInfoColumn::ConstructClassHyperlink( ISceneOutlinerTreeItem& TreeItem )
 {
 	if (const FActorTreeItem* ActorItem = TreeItem.CastTo<FActorTreeItem>())
 	{
@@ -284,7 +281,7 @@ TSharedPtr<SWidget> FActorInfoColumn::ConstructClassHyperlink( ITreeItem& TreeIt
 	return nullptr;
 }
 
-void FActorInfoColumn::PopulateSearchStrings( const ITreeItem& Item, TArray< FString >& OutSearchStrings ) const
+void FActorInfoColumn::PopulateSearchStrings( const ISceneOutlinerTreeItem& Item, TArray< FString >& OutSearchStrings ) const
 {
 	{
 		FString String = FGetInfo(CurrentMode)(Item);
@@ -295,9 +292,9 @@ void FActorInfoColumn::PopulateSearchStrings( const ITreeItem& Item, TArray< FSt
 	}
 
 	// We always add the class
-	if (CurrentMode != ECustomColumnMode::Class)
+	if (CurrentMode != SceneOutliner::ECustomColumnMode::Class)
 	{
-		FString String = FGetInfo(ECustomColumnMode::Class)(Item);
+		FString String = FGetInfo(SceneOutliner::ECustomColumnMode::Class)(Item);
 		if (String.Len())
 		{
 			OutSearchStrings.Add(String);
@@ -307,17 +304,17 @@ void FActorInfoColumn::PopulateSearchStrings( const ITreeItem& Item, TArray< FSt
 
 bool FActorInfoColumn::SupportsSorting() const
 {
-	return CurrentMode != ECustomColumnMode::None;
+	return CurrentMode != SceneOutliner::ECustomColumnMode::None;
 }
 
-void FActorInfoColumn::SortItems(TArray<FTreeItemPtr>& RootItems, const EColumnSortMode::Type SortMode) const
+void FActorInfoColumn::SortItems(TArray<FSceneOutlinerTreeItemPtr>& RootItems, const EColumnSortMode::Type SortMode) const
 {
-	FSortHelper<FString>()
+	FSceneOutlinerSortHelper<FString>()
 		.Primary(FGetInfo(CurrentMode), SortMode)
 		.Sort(RootItems);
 }
 
-void FActorInfoColumn::OnModeChanged( TSharedPtr< ECustomColumnMode::Type > NewSelection, ESelectInfo::Type /*SelectInfo*/ )
+void FActorInfoColumn::OnModeChanged( TSharedPtr< SceneOutliner::ECustomColumnMode::Type > NewSelection, ESelectInfo::Type /*SelectInfo*/ )
 {
 	CurrentMode = *NewSelection;
 
@@ -328,7 +325,7 @@ void FActorInfoColumn::OnModeChanged( TSharedPtr< ECustomColumnMode::Type > NewS
 
 EVisibility FActorInfoColumn::GetColumnDataVisibility( bool bIsClassHyperlink ) const
 {
-	if ( CurrentMode == ECustomColumnMode::Class )
+	if ( CurrentMode == SceneOutliner::ECustomColumnMode::Class )
 	{
 		return bIsClassHyperlink ? EVisibility::Visible : EVisibility::Collapsed;
 	}
@@ -338,7 +335,7 @@ EVisibility FActorInfoColumn::GetColumnDataVisibility( bool bIsClassHyperlink )
 	}
 }
 
-FText FActorInfoColumn::GetTextForItem( TWeakPtr<ITreeItem> TreeItem ) const
+FText FActorInfoColumn::GetTextForItem( TWeakPtr<ISceneOutlinerTreeItem> TreeItem ) const
 {
 	auto Item = TreeItem.Pin();
 	return Item.IsValid() ? FText::FromString(FGetInfo(CurrentMode)(*Item)) : FText::GetEmpty();
@@ -346,7 +343,7 @@ FText FActorInfoColumn::GetTextForItem( TWeakPtr<ITreeItem> TreeItem ) const
 
 FText FActorInfoColumn::GetSelectedMode() const
 {
-	if (CurrentMode == ECustomColumnMode::None)
+	if (CurrentMode == SceneOutliner::ECustomColumnMode::None)
 	{
 		return FText();
 	}
@@ -354,41 +351,41 @@ FText FActorInfoColumn::GetSelectedMode() const
 	return MakeComboText(CurrentMode);
 }
 
-FText FActorInfoColumn::MakeComboText( const ECustomColumnMode::Type& Mode ) const
+FText FActorInfoColumn::MakeComboText( const SceneOutliner::ECustomColumnMode::Type& Mode ) const
 {
 	FText ModeName;
 
 	switch( Mode )
 	{
-	case ECustomColumnMode::None:
+	case SceneOutliner::ECustomColumnMode::None:
 		ModeName = LOCTEXT("CustomColumnMode_None", "None");
 		break;
 
-	case ECustomColumnMode::Class:
+	case SceneOutliner::ECustomColumnMode::Class:
 		ModeName = LOCTEXT("CustomColumnMode_Class", "Type");
 		break;
 
-	case ECustomColumnMode::Level:
+	case SceneOutliner::ECustomColumnMode::Level:
 		ModeName = LOCTEXT("CustomColumnMode_Level", "Level");
 		break;
 
-	case ECustomColumnMode::Layer:
+	case SceneOutliner::ECustomColumnMode::Layer:
 		ModeName = LOCTEXT("CustomColumnMode_Layer", "Layer");
 		break;
 
-	case ECustomColumnMode::Socket:
+	case SceneOutliner::ECustomColumnMode::Socket:
 		ModeName = LOCTEXT("CustomColumnMode_Socket", "Socket");
 		break;
 
-	case ECustomColumnMode::InternalName:
+	case SceneOutliner::ECustomColumnMode::InternalName:
 		ModeName = LOCTEXT("CustomColumnMode_InternalName", "ID Name");
 		break;
 
-	case ECustomColumnMode::UncachedLights:
+	case SceneOutliner::ECustomColumnMode::UncachedLights:
 		ModeName = LOCTEXT("CustomColumnMode_UncachedLights", "# Uncached Lights");
 		break;
 
-	case ECustomColumnMode::Mobility:
+	case SceneOutliner::ECustomColumnMode::Mobility:
 		ModeName = LOCTEXT("CustomColumnMode_Mobility", "Mobility");
 		break;
 
@@ -401,41 +398,41 @@ FText FActorInfoColumn::MakeComboText( const ECustomColumnMode::Type& Mode ) con
 }
 
 
-FText FActorInfoColumn::MakeComboToolTipText( const ECustomColumnMode::Type& Mode )
+FText FActorInfoColumn::MakeComboToolTipText( const SceneOutliner::ECustomColumnMode::Type& Mode )
 {
 	FText ToolTipText;
 
 	switch( Mode )
 	{
-	case ECustomColumnMode::None:
+	case SceneOutliner::ECustomColumnMode::None:
 		ToolTipText = LOCTEXT("CustomColumnModeToolTip_None", "Hides all extra actor info");
 		break;
 
-	case ECustomColumnMode::Class:
+	case SceneOutliner::ECustomColumnMode::Class:
 		ToolTipText = LOCTEXT("CustomColumnModeToolTip_Class", "Displays the name of each actor's type");
 		break;
 
-	case ECustomColumnMode::Level:
+	case SceneOutliner::ECustomColumnMode::Level:
 		ToolTipText = LOCTEXT("CustomColumnModeToolTip_Level", "Displays the level each actor is in, and allows you to search by level name");
 		break;
 
-	case ECustomColumnMode::Layer:
+	case SceneOutliner::ECustomColumnMode::Layer:
 		ToolTipText = LOCTEXT("CustomColumnModeToolTip_Layer", "Displays the layer each actor is in, and allows you to search by layer name");
 		break;
 
-	case ECustomColumnMode::Socket:
+	case SceneOutliner::ECustomColumnMode::Socket:
 		ToolTipText = LOCTEXT("CustomColumnModeToolTip_Socket", "Shows the socket the actor is attached to, and allows you to search by socket name");
 		break;
 
-	case ECustomColumnMode::InternalName:
+	case SceneOutliner::ECustomColumnMode::InternalName:
 		ToolTipText = LOCTEXT("CustomColumnModeToolTip_InternalName", "Shows the internal name of the actor (for diagnostics)");
 		break;
 
-	case ECustomColumnMode::UncachedLights:
+	case SceneOutliner::ECustomColumnMode::UncachedLights:
 		ToolTipText = LOCTEXT("CustomColumnModeToolTip_UncachedLights", "Shows the number of uncached static lights (missing in lightmap)");
 		break;
 
-	case ECustomColumnMode::Mobility:
+	case SceneOutliner::ECustomColumnMode::Mobility:
 		ToolTipText = LOCTEXT("CustomColumnModeToolTip_Mobility", "Shows the mobility of each actor");
 		break;
 
@@ -449,10 +446,10 @@ FText FActorInfoColumn::MakeComboToolTipText( const ECustomColumnMode::Type& Mod
 }
 
 
-TSharedRef< ITableRow > FActorInfoColumn::MakeComboButtonItemWidget( TSharedPtr< ECustomColumnMode::Type > Mode, const TSharedRef<STableViewBase>& Owner )
+TSharedRef< ITableRow > FActorInfoColumn::MakeComboButtonItemWidget( TSharedPtr< SceneOutliner::ECustomColumnMode::Type > Mode, const TSharedRef<STableViewBase>& Owner )
 {
 	return
-		SNew( STableRow< TSharedPtr<ECustomColumnMode::Type> >, Owner )
+		SNew( STableRow< TSharedPtr<SceneOutliner::ECustomColumnMode::Type> >, Owner )
 		.Style(FAppStyle::Get(), "ComboBox.Row")
 		[
 			SNew( STextBlock )
@@ -461,6 +458,4 @@ TSharedRef< ITableRow > FActorInfoColumn::MakeComboButtonItemWidget( TSharedPtr<
 		];
 }
 
-}		// namespace SceneOutliner
-
 #undef LOCTEXT_NAMESPACE
diff --git a/Engine/Source/Editor/SceneOutliner/Private/SceneOutlinerActorInfoColumn.h b/Engine/Source/Editor/SceneOutliner/Private/SceneOutlinerActorInfoColumn.h
--- a/Engine/Source/Editor/SceneOutliner/Private/SceneOutlinerActorInfoColumn.h
+++ b/Engine/Source/Editor/SceneOutliner/Private/SceneOutlinerActorInfoColumn.h
@@ -16,41 +16,41 @@ class ISceneOutliner;
 
 namespace SceneOutliner
 {
-
-/** Types of actor data we can display in a 'custom' tree column */
-namespace ECustomColumnMode
-{
-	enum Type
+	/** Types of actor data we can display in a 'custom' tree column */
+	namespace ECustomColumnMode
 	{
-		/** Empty column -- doesn't display anything */
-		None = 0,
+		enum Type
+		{
+			/** Empty column -- doesn't display anything */
+			None = 0,
 
-		/** Class */
-		Class,
+			/** Class */
+			Class,
 
-		/** Mobility */
-		Mobility,
+			/** Mobility */
+			Mobility,
 
-		/** Level */
-		Level,
+			/** Level */
+			Level,
 
-		/** Layer */
-		Layer,
+			/** Layer */
+			Layer,
 
-		/** The socket the actor is attached to. */
-		Socket,
+			/** The socket the actor is attached to. */
+			Socket,
 
-		/** Actor's internal name (FName) */
-		InternalName,
+			/** Actor's internal name (FName) */
+			InternalName,
 
-		/** Actor's number of uncached lights */
-		UncachedLights,
+			/** Actor's number of uncached lights */
+			UncachedLights,
 
-		// ---
+			// ---
 
-		/** Number of options */
-		Count
-	};
+			/** Number of options */
+			Count
+		};
+	}
 }
 
 /**
@@ -64,11 +64,11 @@ public:
 	/**
 	 *	Constructor
 	 */
-	FActorInfoColumn( ISceneOutliner& Outliner, ECustomColumnMode::Type InDefaultCustomColumnMode = ECustomColumnMode::Class );
+	FActorInfoColumn( ISceneOutliner& Outliner, SceneOutliner::ECustomColumnMode::Type InDefaultCustomColumnMode = SceneOutliner::ECustomColumnMode::Class );
 
 	virtual ~FActorInfoColumn() {}
 
-	static FName GetID() { return FBuiltInColumnTypes::ActorInfo(); }
+	static FName GetID() { return FSceneOutlinerBuiltInColumnTypes::ActorInfo(); }
 
 	//////////////////////////////////////////////////////////////////////////
 	// Begin ISceneOutlinerColumn Implementation
@@ -77,32 +77,32 @@ public:
 
 	virtual SHeaderRow::FColumn::FArguments ConstructHeaderRowColumn() override;
 
-	virtual const TSharedRef< SWidget > ConstructRowWidget( FTreeItemRef TreeItem, const STableRow<FTreeItemPtr>& Row ) override;
+	virtual const TSharedRef< SWidget > ConstructRowWidget( FSceneOutlinerTreeItemRef TreeItem, const STableRow<FSceneOutlinerTreeItemPtr>& Row ) override;
 
-	virtual void PopulateSearchStrings( const ITreeItem& Item, TArray< FString >& OutSearchStrings ) const override;
+	virtual void PopulateSearchStrings( const ISceneOutlinerTreeItem& Item, TArray< FString >& OutSearchStrings ) const override;
 
 	virtual bool SupportsSorting() const override;
 
-	virtual void SortItems(TArray<FTreeItemPtr>& RootItems, const EColumnSortMode::Type SortMode) const override;
+	virtual void SortItems(TArray<FSceneOutlinerTreeItemPtr>& RootItems, const EColumnSortMode::Type SortMode) const override;
 	
 	// End ISceneOutlinerColumn Implementation
 	//////////////////////////////////////////////////////////////////////////
 
-	FText GetTextForItem( const TWeakPtr<ITreeItem> TreeItem ) const;
+	FText GetTextForItem( const TWeakPtr<ISceneOutlinerTreeItem> TreeItem ) const;
 
 private:
 
-	TSharedPtr<SWidget> ConstructClassHyperlink( ITreeItem& TreeItem );
+	TSharedPtr<SWidget> ConstructClassHyperlink( ISceneOutlinerTreeItem& TreeItem );
 
-	void OnModeChanged( TSharedPtr< ECustomColumnMode::Type > NewSelection, ESelectInfo::Type SelectInfo );
+	void OnModeChanged( TSharedPtr< SceneOutliner::ECustomColumnMode::Type > NewSelection, ESelectInfo::Type SelectInfo );
 
 	EVisibility GetColumnDataVisibility( bool bIsClassHyperlink ) const;
 
-	FText MakeComboText( const ECustomColumnMode::Type& Mode ) const;
+	FText MakeComboText( const SceneOutliner::ECustomColumnMode::Type& Mode ) const;
 
-	FText MakeComboToolTipText( const ECustomColumnMode::Type& Mode );
+	FText MakeComboToolTipText( const SceneOutliner::ECustomColumnMode::Type& Mode );
 
-	TSharedRef< ITableRow > MakeComboButtonItemWidget( TSharedPtr< ECustomColumnMode::Type > Mode, const TSharedRef<STableViewBase> & );
+	TSharedRef< ITableRow > MakeComboButtonItemWidget( TSharedPtr< SceneOutliner::ECustomColumnMode::Type > Mode, const TSharedRef<STableViewBase> & );
 
 	FText GetSelectedMode() const;
 
@@ -112,16 +112,13 @@ private:
 	 * Current custom column mode.  This is used for displaying a bit of extra data about the actors, as well as
 	 * allowing the user to search by additional criteria 
 	 */
-	mutable ECustomColumnMode::Type CurrentMode;
+	mutable SceneOutliner::ECustomColumnMode::Type CurrentMode;
 
 	/** A list of available custom column modes for Slate */
-	static TArray< TSharedPtr< ECustomColumnMode::Type > > ModeOptions;
+	static TArray< TSharedPtr< SceneOutliner::ECustomColumnMode::Type > > ModeOptions;
 
 	/** Weak reference to the outliner widget that owns our list */
 	TWeakPtr< ISceneOutliner > SceneOutlinerWeak;
 };
 
-
-}	// namespace SceneOutliner
-
 ENUM_RANGE_BY_COUNT(SceneOutliner::ECustomColumnMode::Type, SceneOutliner::ECustomColumnMode::Count)
diff --git a/Engine/Source/Editor/SceneOutliner/Private/SceneOutlinerDragDrop.cpp b/Engine/Source/Editor/SceneOutliner/Private/SceneOutlinerDragDrop.cpp
--- a/Engine/Source/Editor/SceneOutliner/Private/SceneOutlinerDragDrop.cpp
+++ b/Engine/Source/Editor/SceneOutliner/Private/SceneOutlinerDragDrop.cpp
@@ -1,88 +1,85 @@
 // Copyright Epic Games, Inc. All Rights Reserved.
 
 #include "SceneOutlinerDragDrop.h"
-#include "ITreeItem.h"
+#include "ISceneOutlinerTreeItem.h"
 #include "SceneOutlinerStandaloneTypes.h"
 
 #define LOCTEXT_NAMESPACE "SSceneOutliner"
 
-namespace SceneOutliner
+void FFolderDragDropOp::Init(TArray<FName> InFolders)
 {
-	void FFolderDragDropOp::Init(FFolderPaths InFolders)
-	{
-		Folders = MoveTemp(InFolders);
-
-		CurrentIconBrush = FEditorStyle::Get().GetBrush(TEXT("SceneOutliner.FolderClosed"));
-		if (Folders.Num() == 1)
-		{
-			CurrentHoverText = FText::FromName(GetFolderLeafName(Folders[0]));
-		}
-		else
-		{
-			CurrentHoverText = FText::Format(NSLOCTEXT("FFolderDragDropOp", "FormatFolders", "{0} Folders"), FText::AsNumber(Folders.Num()));
-		}
-	}
+	Folders = MoveTemp(InFolders);
 
-	FSceneOutlinerDragDropOp::FSceneOutlinerDragDropOp()
-		: OverrideText()
-		, OverrideIcon(nullptr)
+	CurrentIconBrush = FEditorStyle::Get().GetBrush(TEXT("SceneOutliner.FolderClosed"));
+	if (Folders.Num() == 1)
 	{
-		
+		CurrentHoverText = FText::FromName(SceneOutliner::GetFolderLeafName(Folders[0]));
 	}
-	
-	EVisibility FSceneOutlinerDragDropOp::GetOverrideVisibility() const
+	else
 	{
-		return OverrideText.IsEmpty() && OverrideIcon == nullptr ? EVisibility::Collapsed : EVisibility::Visible;
+		CurrentHoverText = FText::Format(NSLOCTEXT("FFolderDragDropOp", "FormatFolders", "{0} Folders"), FText::AsNumber(Folders.Num()));
 	}
+}
+
+FSceneOutlinerDragDropOp::FSceneOutlinerDragDropOp()
+	: OverrideText()
+	, OverrideIcon(nullptr)
+{
+		
+}
 	
-	EVisibility FSceneOutlinerDragDropOp::GetDefaultVisibility() const
-	{
-		return OverrideText.IsEmpty() && OverrideIcon == nullptr ? EVisibility::Visible : EVisibility::Collapsed;
-	}
+EVisibility FSceneOutlinerDragDropOp::GetOverrideVisibility() const
+{
+	return OverrideText.IsEmpty() && OverrideIcon == nullptr ? EVisibility::Collapsed : EVisibility::Visible;
+}
+	
+EVisibility FSceneOutlinerDragDropOp::GetDefaultVisibility() const
+{
+	return OverrideText.IsEmpty() && OverrideIcon == nullptr ? EVisibility::Visible : EVisibility::Collapsed;
+}
 
-	TSharedPtr<SWidget> FSceneOutlinerDragDropOp::GetDefaultDecorator() const
-	{
-		TSharedRef<SVerticalBox> VerticalBox = SNew(SVerticalBox);
+TSharedPtr<SWidget> FSceneOutlinerDragDropOp::GetDefaultDecorator() const
+{
+	TSharedRef<SVerticalBox> VerticalBox = SNew(SVerticalBox);
 
-		VerticalBox->AddSlot()
-		[
-			SNew(SBorder)
-			.BorderImage(FEditorStyle::GetBrush("Graph.ConnectorFeedback.Border"))
-			.Visibility(this, &FSceneOutlinerDragDropOp::GetOverrideVisibility)
-			.Content()
-			[			
-				SNew(SHorizontalBox)
+	VerticalBox->AddSlot()
+	[
+		SNew(SBorder)
+		.BorderImage(FEditorStyle::GetBrush("Graph.ConnectorFeedback.Border"))
+		.Visibility(this, &FSceneOutlinerDragDropOp::GetOverrideVisibility)
+		.Content()
+		[			
+			SNew(SHorizontalBox)
 				
-				+ SHorizontalBox::Slot()
-				.AutoWidth()
-				.Padding( 0.0f, 0.0f, 3.0f, 0.0f )
-				[
-					SNew( SImage )
-					.Image( this, &FSceneOutlinerDragDropOp::GetOverrideIcon )
-				]
+			+ SHorizontalBox::Slot()
+			.AutoWidth()
+			.Padding( 0.0f, 0.0f, 3.0f, 0.0f )
+			[
+				SNew( SImage )
+				.Image( this, &FSceneOutlinerDragDropOp::GetOverrideIcon )
+			]
 
-				+ SHorizontalBox::Slot()
-				.AutoWidth()
-				.VAlign( VAlign_Center )
-				[
-					SNew(STextBlock) 
-					.Text( this, &FSceneOutlinerDragDropOp::GetOverrideText )
-				]
+			+ SHorizontalBox::Slot()
+			.AutoWidth()
+			.VAlign( VAlign_Center )
+			[
+				SNew(STextBlock) 
+				.Text( this, &FSceneOutlinerDragDropOp::GetOverrideText )
 			]
-		];
+		]
+	];
 
-		for (auto& SubOp : SubOps)
+	for (auto& SubOp : SubOps)
+	{
+		auto Content = SubOp->GetDefaultDecorator();
+		if (Content.IsValid())
 		{
-			auto Content = SubOp->GetDefaultDecorator();
-			if (Content.IsValid())
-			{
-				Content->SetVisibility(TAttribute<EVisibility>(this, &FSceneOutlinerDragDropOp::GetDefaultVisibility));
-				VerticalBox->AddSlot()[Content.ToSharedRef()];
-			}
+			Content->SetVisibility(TAttribute<EVisibility>(this, &FSceneOutlinerDragDropOp::GetDefaultVisibility));
+			VerticalBox->AddSlot()[Content.ToSharedRef()];
 		}
-
-		return VerticalBox;
 	}
+
+	return VerticalBox;
 }
 
 #undef LOCTEXT_NAMESPACE
diff --git a/Engine/Source/Editor/SceneOutliner/Private/SceneOutlinerFilters.cpp b/Engine/Source/Editor/SceneOutliner/Private/SceneOutlinerFilters.cpp
--- a/Engine/Source/Editor/SceneOutliner/Private/SceneOutlinerFilters.cpp
+++ b/Engine/Source/Editor/SceneOutliner/Private/SceneOutlinerFilters.cpp
@@ -3,33 +3,30 @@
 #include "SceneOutlinerFilters.h"
 #include "Framework/MultiBox/MultiBoxBuilder.h"
 
-namespace SceneOutliner
-{
-
-void FOutlinerFilterInfo::InitFilter(TSharedPtr<FOutlinerFilters> InFilters)
+void FSceneOutlinerFilterInfo::InitFilter(TSharedPtr<FSceneOutlinerFilters> InFilters)
 {
 	Filters = InFilters;
 
 	ApplyFilter(bActive);
 }
 
-void FOutlinerFilterInfo::AddMenu(FMenuBuilder& InMenuBuilder)
+void FSceneOutlinerFilterInfo::AddMenu(FMenuBuilder& InMenuBuilder)
 {
 	InMenuBuilder.AddMenuEntry(
 		FilterTitle,
 		FilterTooltip,
 		FSlateIcon(),
 		FUIAction(
-			FExecuteAction::CreateRaw( this, &FOutlinerFilterInfo::ToggleFilterActive ),
+			FExecuteAction::CreateRaw( this, &FSceneOutlinerFilterInfo::ToggleFilterActive ),
 			FCanExecuteAction(),
-			FIsActionChecked::CreateRaw( this, &FOutlinerFilterInfo::IsFilterActive )
+			FIsActionChecked::CreateRaw( this, &FSceneOutlinerFilterInfo::IsFilterActive )
 		),
 		NAME_None,
 		EUserInterfaceActionType::ToggleButton
 	);
 }
 
-void FOutlinerFilterInfo::ApplyFilter(bool bInActive)
+void FSceneOutlinerFilterInfo::ApplyFilter(bool bInActive)
 {
 	if ( !Filter.IsValid() )
 	{
@@ -46,7 +43,7 @@ void FOutlinerFilterInfo::ApplyFilter(bool bInActive)
 	}
 }
 
-void FOutlinerFilterInfo::ToggleFilterActive()
+void FSceneOutlinerFilterInfo::ToggleFilterActive()
 {
 	bActive = !bActive;
 
@@ -55,9 +52,7 @@ void FOutlinerFilterInfo::ToggleFilterActive()
 	OnToggleEvent.Broadcast(bActive);
 }
 
-bool FOutlinerFilterInfo::IsFilterActive() const
+bool FSceneOutlinerFilterInfo::IsFilterActive() const
 {
 	return bActive;
 }
-
-} // namespace SceneOutliner
diff --git a/Engine/Source/Editor/SceneOutliner/Private/SceneOutlinerGutter.cpp b/Engine/Source/Editor/SceneOutliner/Private/SceneOutlinerGutter.cpp
--- a/Engine/Source/Editor/SceneOutliner/Private/SceneOutlinerGutter.cpp
+++ b/Engine/Source/Editor/SceneOutliner/Private/SceneOutlinerGutter.cpp
@@ -17,22 +17,22 @@
 
 namespace SceneOutliner
 {
-
-void OnSetItemVisibility(ITreeItem& Item, const bool bNewVisibility)
-{
-	// Apply the same visibility to the children
-	Item.OnVisibilityChanged(bNewVisibility);
-	for (auto& ChildPtr : Item.GetChildren())
+	void OnSetItemVisibility(ISceneOutlinerTreeItem& Item, const bool bNewVisibility)
 	{
-		auto Child = ChildPtr.Pin();
-		if (Child.IsValid())
+		// Apply the same visibility to the children
+		Item.OnVisibilityChanged(bNewVisibility);
+		for (auto& ChildPtr : Item.GetChildren())
 		{
-			OnSetItemVisibility(*Child, bNewVisibility);
+			auto Child = ChildPtr.Pin();
+			if (Child.IsValid())
+			{
+				OnSetItemVisibility(*Child, bNewVisibility);
+			}
 		}
 	}
 }
 
-bool FGetVisibilityCache::RecurseChildren(const ITreeItem& Item) const
+bool FSceneOutlinerVisibilityCache::RecurseChildren(const ISceneOutlinerTreeItem& Item) const
 {
 	if (const bool* Info = VisibilityInfo.Find(&Item))
 	{
@@ -56,7 +56,7 @@ bool FGetVisibilityCache::RecurseChildren(const ITreeItem& Item) const
 	}
 }
 
-bool FGetVisibilityCache::GetVisibility(const ITreeItem& Item) const
+bool FSceneOutlinerVisibilityCache::GetVisibility(const ISceneOutlinerTreeItem& Item) const
 {
 	if (Item.HasVisibilityInfo())
 	{
@@ -118,7 +118,7 @@ public:
 	SLATE_END_ARGS()
 
 	/** Construct this widget */
-	void Construct(const FArguments& InArgs, TWeakPtr<FSceneOutlinerGutter> InWeakColumn, TWeakPtr<ISceneOutliner> InWeakOutliner, TWeakPtr<ITreeItem> InWeakTreeItem, const STableRow<FTreeItemPtr>* InRow)
+	void Construct(const FArguments& InArgs, TWeakPtr<FSceneOutlinerGutter> InWeakColumn, TWeakPtr<ISceneOutliner> InWeakOutliner, TWeakPtr<ISceneOutlinerTreeItem> InWeakTreeItem, const STableRow<FSceneOutlinerTreeItemPtr>* InRow)
 	{
 		WeakTreeItem = InWeakTreeItem;
 		WeakOutliner = InWeakOutliner;
@@ -196,7 +196,7 @@ private:
 			{
 				if (IsVisible(SelectedItem, Column) != bVisible)
 				{
-					OnSetItemVisibility(*SelectedItem, bVisible);
+					SceneOutliner::OnSetItemVisibility(*SelectedItem, bVisible);
 				}		
 			}
 
@@ -279,7 +279,7 @@ private:
 	}
 
 	/** Check if the specified item is visible */
-	static bool IsVisible(const FTreeItemPtr& Item, const TSharedPtr<FSceneOutlinerGutter>& Column)
+	static bool IsVisible(const FSceneOutlinerTreeItemPtr& Item, const TSharedPtr<FSceneOutlinerGutter>& Column)
 	{
 		return Column.IsValid() && Item.IsValid() ? Column->IsItemVisible(*Item) : false;
 	}
@@ -293,12 +293,12 @@ private:
 	/** Set the item this widget is responsible for to be hidden or shown */
 	void SetIsVisible(const bool bVisible)
 	{
-		TSharedPtr<ITreeItem> TreeItem = WeakTreeItem.Pin();
+		TSharedPtr<ISceneOutlinerTreeItem> TreeItem = WeakTreeItem.Pin();
 		TSharedPtr<ISceneOutliner> Outliner = WeakOutliner.Pin();
 
 		if (TreeItem.IsValid() && Outliner.IsValid() && IsVisible() != bVisible)
 		{
-			OnSetItemVisibility(*TreeItem, bVisible);
+			SceneOutliner::OnSetItemVisibility(*TreeItem, bVisible);
 			
 			Outliner->Refresh();
 
@@ -307,7 +307,7 @@ private:
 	}
 
 	/** The tree item we relate to */
-	TWeakPtr<ITreeItem> WeakTreeItem;
+	TWeakPtr<ISceneOutlinerTreeItem> WeakTreeItem;
 	
 	/** Reference back to the outliner so we can set visibility of a whole selection */
 	TWeakPtr<ISceneOutliner> WeakOutliner;
@@ -316,7 +316,7 @@ private:
 	TWeakPtr<FSceneOutlinerGutter> WeakColumn;
 
 	/** Weak pointer back to the row */
-	const STableRow<FTreeItemPtr>* Row;
+	const STableRow<FSceneOutlinerTreeItemPtr>* Row;
 
 	/** Scoped undo transaction */
 	TUniquePtr<FScopedTransaction> UndoTransaction;
@@ -358,7 +358,7 @@ SHeaderRow::FColumn::FArguments FSceneOutlinerGutter::ConstructHeaderRowColumn()
 		];
 }
 
-const TSharedRef<SWidget> FSceneOutlinerGutter::ConstructRowWidget(FTreeItemRef TreeItem, const STableRow<FTreeItemPtr>& Row)
+const TSharedRef<SWidget> FSceneOutlinerGutter::ConstructRowWidget(FSceneOutlinerTreeItemRef TreeItem, const STableRow<FSceneOutlinerTreeItemPtr>& Row)
 {
 	return SNew(SHorizontalBox)
 		+SHorizontalBox::Slot()
@@ -369,17 +369,15 @@ const TSharedRef<SWidget> FSceneOutlinerGutter::ConstructRowWidget(FTreeItemRef
 		];
 }
 
-void FSceneOutlinerGutter::SortItems(TArray<FTreeItemPtr>& RootItems, const EColumnSortMode::Type SortMode) const
+void FSceneOutlinerGutter::SortItems(TArray<FSceneOutlinerTreeItemPtr>& RootItems, const EColumnSortMode::Type SortMode) const
 {
-	FSortHelper<int32, bool>()
+	FSceneOutlinerSortHelper<int32, bool>()
 		/** Sort by type first */
-		.Primary([this](const ITreeItem& Item){ return WeakOutliner.Pin()->GetMode()->GetTypeSortPriority(Item); }, SortMode)
+		.Primary([this](const ISceneOutlinerTreeItem& Item){ return WeakOutliner.Pin()->GetMode()->GetTypeSortPriority(Item); }, SortMode)
 		/** Then by visibility */
-		.Secondary([](const ITreeItem& Item) {return FGetVisibilityCache().GetVisibility(Item); },					SortMode)
+		.Secondary([](const ISceneOutlinerTreeItem& Item) {return FSceneOutlinerVisibilityCache().GetVisibility(Item); },					SortMode)
 		.Sort(RootItems);
 }
 
-}		// namespace SceneOutliner
-
 #undef LOCTEXT_NAMESPACE
 
diff --git a/Engine/Source/Editor/SceneOutliner/Private/SceneOutlinerGutter.h b/Engine/Source/Editor/SceneOutliner/Private/SceneOutlinerGutter.h
--- a/Engine/Source/Editor/SceneOutliner/Private/SceneOutlinerGutter.h
+++ b/Engine/Source/Editor/SceneOutliner/Private/SceneOutlinerGutter.h
@@ -7,14 +7,11 @@
 #include "Widgets/Views/SHeaderRow.h"
 #include "SceneOutlinerPublicTypes.h"
 #include "ISceneOutlinerColumn.h"
-#include "ITreeItem.h"
+#include "ISceneOutlinerTreeItem.h"
 
 class ISceneOutliner;
 template<typename ItemType> class STableRow;
 
-namespace SceneOutliner
-{
-
 /**
  * A gutter for the SceneOutliner which handles setting and visualizing item visibility
  */
@@ -27,7 +24,7 @@ public:
 
 	virtual ~FSceneOutlinerGutter() {}
 
-	static FName GetID() { return FBuiltInColumnTypes::Gutter(); }
+	static FName GetID() { return FSceneOutlinerBuiltInColumnTypes::Gutter(); }
 	
 	// -----------------------------------------
 	// ISceneOutlinerColumn Implementation
@@ -35,17 +32,17 @@ public:
 
 	virtual SHeaderRow::FColumn::FArguments ConstructHeaderRowColumn() override;
 
-	virtual const TSharedRef< SWidget > ConstructRowWidget( FTreeItemRef TreeItem, const STableRow<FTreeItemPtr>& Row ) override;
+	virtual const TSharedRef< SWidget > ConstructRowWidget( FSceneOutlinerTreeItemRef TreeItem, const STableRow<FSceneOutlinerTreeItemPtr>& Row ) override;
 	
 	virtual void Tick(double InCurrentTime, float InDeltaTime) override;
 
 	virtual bool SupportsSorting() const override { return true; }
 
-	virtual void SortItems(TArray<FTreeItemPtr>& RootItems, const EColumnSortMode::Type SortMode) const override;
+	virtual void SortItems(TArray<FSceneOutlinerTreeItemPtr>& RootItems, const EColumnSortMode::Type SortMode) const override;
 	// -----------------------------------------
 
 	/** Check whether the specified item is visible */
-	FORCEINLINE bool IsItemVisible(const ITreeItem& Item)
+	FORCEINLINE bool IsItemVisible(const ISceneOutlinerTreeItem& Item)
 	{
 		return VisibilityCache.GetVisibility(Item);
 	}
@@ -56,7 +53,5 @@ private:
 	TWeakPtr<ISceneOutliner> WeakOutliner;
 
 	/** Get and cache visibility for items. Cached per-frame to avoid expensive recursion. */
-	FGetVisibilityCache VisibilityCache;
+	FSceneOutlinerVisibilityCache VisibilityCache;
 };
-
-}	// namespace SceneOutliner
diff --git a/Engine/Source/Editor/SceneOutliner/Private/SceneOutlinerItemLabelColumn.cpp b/Engine/Source/Editor/SceneOutliner/Private/SceneOutlinerItemLabelColumn.cpp
--- a/Engine/Source/Editor/SceneOutliner/Private/SceneOutlinerItemLabelColumn.cpp
+++ b/Engine/Source/Editor/SceneOutliner/Private/SceneOutlinerItemLabelColumn.cpp
@@ -19,43 +19,39 @@
 #include "SortHelper.h"
 #include "Widgets/SToolTip.h"
 
-
 #define LOCTEXT_NAMESPACE "SceneOutlinerItemLabelColumn"
 
-namespace SceneOutliner
-{
-FName FItemLabelColumn::GetColumnID()
+FName FSceneOutlinerItemLabelColumn::GetColumnID()
 {
 	return GetID();
 }
 
-SHeaderRow::FColumn::FArguments FItemLabelColumn::ConstructHeaderRowColumn()
+SHeaderRow::FColumn::FArguments FSceneOutlinerItemLabelColumn::ConstructHeaderRowColumn()
 {
 	return SHeaderRow::Column(GetColumnID())
 		.DefaultLabel(LOCTEXT("ItemLabel_HeaderText", "Label"))
 		.FillWidth( 5.0f );
 }
-const TSharedRef<SWidget> FItemLabelColumn::ConstructRowWidget(FTreeItemRef TreeItem, const STableRow<FTreeItemPtr>& Row)
+const TSharedRef<SWidget> FSceneOutlinerItemLabelColumn::ConstructRowWidget(FSceneOutlinerTreeItemRef TreeItem, const STableRow<FSceneOutlinerTreeItemPtr>& Row)
 {
 	ISceneOutliner* Outliner = WeakSceneOutliner.Pin().Get();
 	check(Outliner);
 	return TreeItem->GenerateLabelWidget(*Outliner, Row);
 }
 
-void FItemLabelColumn::PopulateSearchStrings( const ITreeItem& Item, TArray< FString >& OutSearchStrings ) const
+void FSceneOutlinerItemLabelColumn::PopulateSearchStrings( const ISceneOutlinerTreeItem& Item, TArray< FString >& OutSearchStrings ) const
 {
 	OutSearchStrings.Add(Item.GetDisplayString());
 }
 
-void FItemLabelColumn::SortItems(TArray<FTreeItemPtr>& OutItems, const EColumnSortMode::Type SortMode) const
+void FSceneOutlinerItemLabelColumn::SortItems(TArray<FSceneOutlinerTreeItemPtr>& OutItems, const EColumnSortMode::Type SortMode) const
 {
-	typedef FSortHelper<int32, FNumericStringWrapper> FSort;
+	typedef FSceneOutlinerSortHelper<int32, SceneOutliner::FNumericStringWrapper> FSort;
 
 	FSort()
-		.Primary([this](const ITreeItem& Item){ return WeakSceneOutliner.Pin()->GetTypeSortPriority(Item); },SortMode)
-		.Secondary([](const ITreeItem& Item){ return FNumericStringWrapper(Item.GetDisplayString()); }, 			SortMode)
+		.Primary([this](const ISceneOutlinerTreeItem& Item){ return WeakSceneOutliner.Pin()->GetTypeSortPriority(Item); },			SortMode)
+		.Secondary([](const ISceneOutlinerTreeItem& Item){ return SceneOutliner::FNumericStringWrapper(Item.GetDisplayString()); }, SortMode)
 		.Sort(OutItems);
 }
-}	// namespace SceneOutliner
 
 #undef LOCTEXT_NAMESPACE
diff --git a/Engine/Source/Editor/SceneOutliner/Private/SceneOutlinerItemLabelColumn.h b/Engine/Source/Editor/SceneOutliner/Private/SceneOutlinerItemLabelColumn.h
--- a/Engine/Source/Editor/SceneOutliner/Private/SceneOutlinerItemLabelColumn.h
+++ b/Engine/Source/Editor/SceneOutliner/Private/SceneOutlinerItemLabelColumn.h
@@ -11,18 +11,16 @@
 
 template<typename ItemType> class STableRow;
 
-namespace SceneOutliner
-{
 /** A column for the SceneOutliner that displays the item label */
-class FItemLabelColumn : public ISceneOutlinerColumn
+class FSceneOutlinerItemLabelColumn : public ISceneOutlinerColumn
 {
 
 public:
-	FItemLabelColumn(ISceneOutliner& SceneOutliner) : WeakSceneOutliner(StaticCastSharedRef<ISceneOutliner>(SceneOutliner.AsShared())) {}
+	FSceneOutlinerItemLabelColumn(ISceneOutliner& SceneOutliner) : WeakSceneOutliner(StaticCastSharedRef<ISceneOutliner>(SceneOutliner.AsShared())) {}
 	
-	virtual ~FItemLabelColumn() {}
+	virtual ~FSceneOutlinerItemLabelColumn() {}
 
-	static FName GetID() { return FBuiltInColumnTypes::Label(); }
+	static FName GetID() { return FSceneOutlinerBuiltInColumnTypes::Label(); }
 
 	//////////////////////////////////////////////////////////////////////////
 	// Begin ISceneOutlinerColumn Implementation
@@ -31,13 +29,13 @@ public:
 
 	virtual SHeaderRow::FColumn::FArguments ConstructHeaderRowColumn() override;
 	
-	virtual const TSharedRef< SWidget > ConstructRowWidget( FTreeItemRef TreeItem, const STableRow<FTreeItemPtr>& Row ) override;
+	virtual const TSharedRef< SWidget > ConstructRowWidget( FSceneOutlinerTreeItemRef TreeItem, const STableRow<FSceneOutlinerTreeItemPtr>& Row ) override;
 
-	virtual void PopulateSearchStrings( const ITreeItem& Item, TArray< FString >& OutSearchStrings ) const override;
+	virtual void PopulateSearchStrings( const ISceneOutlinerTreeItem& Item, TArray< FString >& OutSearchStrings ) const override;
 
 	virtual bool SupportsSorting() const override { return true; }
 
-	virtual void SortItems(TArray<FTreeItemPtr>& RootItems, const EColumnSortMode::Type SortMode) const override;
+	virtual void SortItems(TArray<FSceneOutlinerTreeItemPtr>& RootItems, const EColumnSortMode::Type SortMode) const override;
 	
 	// End ISceneOutlinerColumn Implementation
 	//////////////////////////////////////////////////////////////////////////
@@ -46,5 +44,3 @@ private:
 	TWeakPtr<ISceneOutliner> WeakSceneOutliner;
 
 };
-
-}	// namespace SceneOutliner
diff --git a/Engine/Source/Editor/SceneOutliner/Private/SceneOutlinerModule.cpp b/Engine/Source/Editor/SceneOutliner/Private/SceneOutlinerModule.cpp
--- a/Engine/Source/Editor/SceneOutliner/Private/SceneOutlinerModule.cpp
+++ b/Engine/Source/Editor/SceneOutliner/Private/SceneOutlinerModule.cpp
@@ -21,32 +21,32 @@
 
 void FSceneOutlinerModule::StartupModule()
 {
-	RegisterDefaultColumnType< SceneOutliner::FItemLabelColumn >(SceneOutliner::FColumnInfo(SceneOutliner::EColumnVisibility::Visible, 10));
+	RegisterDefaultColumnType< FSceneOutlinerItemLabelColumn >(FSceneOutlinerColumnInfo(ESceneOutlinerColumnVisibility::Visible, 10));
 
 	// Register builtin column types which are not active by default
-	RegisterColumnType<SceneOutliner::FSceneOutlinerGutter>();
-	RegisterColumnType<SceneOutliner::FActorInfoColumn>();
+	RegisterColumnType<FSceneOutlinerGutter>();
+	RegisterColumnType<FActorInfoColumn>();
 }
 
 
 void FSceneOutlinerModule::ShutdownModule()
 {
-	UnRegisterColumnType<SceneOutliner::FSceneOutlinerGutter>();
-	UnRegisterColumnType<SceneOutliner::FItemLabelColumn>();
-	UnRegisterColumnType<SceneOutliner::FActorInfoColumn>();
+	UnRegisterColumnType<FSceneOutlinerGutter>();
+	UnRegisterColumnType<FSceneOutlinerItemLabelColumn>();
+	UnRegisterColumnType<FActorInfoColumn>();
 }
 
-TSharedRef<ISceneOutliner> FSceneOutlinerModule::CreateSceneOutliner(const SceneOutliner::FInitializationOptions& InitOptions) const
+TSharedRef<ISceneOutliner> FSceneOutlinerModule::CreateSceneOutliner(const FSceneOutlinerInitializationOptions& InitOptions) const
 {
-	return SNew(SceneOutliner::SSceneOutliner, InitOptions)
+	return SNew(SSceneOutliner, InitOptions)
 		.IsEnabled(FSlateApplication::Get().GetNormalExecutionAttribute());
 }
 
-TSharedRef<ISceneOutliner> FSceneOutlinerModule::CreateActorPicker(const SceneOutliner::FInitializationOptions& InInitOptions, const FOnActorPicked& OnActorPickedDelegate, TWeakObjectPtr<UWorld> SpecifiedWorld) const
+TSharedRef<ISceneOutliner> FSceneOutlinerModule::CreateActorPicker(const FSceneOutlinerInitializationOptions& InInitOptions, const FOnActorPicked& OnActorPickedDelegate, TWeakObjectPtr<UWorld> SpecifiedWorld) const
 {
-	auto OnItemPicked = FOnSceneOutlinerItemPicked::CreateLambda([OnActorPickedDelegate](TSharedRef<SceneOutliner::ITreeItem> Item)
+	auto OnItemPicked = FOnSceneOutlinerItemPicked::CreateLambda([OnActorPickedDelegate](TSharedRef<ISceneOutlinerTreeItem> Item)
 		{
-			if (SceneOutliner::FActorTreeItem* ActorItem = Item->CastTo<SceneOutliner::FActorTreeItem>())
+			if (FActorTreeItem* ActorItem = Item->CastTo<FActorTreeItem>())
 			{
 				if (ActorItem->IsValid())
 				{
@@ -56,28 +56,28 @@ TSharedRef<ISceneOutliner> FSceneOutlinerModule::CreateActorPicker(const SceneOu
 			}
 		});
 
-	SceneOutliner::FCreateOutlinerMode ModeFactory = SceneOutliner::FCreateOutlinerMode::CreateLambda([&OnItemPicked, &SpecifiedWorld](SceneOutliner::SSceneOutliner* Outliner)
+	FCreateSceneOutlinerMode ModeFactory = FCreateSceneOutlinerMode::CreateLambda([&OnItemPicked, &SpecifiedWorld](SSceneOutliner* Outliner)
 		{
-			return new SceneOutliner::FActorPickingMode(Outliner, true, OnItemPicked, SpecifiedWorld);
+			return new FActorPickingMode(Outliner, true, OnItemPicked, SpecifiedWorld);
 		});
 
 	
 
-	SceneOutliner::FInitializationOptions InitOptions(InInitOptions);
+	FSceneOutlinerInitializationOptions InitOptions(InInitOptions);
 	InitOptions.ModeFactory = ModeFactory;
 	if (InitOptions.ColumnMap.Num() == 0)
 	{
-		InitOptions.ColumnMap.Add(SceneOutliner::FBuiltInColumnTypes::Label(), SceneOutliner::FColumnInfo(SceneOutliner::EColumnVisibility::Visible, 0));
-		InitOptions.ColumnMap.Add(SceneOutliner::FBuiltInColumnTypes::ActorInfo(), SceneOutliner::FColumnInfo(SceneOutliner::EColumnVisibility::Visible, 10));
+		InitOptions.ColumnMap.Add(FSceneOutlinerBuiltInColumnTypes::Label(), FSceneOutlinerColumnInfo(ESceneOutlinerColumnVisibility::Visible, 0));
+		InitOptions.ColumnMap.Add(FSceneOutlinerBuiltInColumnTypes::ActorInfo(), FSceneOutlinerColumnInfo(ESceneOutlinerColumnVisibility::Visible, 10));
 	}
 	return CreateSceneOutliner(InitOptions);
 }
 
-TSharedRef<ISceneOutliner> FSceneOutlinerModule::CreateComponentPicker(const SceneOutliner::FInitializationOptions& InInitOptions, const FOnComponentPicked& OnComponentPickedDelegate, TWeakObjectPtr<UWorld> SpecifiedWorld) const
+TSharedRef<ISceneOutliner> FSceneOutlinerModule::CreateComponentPicker(const FSceneOutlinerInitializationOptions& InInitOptions, const FOnComponentPicked& OnComponentPickedDelegate, TWeakObjectPtr<UWorld> SpecifiedWorld) const
 {
-	auto OnItemPicked = FOnSceneOutlinerItemPicked::CreateLambda([OnComponentPickedDelegate](TSharedRef<SceneOutliner::ITreeItem> Item)
+	auto OnItemPicked = FOnSceneOutlinerItemPicked::CreateLambda([OnComponentPickedDelegate](TSharedRef<ISceneOutlinerTreeItem> Item)
 		{
-			if (SceneOutliner::FComponentTreeItem* ComponentItem = Item->CastTo<SceneOutliner::FComponentTreeItem>())
+			if (FComponentTreeItem* ComponentItem = Item->CastTo<FComponentTreeItem>())
 			{
 				if (ComponentItem->IsValid())
 				{
@@ -86,35 +86,35 @@ TSharedRef<ISceneOutliner> FSceneOutlinerModule::CreateComponentPicker(const Sce
 			}
 		});
 
-	SceneOutliner::FCreateOutlinerMode ModeFactory = SceneOutliner::FCreateOutlinerMode::CreateLambda([&OnItemPicked, &SpecifiedWorld](SceneOutliner::SSceneOutliner* Outliner)
+	FCreateSceneOutlinerMode ModeFactory = FCreateSceneOutlinerMode::CreateLambda([&OnItemPicked, &SpecifiedWorld](SSceneOutliner* Outliner)
 		{
-			return new SceneOutliner::FActorPickingMode(Outliner, false, OnItemPicked, SpecifiedWorld);
+			return new FActorPickingMode(Outliner, false, OnItemPicked, SpecifiedWorld);
 		});
 
-	SceneOutliner::FInitializationOptions InitOptions(InInitOptions);
+	FSceneOutlinerInitializationOptions InitOptions(InInitOptions);
 	InitOptions.ModeFactory = ModeFactory;
 	if (InitOptions.ColumnMap.Num() == 0)
 	{
-		InitOptions.ColumnMap.Add(SceneOutliner::FBuiltInColumnTypes::Label(), SceneOutliner::FColumnInfo(SceneOutliner::EColumnVisibility::Visible, 0));
-		InitOptions.ColumnMap.Add(SceneOutliner::FBuiltInColumnTypes::ActorInfo(), SceneOutliner::FColumnInfo(SceneOutliner::EColumnVisibility::Visible, 10));
+		InitOptions.ColumnMap.Add(FSceneOutlinerBuiltInColumnTypes::Label(), FSceneOutlinerColumnInfo(ESceneOutlinerColumnVisibility::Visible, 0));
+		InitOptions.ColumnMap.Add(FSceneOutlinerBuiltInColumnTypes::ActorInfo(), FSceneOutlinerColumnInfo(ESceneOutlinerColumnVisibility::Visible, 10));
 	}
 	return CreateSceneOutliner(InitOptions);
 }
 
-TSharedRef< ISceneOutliner > FSceneOutlinerModule::CreateActorBrowser(const SceneOutliner::FInitializationOptions& InInitOptions, TWeakObjectPtr<UWorld> SpecifiedWorld) const
+TSharedRef< ISceneOutliner > FSceneOutlinerModule::CreateActorBrowser(const FSceneOutlinerInitializationOptions& InInitOptions, TWeakObjectPtr<UWorld> SpecifiedWorld) const
 {
-	SceneOutliner::FCreateOutlinerMode ModeFactory = SceneOutliner::FCreateOutlinerMode::CreateLambda([&SpecifiedWorld](SceneOutliner::SSceneOutliner* Outliner)
+	FCreateSceneOutlinerMode ModeFactory = FCreateSceneOutlinerMode::CreateLambda([&SpecifiedWorld](SSceneOutliner* Outliner)
 		{
-			return new SceneOutliner::FActorBrowsingMode(Outliner, SpecifiedWorld);
+			return new FActorBrowsingMode(Outliner, SpecifiedWorld);
 		});
 
-	SceneOutliner::FInitializationOptions InitOptions(InInitOptions);
+	FSceneOutlinerInitializationOptions InitOptions(InInitOptions);
 	InitOptions.ModeFactory = ModeFactory;
 	if (InitOptions.ColumnMap.Num() == 0)
 	{
 		InitOptions.UseDefaultColumns();
-		InitOptions.ColumnMap.Add(SceneOutliner::FBuiltInColumnTypes::Gutter(), SceneOutliner::FColumnInfo(SceneOutliner::EColumnVisibility::Visible, 0));
-		InitOptions.ColumnMap.Add(SceneOutliner::FBuiltInColumnTypes::ActorInfo(), SceneOutliner::FColumnInfo(SceneOutliner::EColumnVisibility::Visible, 20));
+		InitOptions.ColumnMap.Add(FSceneOutlinerBuiltInColumnTypes::Gutter(), FSceneOutlinerColumnInfo(ESceneOutlinerColumnVisibility::Visible, 0));
+		InitOptions.ColumnMap.Add(FSceneOutlinerBuiltInColumnTypes::ActorInfo(), FSceneOutlinerColumnInfo(ESceneOutlinerColumnVisibility::Visible, 20));
 	}
 	return CreateSceneOutliner(InitOptions);
 }
diff --git a/Engine/Source/Editor/SceneOutliner/Private/SceneOutlinerSharedData.cpp b/Engine/Source/Editor/SceneOutliner/Private/SceneOutlinerSharedData.cpp
--- a/Engine/Source/Editor/SceneOutliner/Private/SceneOutlinerSharedData.cpp
+++ b/Engine/Source/Editor/SceneOutliner/Private/SceneOutlinerSharedData.cpp
@@ -6,10 +6,7 @@
 
 #include "SceneOutlinerModule.h"
 
-namespace SceneOutliner
-{
-
-void FSharedOutlinerData::UseDefaultColumns()
+void FSharedSceneOutlinerData::UseDefaultColumns()
 {
 	FSceneOutlinerModule& SceneOutlinerModule = FModuleManager::LoadModuleChecked<FSceneOutlinerModule>("SceneOutliner");
 
@@ -19,5 +16,3 @@ void FSharedOutlinerData::UseDefaultColumns()
 		ColumnMap.Add(DefaultColumn.Key, DefaultColumn.Value);
 	}
 }
-
-}	// namespace SceneOutliner 
\ No newline at end of file
diff --git a/Engine/Source/Editor/SceneOutliner/Private/SceneOutlinerStandaloneTypes.cpp b/Engine/Source/Editor/SceneOutliner/Private/SceneOutlinerStandaloneTypes.cpp
--- a/Engine/Source/Editor/SceneOutliner/Private/SceneOutlinerStandaloneTypes.cpp
+++ b/Engine/Source/Editor/SceneOutliner/Private/SceneOutlinerStandaloneTypes.cpp
@@ -3,57 +3,57 @@
 #include "SceneOutlinerStandaloneTypes.h"
 
 #include "EditorActorFolders.h"
-#include "ITreeItem.h"
+#include "ISceneOutlinerTreeItem.h"
 #include "ISceneOutliner.h"
 #include "ISceneOutlinerMode.h"
 
 
 #define LOCTEXT_NAMESPACE "SceneOutlinerStandaloneTypes"
 
-namespace SceneOutliner
-{
-	uint32 FTreeItemType::NextUniqueID = 0;
-	const FTreeItemType ITreeItem::Type;
+uint32 FSceneOutlinerTreeItemType::NextUniqueID = 0;
+const FSceneOutlinerTreeItemType ISceneOutlinerTreeItem::Type;
 
-	const FLinearColor FCommonLabelData::DarkColor(0.3f, 0.3f, 0.3f);
+const FLinearColor FSceneOutlinerCommonLabelData::DarkColor(0.3f, 0.3f, 0.3f);
 
-	TOptional<FLinearColor> FCommonLabelData::GetForegroundColor(const ITreeItem& TreeItem) const
+TOptional<FLinearColor> FSceneOutlinerCommonLabelData::GetForegroundColor(const ISceneOutlinerTreeItem& TreeItem) const
+{
+	if (!TreeItem.IsValid())
 	{
-		if (!TreeItem.IsValid())
-		{
-			return DarkColor;
-		}
+		return DarkColor;
+	}
 
-		// Darken items that aren't suitable targets for an active drag and drop action
-		if (FSlateApplication::Get().IsDragDropping())
-		{
-			TSharedPtr<FDragDropOperation> DragDropOp = FSlateApplication::Get().GetDragDroppingContent();
-
-			FDragDropPayload DraggedObjects;
-			const auto Outliner = WeakSceneOutliner.Pin();
-			if (Outliner->GetMode()->ParseDragDrop(DraggedObjects, *DragDropOp) && !Outliner->GetMode()->ValidateDrop(TreeItem, DraggedObjects).IsValid())
-			{
-				return DarkColor;
-			}
-		}
+	// Darken items that aren't suitable targets for an active drag and drop action
+	if (FSlateApplication::Get().IsDragDropping())
+	{
+		TSharedPtr<FDragDropOperation> DragDropOp = FSlateApplication::Get().GetDragDroppingContent();
 
-		if (!TreeItem.CanInteract())
+		FSceneOutlinerDragDropPayload DraggedObjects;
+		const auto Outliner = WeakSceneOutliner.Pin();
+		if (Outliner->GetMode()->ParseDragDrop(DraggedObjects, *DragDropOp) && !Outliner->GetMode()->ValidateDrop(TreeItem, DraggedObjects).IsValid())
 		{
 			return DarkColor;
 		}
+	}
 
-		return TOptional<FLinearColor>();
+	if (!TreeItem.CanInteract())
+	{
+		return DarkColor;
 	}
 
-	bool FCommonLabelData::CanExecuteRenameRequest(const ITreeItem& Item) const
+	return TOptional<FLinearColor>();
+}
+
+bool FSceneOutlinerCommonLabelData::CanExecuteRenameRequest(const ISceneOutlinerTreeItem& Item) const
+{
+	if (const ISceneOutliner* SceneOutliner = WeakSceneOutliner.Pin().Get())
 	{
-		if (const ISceneOutliner* SceneOutliner = WeakSceneOutliner.Pin().Get())
-		{
-			return SceneOutliner->CanExecuteRenameRequest(Item);
-		}
-		return false;
+		return SceneOutliner->CanExecuteRenameRequest(Item);
 	}
+	return false;
+}
 
+namespace SceneOutliner
+{
 	/** Parse a new path (including leaf-name) into this tree item. Does not do any notification */
 	FName GetFolderLeafName(FName InPath)
 	{
diff --git a/Engine/Source/Editor/SceneOutliner/Private/WorldTreeItem.cpp b/Engine/Source/Editor/SceneOutliner/Private/WorldTreeItem.cpp
--- a/Engine/Source/Editor/SceneOutliner/Private/WorldTreeItem.cpp
+++ b/Engine/Source/Editor/SceneOutliner/Private/WorldTreeItem.cpp
@@ -22,65 +22,65 @@
 
 namespace SceneOutliner
 {
-FText GetWorldDescription(UWorld* World)
-{
-	FText Description;
-	if (World)
+	FText GetWorldDescription(UWorld* World)
 	{
-		FText PostFix;
-		const FWorldContext* WorldContext = nullptr;
-		for (const FWorldContext& Context : GEngine->GetWorldContexts())
+		FText Description;
+		if (World)
 		{
-			if (Context.World() == World)
+			FText PostFix;
+			const FWorldContext* WorldContext = nullptr;
+			for (const FWorldContext& Context : GEngine->GetWorldContexts())
 			{
-				WorldContext = &Context;
-				break;
+				if (Context.World() == World)
+				{
+					WorldContext = &Context;
+					break;
+				}
 			}
-		}
 
-		if (World->WorldType == EWorldType::PIE)
-		{
-			switch (World->GetNetMode())
+			if (World->WorldType == EWorldType::PIE)
 			{
-			case NM_Client:
-				if (WorldContext)
-				{
-					PostFix = FText::Format(LOCTEXT("ClientPostfixFormat", "(Client {0})"), FText::AsNumber(WorldContext->PIEInstance - 1));
-				}
-				else
+				switch (World->GetNetMode())
 				{
-					PostFix = LOCTEXT("ClientPostfix", "(Client)");
+				case NM_Client:
+					if (WorldContext)
+					{
+						PostFix = FText::Format(LOCTEXT("ClientPostfixFormat", "(Client {0})"), FText::AsNumber(WorldContext->PIEInstance - 1));
+					}
+					else
+					{
+						PostFix = LOCTEXT("ClientPostfix", "(Client)");
+					}
+					break;
+				case NM_DedicatedServer:
+				case NM_ListenServer:
+					PostFix = LOCTEXT("ServerPostfix", "(Server)");
+					break;
+				case NM_Standalone:
+					PostFix = LOCTEXT("PlayInEditorPostfix", "(Play In Editor)");
+					break;
 				}
-				break;
-			case NM_DedicatedServer:
-			case NM_ListenServer:
-				PostFix = LOCTEXT("ServerPostfix", "(Server)");
-				break;
-			case NM_Standalone:
-				PostFix = LOCTEXT("PlayInEditorPostfix", "(Play In Editor)");
-				break;
 			}
-		}
-		else if (World->WorldType == EWorldType::Editor)
-		{
-			PostFix = LOCTEXT("EditorPostfix", "(Editor)");
+			else if (World->WorldType == EWorldType::Editor)
+			{
+				PostFix = LOCTEXT("EditorPostfix", "(Editor)");
+			}
+
+			Description = FText::Format(LOCTEXT("WorldFormat", "{0} {1}"), FText::FromString(World->GetFName().GetPlainNameString()), PostFix);
 		}
 
-		Description = FText::Format(LOCTEXT("WorldFormat", "{0} {1}"), FText::FromString(World->GetFName().GetPlainNameString()), PostFix);
+		return Description;
 	}
-
-	return Description;
 }
 
+const FSceneOutlinerTreeItemType FWorldTreeItem::Type(&ISceneOutlinerTreeItem::Type);
 
-const FTreeItemType FWorldTreeItem::Type(&ITreeItem::Type);
-
-struct SWorldTreeLabel : FCommonLabelData, public SCompoundWidget
+struct SWorldTreeLabel : FSceneOutlinerCommonLabelData, public SCompoundWidget
 {
 	SLATE_BEGIN_ARGS(SWorldTreeLabel) {}
 	SLATE_END_ARGS()
 
-		void Construct(const FArguments& InArgs, FWorldTreeItem& WorldItem, ISceneOutliner& SceneOutliner, const STableRow<FTreeItemPtr>& InRow)
+		void Construct(const FArguments& InArgs, FWorldTreeItem& WorldItem, ISceneOutliner& SceneOutliner, const STableRow<FSceneOutlinerTreeItemPtr>& InRow)
 	{
 		TreeItemPtr = StaticCastSharedRef<FWorldTreeItem>(WorldItem.AsShared());
 		WeakSceneOutliner = StaticCastSharedRef<ISceneOutliner>(SceneOutliner.AsShared());
@@ -92,11 +92,11 @@ struct SWorldTreeLabel : FCommonLabelData, public SCompoundWidget
 				+ SHorizontalBox::Slot()
 			.AutoWidth()
 			.VAlign(VAlign_Center)
-			.Padding(FDefaultTreeItemMetrics::IconPadding())
+			.Padding(FSceneOutlinerDefaultTreeItemMetrics::IconPadding())
 			[
 				SNew(SBox)
-				.WidthOverride(FDefaultTreeItemMetrics::IconSize())
-			.HeightOverride(FDefaultTreeItemMetrics::IconSize())
+				.WidthOverride(FSceneOutlinerDefaultTreeItemMetrics::IconSize())
+			.HeightOverride(FSceneOutlinerDefaultTreeItemMetrics::IconSize())
 			[
 				SNew(SImage)
 				.Image(FSlateIconFinder::FindIconBrushForClass(UWorld::StaticClass()))
@@ -147,7 +147,7 @@ private:
 
 	FSlateColor GetForegroundColor() const
 	{
-		if (auto BaseColor = FCommonLabelData::GetForegroundColor(*TreeItemPtr.Pin()))
+		if (auto BaseColor = FSceneOutlinerCommonLabelData::GetForegroundColor(*TreeItemPtr.Pin()))
 		{
 			return BaseColor.GetValue();
 		}
@@ -157,20 +157,20 @@ private:
 };
 
 FWorldTreeItem::FWorldTreeItem(UWorld* InWorld)
-	: ITreeItem(Type)
+	: ISceneOutlinerTreeItem(Type)
 	, World(InWorld)
 	, ID(InWorld)
 {
 }
 
 FWorldTreeItem::FWorldTreeItem(TWeakObjectPtr<UWorld> InWorld)
-	: ITreeItem(Type)
+	: ISceneOutlinerTreeItem(Type)
 	, World(InWorld)
 	, ID(InWorld.Get())
 {
 }
 
-FTreeItemID FWorldTreeItem::GetID() const
+FSceneOutlinerTreeItemID FWorldTreeItem::GetID() const
 {
 	return ID;
 }
@@ -213,7 +213,7 @@ void FWorldTreeItem::GenerateContextMenu(UToolMenu* Menu, SSceneOutliner& Outlin
 	Section.AddMenuEntry("OpenWorldSettings", LOCTEXT("OpenWorldSettings", "World Settings"), FText(), WorldSettingsIcon, FExecuteAction::CreateSP(this, &FWorldTreeItem::OpenWorldSettings));
 }
 
-TSharedRef<SWidget> FWorldTreeItem::GenerateLabelWidget(ISceneOutliner& Outliner, const STableRow<FTreeItemPtr>& InRow)
+TSharedRef<SWidget> FWorldTreeItem::GenerateLabelWidget(ISceneOutliner& Outliner, const STableRow<FSceneOutlinerTreeItemPtr>& InRow)
 {
 	return SNew(SWorldTreeLabel, *this, Outliner, InRow);
 }
@@ -224,6 +224,4 @@ void FWorldTreeItem::OpenWorldSettings() const
 	LevelEditorModule.GetLevelEditorTabManager()->TryInvokeTab(FName("WorldSettingsTab"));	
 }
 
-}		// namespace SceneOutliner
-
 #undef LOCTEXT_NAMESPACE
diff --git a/Engine/Source/Editor/SceneOutliner/Public/ActorBrowsingMode.h b/Engine/Source/Editor/SceneOutliner/Public/ActorBrowsingMode.h
--- a/Engine/Source/Editor/SceneOutliner/Public/ActorBrowsingMode.h
+++ b/Engine/Source/Editor/SceneOutliner/Public/ActorBrowsingMode.h
@@ -6,104 +6,101 @@
 
 class UActorBrowsingModeSettings;
 
-namespace SceneOutliner
+class FActorBrowsingMode : public FActorModeInteractive
 {
-	class FActorBrowsingMode : public FActorModeInteractive
-	{
-	public:
-		FActorBrowsingMode(SSceneOutliner* InSceneOutliner, TWeakObjectPtr<UWorld> InSpecifiedWorldToDisplay = nullptr);
-		virtual ~FActorBrowsingMode();
-
-		/* Begin ISceneOutlinerMode Interface */
-		virtual void Rebuild() override;
-		virtual FCreateOutlinerMode CreateFolderPickerMode() const override;
-		virtual void CreateViewContent(FMenuBuilder& MenuBuilder) override;
-		virtual TSharedPtr<FDragDropOperation> CreateDragDropOperation(const TArray<FTreeItemPtr>& InTreeItems) const override;
-		virtual bool ParseDragDrop(FDragDropPayload& OutPayload, const FDragDropOperation& Operation) const override;
-		virtual FDragValidationInfo ValidateDrop(const ITreeItem& DropTarget, const FDragDropPayload& Payload) const override;
-		virtual TSharedPtr<SWidget> CreateContextMenu() override;
-		virtual void OnItemAdded(FTreeItemPtr Item) override;
-		virtual void OnItemRemoved(FTreeItemPtr Item) override;
-		virtual void OnItemSelectionChanged(FTreeItemPtr Item, ESelectInfo::Type SelectionType, const FItemSelection& Selection) override;
-		virtual void OnItemDoubleClick(FTreeItemPtr Item) override;
-		virtual void OnFilterTextCommited(FItemSelection& Selection, ETextCommit::Type CommitType) override;
-		virtual void OnItemPassesFilters(const ITreeItem& Item) override;
-		virtual FReply OnKeyDown(const FKeyEvent& InKeyEvent) override;
-		virtual void OnDuplicateSelected() override;
-		virtual void OnDrop(ITreeItem& DropTarget, const FDragDropPayload& Payload, const FDragValidationInfo& ValidationInfo) const override;
-		virtual bool CanRenameItem(const ITreeItem& Item) const override;
-		virtual FText GetStatusText() const override;
-		virtual FSlateColor GetStatusTextColor() const override;
-		virtual ESelectionMode::Type GetSelectionMode() const override { return ESelectionMode::Multi; }
-		virtual bool SupportsKeyboardFocus() const override { return true; }
-		virtual bool ShouldShowFolders() const override { return true; }
-		virtual bool SupportsCreateNewFolder() const override { return true; }
-		virtual bool ShowStatusBar() const override { return true; }
-		virtual bool ShowViewButton() const override { return true; }
-		virtual bool ShowFilterOptions() const override { return true; }
-		virtual bool CanRename() const override;
-		virtual bool CanCut() const override;
-		virtual bool CanCopy() const override;
-		virtual bool CanPaste() const override;
-		virtual bool CanSupportDragAndDrop() const { return true; }
-		virtual FName CreateNewFolder() override;
-		virtual FName CreateFolder(const FName& ParentPath, const FName& LeafName) override;
-		virtual bool ReparentItemToFolder(const FName& FolderPath, const FTreeItemPtr& Item) override;
-		virtual void SelectFoldersDescendants(const TArray<FFolderTreeItem*>& FolderItems, bool bSelectImmediateChildrenOnly) override;
-		/* End ISceneOutlinerMode Interface */
-	public:
-		/* External events this mode must respond to */
-
-		/** Called by the engine when a component is updated */
-		void OnComponentsUpdated();
-		/** Called by the engine when an actor is deleted */
-		void OnLevelActorDeleted(AActor* Actor);
-
-		/** Called by engine when edit cut actors begins */
-		void OnEditCutActorsBegin();
-
-		/** Called by engine when edit cut actors ends */
-		void OnEditCutActorsEnd();
-
-		/** Called by engine when edit copy actors begins */
-		void OnEditCopyActorsBegin();
-
-		/** Called by engine when edit copy actors ends */
-		void OnEditCopyActorsEnd();
-
-		/** Called by engine when edit paste actors begins */
-		void OnEditPasteActorsBegin();
-
-		/** Called by engine when edit paste actors ends */
-		void OnEditPasteActorsEnd();
-
-		/** Called by engine when edit duplicate actors begins */
-		void OnDuplicateActorsBegin();
-
-		/** Called by engine when edit duplicate actors ends */
-		void OnDuplicateActorsEnd();
-
-		/** Called by engine when edit delete actors begins */
-		void OnDeleteActorsBegin();
-
-		/** Called by engine when edit delete actors ends */
-		void OnDeleteActorsEnd();
-	private:
-		/** Build and up the context menu */
-		TSharedPtr<SWidget> BuildContextMenu();
-		/** Register the context menu with the engine */
-		void RegisterContextMenu();
-		bool CanPasteFoldersOnlyFromClipboard() const;
-
-		/** Filter factories */
-		static TSharedRef<FOutlinerFilter> CreateShowOnlySelectedActorsFilter();
-		static TSharedRef<FOutlinerFilter> CreateHideTemporaryActorsFilter();
-		static TSharedRef<FOutlinerFilter> CreateIsInCurrentLevelFilter();
-		static TSharedRef<FOutlinerFilter> CreateHideComponentsFilter();
-	private:
-		/** Number of actors which have passed through the filters */
-		uint32 FilteredActorCount = 0;
-		/** List of actors which passed the regular filters and may or may not have passed the search filter */
-		TSet<TWeakObjectPtr<AActor>> ApplicableActors;
-	};
-}
\ No newline at end of file
+public:
+	FActorBrowsingMode(SSceneOutliner* InSceneOutliner, TWeakObjectPtr<UWorld> InSpecifiedWorldToDisplay = nullptr);
+	virtual ~FActorBrowsingMode();
+
+	/* Begin ISceneOutlinerMode Interface */
+	virtual void Rebuild() override;
+	virtual FCreateSceneOutlinerMode CreateFolderPickerMode() const override;
+	virtual void CreateViewContent(FMenuBuilder& MenuBuilder) override;
+	virtual TSharedPtr<FDragDropOperation> CreateDragDropOperation(const TArray<FSceneOutlinerTreeItemPtr>& InTreeItems) const override;
+	virtual bool ParseDragDrop(FSceneOutlinerDragDropPayload& OutPayload, const FDragDropOperation& Operation) const override;
+	virtual FSceneOutlinerDragValidationInfo ValidateDrop(const ISceneOutlinerTreeItem& DropTarget, const FSceneOutlinerDragDropPayload& Payload) const override;
+	virtual TSharedPtr<SWidget> CreateContextMenu() override;
+	virtual void OnItemAdded(FSceneOutlinerTreeItemPtr Item) override;
+	virtual void OnItemRemoved(FSceneOutlinerTreeItemPtr Item) override;
+	virtual void OnItemSelectionChanged(FSceneOutlinerTreeItemPtr Item, ESelectInfo::Type SelectionType, const FSceneOutlinerItemSelection& Selection) override;
+	virtual void OnItemDoubleClick(FSceneOutlinerTreeItemPtr Item) override;
+	virtual void OnFilterTextCommited(FSceneOutlinerItemSelection& Selection, ETextCommit::Type CommitType) override;
+	virtual void OnItemPassesFilters(const ISceneOutlinerTreeItem& Item) override;
+	virtual FReply OnKeyDown(const FKeyEvent& InKeyEvent) override;
+	virtual void OnDuplicateSelected() override;
+	virtual void OnDrop(ISceneOutlinerTreeItem& DropTarget, const FSceneOutlinerDragDropPayload& Payload, const FSceneOutlinerDragValidationInfo& ValidationInfo) const override;
+	virtual bool CanRenameItem(const ISceneOutlinerTreeItem& Item) const override;
+	virtual FText GetStatusText() const override;
+	virtual FSlateColor GetStatusTextColor() const override;
+	virtual ESelectionMode::Type GetSelectionMode() const override { return ESelectionMode::Multi; }
+	virtual bool SupportsKeyboardFocus() const override { return true; }
+	virtual bool ShouldShowFolders() const override { return true; }
+	virtual bool SupportsCreateNewFolder() const override { return true; }
+	virtual bool ShowStatusBar() const override { return true; }
+	virtual bool ShowViewButton() const override { return true; }
+	virtual bool ShowFilterOptions() const override { return true; }
+	virtual bool CanRename() const override;
+	virtual bool CanCut() const override;
+	virtual bool CanCopy() const override;
+	virtual bool CanPaste() const override;
+	virtual bool CanSupportDragAndDrop() const { return true; }
+	virtual FName CreateNewFolder() override;
+	virtual FName CreateFolder(const FName& ParentPath, const FName& LeafName) override;
+	virtual bool ReparentItemToFolder(const FName& FolderPath, const FSceneOutlinerTreeItemPtr& Item) override;
+	virtual void SelectFoldersDescendants(const TArray<FFolderTreeItem*>& FolderItems, bool bSelectImmediateChildrenOnly) override;
+	/* End ISceneOutlinerMode Interface */
+public:
+	/* External events this mode must respond to */
+
+	/** Called by the engine when a component is updated */
+	void OnComponentsUpdated();
+	/** Called by the engine when an actor is deleted */
+	void OnLevelActorDeleted(AActor* Actor);
+
+	/** Called by engine when edit cut actors begins */
+	void OnEditCutActorsBegin();
+
+	/** Called by engine when edit cut actors ends */
+	void OnEditCutActorsEnd();
+
+	/** Called by engine when edit copy actors begins */
+	void OnEditCopyActorsBegin();
+
+	/** Called by engine when edit copy actors ends */
+	void OnEditCopyActorsEnd();
+
+	/** Called by engine when edit paste actors begins */
+	void OnEditPasteActorsBegin();
+
+	/** Called by engine when edit paste actors ends */
+	void OnEditPasteActorsEnd();
+
+	/** Called by engine when edit duplicate actors begins */
+	void OnDuplicateActorsBegin();
+
+	/** Called by engine when edit duplicate actors ends */
+	void OnDuplicateActorsEnd();
+
+	/** Called by engine when edit delete actors begins */
+	void OnDeleteActorsBegin();
+
+	/** Called by engine when edit delete actors ends */
+	void OnDeleteActorsEnd();
+private:
+	/** Build and up the context menu */
+	TSharedPtr<SWidget> BuildContextMenu();
+	/** Register the context menu with the engine */
+	void RegisterContextMenu();
+	bool CanPasteFoldersOnlyFromClipboard() const;
+
+	/** Filter factories */
+	static TSharedRef<FSceneOutlinerFilter> CreateShowOnlySelectedActorsFilter();
+	static TSharedRef<FSceneOutlinerFilter> CreateHideTemporaryActorsFilter();
+	static TSharedRef<FSceneOutlinerFilter> CreateIsInCurrentLevelFilter();
+	static TSharedRef<FSceneOutlinerFilter> CreateHideComponentsFilter();
+private:
+	/** Number of actors which have passed through the filters */
+	uint32 FilteredActorCount = 0;
+	/** List of actors which passed the regular filters and may or may not have passed the search filter */
+	TSet<TWeakObjectPtr<AActor>> ApplicableActors;
+};
diff --git a/Engine/Source/Editor/SceneOutliner/Public/ActorFolderHierarchy.h b/Engine/Source/Editor/SceneOutliner/Public/ActorFolderHierarchy.h
--- a/Engine/Source/Editor/SceneOutliner/Public/ActorFolderHierarchy.h
+++ b/Engine/Source/Editor/SceneOutliner/Public/ActorFolderHierarchy.h
@@ -5,27 +5,24 @@
 #include "CoreMinimal.h"
 #include "ISceneOutlinerHierarchy.h"
 
-namespace SceneOutliner
+class FActorFolderHierarchy : public ISceneOutlinerHierarchy
 {
-	class FActorFolderHierarchy : public ISceneOutlinerHierarchy
-	{
-	public:
-		FActorFolderHierarchy(ISceneOutlinerMode* InMode, const TWeakObjectPtr<UWorld>& World);
-		virtual ~FActorFolderHierarchy() {}
+public:
+	FActorFolderHierarchy(ISceneOutlinerMode* InMode, const TWeakObjectPtr<UWorld>& World);
+	virtual ~FActorFolderHierarchy() {}
 
 
-		/* Begin ISceneOutlinerHierarchy Implementation */
-		virtual void FindChildren(const ITreeItem& Item, const TMap<FTreeItemID, FTreeItemPtr>& Items, TArray<FTreeItemPtr>& OutChildItems) const override;
-		virtual FTreeItemPtr FindParent(const ITreeItem& Item, const TMap<FTreeItemID, FTreeItemPtr>& Items) const override;
-		virtual void CreateItems(TArray<FTreeItemPtr>& OutItems) const override;
-		virtual void CreateChildren(const FTreeItemPtr& Item, TArray<FTreeItemPtr>& OutChildren) const override;
-		virtual FTreeItemPtr CreateParentItem(const FTreeItemPtr& Item) const override;
-		/* End ISceneOutlinerHierarchy Implementation */
-	private:
-		/** Adds all the direct and indirect children of a world to OutItems */
-		void CreateWorldChildren(UWorld* World, TArray<FTreeItemPtr>& OutItems) const;
-	private:
-		/** The world which this hierarchy is representing */
-		TWeakObjectPtr<UWorld> RepresentingWorld;
-	};
-}
\ No newline at end of file
+	/* Begin ISceneOutlinerHierarchy Implementation */
+	virtual void FindChildren(const ISceneOutlinerTreeItem& Item, const TMap<FSceneOutlinerTreeItemID, FSceneOutlinerTreeItemPtr>& Items, TArray<FSceneOutlinerTreeItemPtr>& OutChildItems) const override;
+	virtual FSceneOutlinerTreeItemPtr FindParent(const ISceneOutlinerTreeItem& Item, const TMap<FSceneOutlinerTreeItemID, FSceneOutlinerTreeItemPtr>& Items) const override;
+	virtual void CreateItems(TArray<FSceneOutlinerTreeItemPtr>& OutItems) const override;
+	virtual void CreateChildren(const FSceneOutlinerTreeItemPtr& Item, TArray<FSceneOutlinerTreeItemPtr>& OutChildren) const override;
+	virtual FSceneOutlinerTreeItemPtr CreateParentItem(const FSceneOutlinerTreeItemPtr& Item) const override;
+	/* End ISceneOutlinerHierarchy Implementation */
+private:
+	/** Adds all the direct and indirect children of a world to OutItems */
+	void CreateWorldChildren(UWorld* World, TArray<FSceneOutlinerTreeItemPtr>& OutItems) const;
+private:
+	/** The world which this hierarchy is representing */
+	TWeakObjectPtr<UWorld> RepresentingWorld;
+};
diff --git a/Engine/Source/Editor/SceneOutliner/Public/ActorFolderPickingMode.h b/Engine/Source/Editor/SceneOutliner/Public/ActorFolderPickingMode.h
--- a/Engine/Source/Editor/SceneOutliner/Public/ActorFolderPickingMode.h
+++ b/Engine/Source/Editor/SceneOutliner/Public/ActorFolderPickingMode.h
@@ -4,26 +4,22 @@
 
 #include "ActorMode.h"
 
-namespace SceneOutliner
+class FActorFolderPickingMode : public FActorMode
 {
+public:
+	FActorFolderPickingMode(SSceneOutliner* InSceneOutliner, FOnSceneOutlinerItemPicked InOnItemPicked, TWeakObjectPtr<UWorld> InSpecifiedWorldToDisplay = nullptr);
+	virtual ~FActorFolderPickingMode() {}
 
-	class FActorFolderPickingMode : public FActorMode
-	{
-	public:
-		FActorFolderPickingMode(SSceneOutliner* InSceneOutliner, FOnSceneOutlinerItemPicked InOnItemPicked, TWeakObjectPtr<UWorld> InSpecifiedWorldToDisplay = nullptr);
-		virtual ~FActorFolderPickingMode() {}
+	/* Begin ISceneOutlinerMode Implementation */
+	virtual void OnFilterTextCommited(FSceneOutlinerItemSelection& Selection, ETextCommit::Type CommitType) override;
+	virtual void OnItemSelectionChanged(FSceneOutlinerTreeItemPtr Item, ESelectInfo::Type SelectionType, const FSceneOutlinerItemSelection& Selection) override;
+	virtual void CreateViewContent(FMenuBuilder& MenuBuilder) override;
+	virtual bool ShowViewButton() const override { return true; }
+	virtual bool ShouldShowFolders() const { return true; }
+protected:
+	virtual TUniquePtr<ISceneOutlinerHierarchy> CreateHierarchy() override;
+	/* End ISceneOutlinerMode Implementation */
 
-		/* Begin ISceneOutlinerMode Implementation */
-		virtual void OnFilterTextCommited(FItemSelection& Selection, ETextCommit::Type CommitType) override;
-		virtual void OnItemSelectionChanged(FTreeItemPtr Item, ESelectInfo::Type SelectionType, const FItemSelection& Selection) override;
-		virtual void CreateViewContent(FMenuBuilder& MenuBuilder) override;
-		virtual bool ShowViewButton() const override { return true; }
-		virtual bool ShouldShowFolders() const { return true; }
-	protected:
-		virtual TUniquePtr<ISceneOutlinerHierarchy> CreateHierarchy() override;
-		/* End ISceneOutlinerMode Implementation */
-
-		/** Delegate to call when an item is picked */
-		FOnSceneOutlinerItemPicked OnItemPicked;
-	};
-}
\ No newline at end of file
+	/** Delegate to call when an item is picked */
+	FOnSceneOutlinerItemPicked OnItemPicked;
+};
\ No newline at end of file
diff --git a/Engine/Source/Editor/SceneOutliner/Public/ActorFolderTreeItem.h b/Engine/Source/Editor/SceneOutliner/Public/ActorFolderTreeItem.h
--- a/Engine/Source/Editor/SceneOutliner/Public/ActorFolderTreeItem.h
+++ b/Engine/Source/Editor/SceneOutliner/Public/ActorFolderTreeItem.h
@@ -5,30 +5,27 @@
 #include "CoreMinimal.h"
 #include "FolderTreeItem.h"
 
-namespace SceneOutliner
+struct SCENEOUTLINER_API FActorFolderTreeItem : public FFolderTreeItem
 {
-	struct SCENEOUTLINER_API FActorFolderTreeItem : public FFolderTreeItem
-	{
-	public:
-		/** Static type identifier for this tree item class */
-		static const FTreeItemType Type;
+public:
+	/** Static type identifier for this tree item class */
+	static const FSceneOutlinerTreeItemType Type;
 
-		FActorFolderTreeItem(FName InPath, TWeakObjectPtr<UWorld> InWorld);
+	FActorFolderTreeItem(FName InPath, TWeakObjectPtr<UWorld> InWorld);
 
-		/** The world which this folder belongs to */
-		TWeakObjectPtr<UWorld> World;
+	/** The world which this folder belongs to */
+	TWeakObjectPtr<UWorld> World;
 
-		/* Begin ITreeItem Implementation */
-		virtual bool IsValid() const { return World.IsValid(); }
-		virtual void OnExpansionChanged() override;
-		virtual void Delete(FName InNewParentPath) override;
-		virtual TSharedRef<SWidget> GenerateLabelWidget(ISceneOutliner& Outliner, const STableRow<FTreeItemPtr>& InRow) override;
-		/* End FFolderTreeItem Implementation */
+	/* Begin ISceneOutlinerTreeItem Implementation */
+	virtual bool IsValid() const { return World.IsValid(); }
+	virtual void OnExpansionChanged() override;
+	virtual void Delete(FName InNewParentPath) override;
+	virtual TSharedRef<SWidget> GenerateLabelWidget(ISceneOutliner& Outliner, const STableRow<FSceneOutlinerTreeItemPtr>& InRow) override;
+	/* End FFolderTreeItem Implementation */
 		
-		/* Begin FFolderTreeItem Implementation */
-		virtual FName MoveTo(const FName& NewParent) override;
-	private:
-		virtual void CreateSubFolder(TWeakPtr<SSceneOutliner> WeakOutliner) override;
-		/* End FFolderTreeItem Implementation */
-	};
-}
\ No newline at end of file
+	/* Begin FFolderTreeItem Implementation */
+	virtual FName MoveTo(const FName& NewParent) override;
+private:
+	virtual void CreateSubFolder(TWeakPtr<SSceneOutliner> WeakOutliner) override;
+	/* End FFolderTreeItem Implementation */
+};
diff --git a/Engine/Source/Editor/SceneOutliner/Public/ActorHierarchy.h b/Engine/Source/Editor/SceneOutliner/Public/ActorHierarchy.h
--- a/Engine/Source/Editor/SceneOutliner/Public/ActorHierarchy.h
+++ b/Engine/Source/Editor/SceneOutliner/Public/ActorHierarchy.h
@@ -5,75 +5,72 @@
 #include "CoreMinimal.h"
 #include "ISceneOutlinerHierarchy.h"
 
-namespace SceneOutliner
+class FActorHierarchy : public ISceneOutlinerHierarchy
 {
-	class FActorHierarchy : public ISceneOutlinerHierarchy
-	{
-	public:
-		virtual ~FActorHierarchy();
-
-		static TUniquePtr<FActorHierarchy> Create(ISceneOutlinerMode* Mode, const TWeakObjectPtr<UWorld>& World, TFunction<bool(const AActor*)> InIsActorDisplayableCallback);
-
-		/** Find all direct children of a tree item in an existing item map, if any exist. */
-		virtual void FindChildren(const ITreeItem& Item, const TMap<FTreeItemID, FTreeItemPtr>& Items, TArray<FTreeItemPtr>& OutChildItems) const override;
-		/** Find the parent of a tree item in an existing item map, it if exists. */
-		virtual FTreeItemPtr FindParent(const ITreeItem& Item, const TMap<FTreeItemID, FTreeItemPtr>& Items) const override;
-
-		/** Create a linearization of all applicable items in the hierarchy */
-		virtual void CreateItems(TArray<FTreeItemPtr>& OutItems) const override;
-		/** Create a linearization of all direct and indirect children of a given item in the hierarchy */
-		virtual void CreateChildren(const FTreeItemPtr& Item, TArray<FTreeItemPtr>& OutChildren) const override;
-		/** Forcibly create a parent item for a given tree item */
-		virtual FTreeItemPtr CreateParentItem(const FTreeItemPtr& Item) const override;
-
-		void SetShowingComponents(bool bInShowingComponents) { bShowingComponents = bInShowingComponents; }
-	private:
-		/** Adds all the direct and indirect children of a world to OutItems */
-		void CreateWorldChildren(UWorld* World, TArray<FTreeItemPtr>& OutItems) const;
-		/** Create all component items for an actor if we are showing components and place them in OutItems */
-		void CreateComponentItems(const AActor* Actor, TArray<FTreeItemPtr>& OutItems) const;
-	private:
-		// Update the hierarchy when actor or world changing events occur
-
-		void OnLevelActorAdded(AActor* InActor);
-		void OnLevelActorDeleted(AActor* InActor);
+public:
+	virtual ~FActorHierarchy();
+
+	static TUniquePtr<FActorHierarchy> Create(ISceneOutlinerMode* Mode, const TWeakObjectPtr<UWorld>& World, TFunction<bool(const AActor*)> InIsActorDisplayableCallback);
+
+	/** Find all direct children of a tree item in an existing item map, if any exist. */
+	virtual void FindChildren(const ISceneOutlinerTreeItem& Item, const TMap<FSceneOutlinerTreeItemID, FSceneOutlinerTreeItemPtr>& Items, TArray<FSceneOutlinerTreeItemPtr>& OutChildItems) const override;
+	/** Find the parent of a tree item in an existing item map, it if exists. */
+	virtual FSceneOutlinerTreeItemPtr FindParent(const ISceneOutlinerTreeItem& Item, const TMap<FSceneOutlinerTreeItemID, FSceneOutlinerTreeItemPtr>& Items) const override;
+
+	/** Create a linearization of all applicable items in the hierarchy */
+	virtual void CreateItems(TArray<FSceneOutlinerTreeItemPtr>& OutItems) const override;
+	/** Create a linearization of all direct and indirect children of a given item in the hierarchy */
+	virtual void CreateChildren(const FSceneOutlinerTreeItemPtr& Item, TArray<FSceneOutlinerTreeItemPtr>& OutChildren) const override;
+	/** Forcibly create a parent item for a given tree item */
+	virtual FSceneOutlinerTreeItemPtr CreateParentItem(const FSceneOutlinerTreeItemPtr& Item) const override;
+
+	void SetShowingComponents(bool bInShowingComponents) { bShowingComponents = bInShowingComponents; }
+private:
+	/** Adds all the direct and indirect children of a world to OutItems */
+	void CreateWorldChildren(UWorld* World, TArray<FSceneOutlinerTreeItemPtr>& OutItems) const;
+	/** Create all component items for an actor if we are showing components and place them in OutItems */
+	void CreateComponentItems(const AActor* Actor, TArray<FSceneOutlinerTreeItemPtr>& OutItems) const;
+private:
+	// Update the hierarchy when actor or world changing events occur
+
+	void OnLevelActorAdded(AActor* InActor);
+	void OnLevelActorDeleted(AActor* InActor);
 		
-		void OnLevelActorAttached(AActor* InActor, const AActor* InParent);
-		void OnLevelActorDetached(AActor* InActor, const AActor* InParent);
+	void OnLevelActorAttached(AActor* InActor, const AActor* InParent);
+	void OnLevelActorDetached(AActor* InActor, const AActor* InParent);
 
-		void OnLevelActorFolderChanged(const AActor* InActor, FName OldPath);
+	void OnLevelActorFolderChanged(const AActor* InActor, FName OldPath);
 
-		void OnComponentsUpdated();
+	void OnComponentsUpdated();
 
-		void OnLevelActorListChanged();
+	void OnLevelActorListChanged();
 		
-		void OnLevelAdded(ULevel* InLevel, UWorld* InWorld);
-		void OnLevelRemoved(ULevel* InLevel, UWorld* InWorld);
+	void OnLevelAdded(ULevel* InLevel, UWorld* InWorld);
+	void OnLevelRemoved(ULevel* InLevel, UWorld* InWorld);
 
-		/** Called when a folder is to be created */
-		void OnBroadcastFolderCreate(UWorld& InWorld, FName NewPath);
+	/** Called when a folder is to be created */
+	void OnBroadcastFolderCreate(UWorld& InWorld, FName NewPath);
 
-		/** Called when a folder is to be moved */
-		void OnBroadcastFolderMove(UWorld& InWorld, FName OldPath, FName NewPath);
+	/** Called when a folder is to be moved */
+	void OnBroadcastFolderMove(UWorld& InWorld, FName OldPath, FName NewPath);
 
-		/** Called when a folder is to be deleted */
-		void OnBroadcastFolderDelete(UWorld& InWorld, FName Path);
+	/** Called when a folder is to be deleted */
+	void OnBroadcastFolderDelete(UWorld& InWorld, FName Path);
 
-	private:
-		/** Send a an event indicating a full refresh of the hierarchy is required */
-		void FullRefreshEvent();
+private:
+	/** Send a an event indicating a full refresh of the hierarchy is required */
+	void FullRefreshEvent();
 
-		bool bShowingComponents = false;
+	bool bShowingComponents = false;
 
-		TWeakObjectPtr<UWorld> RepresentingWorld;
+	TWeakObjectPtr<UWorld> RepresentingWorld;
 
-		/** Used to determine if an actor can be displayed in the outliner */
-		TFunction<bool(const AActor*)> IsActorDisplayableCallback;
+	/** Used to determine if an actor can be displayed in the outliner */
+	TFunction<bool(const AActor*)> IsActorDisplayableCallback;
 
-	private:
-		FActorHierarchy(ISceneOutlinerMode* Mode, const TWeakObjectPtr<UWorld>& Worlds, TFunction<bool(const AActor*)> InIsActorDisplayableCallback);
+private:
+	FActorHierarchy(ISceneOutlinerMode* Mode, const TWeakObjectPtr<UWorld>& Worlds, TFunction<bool(const AActor*)> InIsActorDisplayableCallback);
 
-		FActorHierarchy(const FActorHierarchy&) = delete;
-		FActorHierarchy& operator=(const FActorHierarchy&) = delete;
-	};
-}
\ No newline at end of file
+	FActorHierarchy(const FActorHierarchy&) = delete;
+	FActorHierarchy& operator=(const FActorHierarchy&) = delete;
+};
diff --git a/Engine/Source/Editor/SceneOutliner/Public/ActorMode.h b/Engine/Source/Editor/SceneOutliner/Public/ActorMode.h
--- a/Engine/Source/Editor/SceneOutliner/Public/ActorMode.h
+++ b/Engine/Source/Editor/SceneOutliner/Public/ActorMode.h
@@ -9,57 +9,57 @@ namespace SceneOutliner
 	/** Functor which can be used to get weak actor pointers from a selection */
 	struct FWeakActorSelector
 	{
-		bool operator()(const TWeakPtr<ITreeItem>& Item, TWeakObjectPtr<AActor>& DataOut) const;
+		bool operator()(const TWeakPtr<ISceneOutlinerTreeItem>& Item, TWeakObjectPtr<AActor>& DataOut) const;
 	};
 
 	/** Functor which can be used to get actors from a selection including component parents */
 	struct FActorSelector
 	{
-		bool operator()(const TWeakPtr<ITreeItem>& Item, AActor*& ActorPtrOut) const;
+		bool operator()(const TWeakPtr<ISceneOutlinerTreeItem>& Item, AActor*& ActorPtrOut) const;
 	};
+}
 
-	class SCENEOUTLINER_API FActorMode : public ISceneOutlinerMode
+class SCENEOUTLINER_API FActorMode : public ISceneOutlinerMode
+{
+public:
+	struct EItemSortOrder
 	{
-	public:
-		struct EItemSortOrder
-		{
-			enum Type { World = 0, Folder = 10, Actor = 20 };
-		};
+		enum Type { World = 0, Folder = 10, Actor = 20 };
+	};
 	
-		FActorMode(SSceneOutliner* InSceneOutliner, bool bHideComponents, TWeakObjectPtr<UWorld> InSpecifiedWorldToDisplay = nullptr);
-		virtual ~FActorMode();
+	FActorMode(SSceneOutliner* InSceneOutliner, bool bHideComponents, TWeakObjectPtr<UWorld> InSpecifiedWorldToDisplay = nullptr);
+	virtual ~FActorMode();
 
-		virtual void Rebuild() override;
+	virtual void Rebuild() override;
 
-		void BuildWorldPickerMenu(FMenuBuilder& MenuBuilder);
+	void BuildWorldPickerMenu(FMenuBuilder& MenuBuilder);
 
-		virtual void SynchronizeSelection() override { SynchronizeActorSelection(); }
+	virtual void SynchronizeSelection() override { SynchronizeActorSelection(); }
 
-		virtual void OnFilterTextChanged(const FText& InFilterText) override;
+	virtual void OnFilterTextChanged(const FText& InFilterText) override;
 
-		virtual int32 GetTypeSortPriority(const ITreeItem& Item) const override;
-	private:
-		/** Called when the user selects a world in the world picker menu */
-		void OnSelectWorld(TWeakObjectPtr<UWorld> World);
-	private:
-		/* Private Helpers */
+	virtual int32 GetTypeSortPriority(const ISceneOutlinerTreeItem& Item) const override;
+private:
+	/** Called when the user selects a world in the world picker menu */
+	void OnSelectWorld(TWeakObjectPtr<UWorld> World);
+private:
+	/* Private Helpers */
 
-		void ChooseRepresentingWorld();
-		bool IsWorldChecked(TWeakObjectPtr<UWorld> World) const;
-	protected:
-		void SynchronizeActorSelection();
-		bool IsActorDisplayable(const AActor* InActor) const;
+	void ChooseRepresentingWorld();
+	bool IsWorldChecked(TWeakObjectPtr<UWorld> World) const;
+protected:
+	void SynchronizeActorSelection();
+	bool IsActorDisplayable(const AActor* InActor) const;
 
-		virtual TUniquePtr<ISceneOutlinerHierarchy> CreateHierarchy() override;
-	protected:
-		// Should the hide components filter be enabled
-		bool bHideComponents;
-		/** The world which we are currently representing */
-		TWeakObjectPtr<UWorld> RepresentingWorld;
-		/** The world which the user manually selected */
-		TWeakObjectPtr<UWorld> UserChosenWorld;
+	virtual TUniquePtr<ISceneOutlinerHierarchy> CreateHierarchy() override;
+protected:
+	// Should the hide components filter be enabled
+	bool bHideComponents;
+	/** The world which we are currently representing */
+	TWeakObjectPtr<UWorld> RepresentingWorld;
+	/** The world which the user manually selected */
+	TWeakObjectPtr<UWorld> UserChosenWorld;
 
-		/** If this mode was created to display a specific world, don't allow it to be reassigned */
-		const TWeakObjectPtr<UWorld> SpecifiedWorldToDisplay;
-	};
-}
\ No newline at end of file
+	/** If this mode was created to display a specific world, don't allow it to be reassigned */
+	const TWeakObjectPtr<UWorld> SpecifiedWorldToDisplay;
+};
diff --git a/Engine/Source/Editor/SceneOutliner/Public/ActorModeInteractive.h b/Engine/Source/Editor/SceneOutliner/Public/ActorModeInteractive.h
--- a/Engine/Source/Editor/SceneOutliner/Public/ActorModeInteractive.h
+++ b/Engine/Source/Editor/SceneOutliner/Public/ActorModeInteractive.h
@@ -4,23 +4,20 @@
 
 #include "ActorMode.h"
 
-namespace SceneOutliner
+class FActorModeInteractive : public FActorMode
 {
-	class FActorModeInteractive : public FActorMode
-	{
-	public:
-		FActorModeInteractive(SSceneOutliner* InSceneOutliner, bool bHideComponents, TWeakObjectPtr<UWorld> InSpecifiedWorldToDisplay = nullptr);
-		virtual ~FActorModeInteractive();
+public:
+	FActorModeInteractive(SSceneOutliner* InSceneOutliner, bool bHideComponents, TWeakObjectPtr<UWorld> InSpecifiedWorldToDisplay = nullptr);
+	virtual ~FActorModeInteractive();
 
-		virtual bool IsInteractive() const override { return true; }
-	private:
-		/* Events */
+	virtual bool IsInteractive() const override { return true; }
+private:
+	/* Events */
 
-		void OnMapChange(uint32 MapFlags);
-		void OnNewCurrentLevel();
+	void OnMapChange(uint32 MapFlags);
+	void OnNewCurrentLevel();
 
-		void OnLevelSelectionChanged(UObject* Obj);
-		void OnActorLabelChanged(AActor* ChangedActor);
-		void OnLevelActorRequestsRename(const AActor* Actor);
-	};
-}
+	void OnLevelSelectionChanged(UObject* Obj);
+	void OnActorLabelChanged(AActor* ChangedActor);
+	void OnLevelActorRequestsRename(const AActor* Actor);
+};
diff --git a/Engine/Source/Editor/SceneOutliner/Public/ActorPickingMode.h b/Engine/Source/Editor/SceneOutliner/Public/ActorPickingMode.h
--- a/Engine/Source/Editor/SceneOutliner/Public/ActorPickingMode.h
+++ b/Engine/Source/Editor/SceneOutliner/Public/ActorPickingMode.h
@@ -4,24 +4,21 @@
 
 #include "ActorMode.h"
 
-namespace SceneOutliner
+class FActorPickingMode : public FActorMode
 {
-	class FActorPickingMode : public FActorMode
-	{
-	public:
-		FActorPickingMode(SSceneOutliner* InSceneOutliner, bool bInHideComponents, FOnSceneOutlinerItemPicked OnItemPickedDelegate, TWeakObjectPtr<UWorld> InSpecifiedWorldToDisplay = nullptr);
+public:
+	FActorPickingMode(SSceneOutliner* InSceneOutliner, bool bInHideComponents, FOnSceneOutlinerItemPicked OnItemPickedDelegate, TWeakObjectPtr<UWorld> InSpecifiedWorldToDisplay = nullptr);
 
-		virtual ~FActorPickingMode() {};
-	public:
-		virtual void OnItemSelectionChanged(FTreeItemPtr Item, ESelectInfo::Type SelectionType, const FItemSelection& Selection) override;
+	virtual ~FActorPickingMode() {};
+public:
+	virtual void OnItemSelectionChanged(FSceneOutlinerTreeItemPtr Item, ESelectInfo::Type SelectionType, const FSceneOutlinerItemSelection& Selection) override;
 
-		/** Allow the user to commit their selection by pressing enter if it is valid */
-		virtual void OnFilterTextCommited(FItemSelection& Selection, ETextCommit::Type CommitType) override;
+	/** Allow the user to commit their selection by pressing enter if it is valid */
+	virtual void OnFilterTextCommited(FSceneOutlinerItemSelection& Selection, ETextCommit::Type CommitType) override;
 
-		virtual void CreateViewContent(FMenuBuilder& MenuBuilder) override;
-	public:
-		virtual bool ShowViewButton() const override { return true; }
-	private:
-		FOnSceneOutlinerItemPicked OnItemPicked;
-	};
-}
\ No newline at end of file
+	virtual void CreateViewContent(FMenuBuilder& MenuBuilder) override;
+public:
+	virtual bool ShowViewButton() const override { return true; }
+private:
+	FOnSceneOutlinerItemPicked OnItemPicked;
+};
\ No newline at end of file
diff --git a/Engine/Source/Editor/SceneOutliner/Public/ActorTreeItem.h b/Engine/Source/Editor/SceneOutliner/Public/ActorTreeItem.h
--- a/Engine/Source/Editor/SceneOutliner/Public/ActorTreeItem.h
+++ b/Engine/Source/Editor/SceneOutliner/Public/ActorTreeItem.h
@@ -4,53 +4,49 @@
 
 #include "CoreMinimal.h"
 #include "Widgets/SWidget.h"
-#include "ITreeItem.h"
+#include "ISceneOutlinerTreeItem.h"
 #include "UObject/ObjectKey.h"
 
-namespace SceneOutliner
+/** A tree item that represents an actor in the world */
+struct SCENEOUTLINER_API FActorTreeItem : ISceneOutlinerTreeItem
 {
-	/** A tree item that represents an actor in the world */
-	struct SCENEOUTLINER_API FActorTreeItem : ITreeItem
-	{
-	public:
-		DECLARE_DELEGATE_RetVal_OneParam(bool, FFilterPredicate, const AActor*);
-		DECLARE_DELEGATE_RetVal_OneParam(bool, FInteractivePredicate, const AActor*);
+public:
+	DECLARE_DELEGATE_RetVal_OneParam(bool, FFilterPredicate, const AActor*);
+	DECLARE_DELEGATE_RetVal_OneParam(bool, FInteractivePredicate, const AActor*);
 		
-		bool Filter(FFilterPredicate Pred) const
-		{
-			return Pred.Execute(Actor.Get());
-		}
-
-		bool GetInteractiveState(FInteractivePredicate Pred) const 
-		{
-			return Pred.Execute(Actor.Get());
-		}
-
-		/** The actor this tree item is associated with. */
-		mutable TWeakObjectPtr<AActor> Actor;
-
-		/** Constant identifier for this tree item */
-		const FObjectKey ID;
-
-		/** Static type identifier for this tree item class */
-		static const FTreeItemType Type;
-
-		/** Construct this item from an actor */
-		FActorTreeItem(AActor* InActor);
-
-		/* Begin ITreeItem Implementation */
-		virtual bool IsValid() const override { return Actor.IsValid(); }
-		virtual FTreeItemID GetID() const override;
-		virtual FString GetDisplayString() const override;
-		virtual bool CanInteract() const override;
-		virtual TSharedRef<SWidget> GenerateLabelWidget(ISceneOutliner& Outliner, const STableRow<FTreeItemPtr>& InRow) override;
-		virtual void OnVisibilityChanged(const bool bNewVisibility) override;
-		virtual bool HasVisibilityInfo() const override { return true; }
-		virtual bool GetVisibility() const override;
-		/* End ITreeItem Implementation */
-	public:
-		/** true if this item exists in both the current world and PIE. */
-		bool bExistsInCurrentWorldAndPIE;
-	};
-
-}		// namespace SceneOutliner
+	bool Filter(FFilterPredicate Pred) const
+	{
+		return Pred.Execute(Actor.Get());
+	}
+
+	bool GetInteractiveState(FInteractivePredicate Pred) const 
+	{
+		return Pred.Execute(Actor.Get());
+	}
+
+	/** The actor this tree item is associated with. */
+	mutable TWeakObjectPtr<AActor> Actor;
+
+	/** Constant identifier for this tree item */
+	const FObjectKey ID;
+
+	/** Static type identifier for this tree item class */
+	static const FSceneOutlinerTreeItemType Type;
+
+	/** Construct this item from an actor */
+	FActorTreeItem(AActor* InActor);
+
+	/* Begin ISceneOutlinerTreeItem Implementation */
+	virtual bool IsValid() const override { return Actor.IsValid(); }
+	virtual FSceneOutlinerTreeItemID GetID() const override;
+	virtual FString GetDisplayString() const override;
+	virtual bool CanInteract() const override;
+	virtual TSharedRef<SWidget> GenerateLabelWidget(ISceneOutliner& Outliner, const STableRow<FSceneOutlinerTreeItemPtr>& InRow) override;
+	virtual void OnVisibilityChanged(const bool bNewVisibility) override;
+	virtual bool HasVisibilityInfo() const override { return true; }
+	virtual bool GetVisibility() const override;
+	/* End ISceneOutlinerTreeItem Implementation */
+public:
+	/** true if this item exists in both the current world and PIE. */
+	bool bExistsInCurrentWorldAndPIE;
+};
diff --git a/Engine/Source/Editor/SceneOutliner/Public/ComponentTreeItem.h b/Engine/Source/Editor/SceneOutliner/Public/ComponentTreeItem.h
--- a/Engine/Source/Editor/SceneOutliner/Public/ComponentTreeItem.h
+++ b/Engine/Source/Editor/SceneOutliner/Public/ComponentTreeItem.h
@@ -4,54 +4,49 @@
 
 #include "CoreMinimal.h"
 #include "Widgets/SWidget.h"
-#include "ITreeItem.h"
+#include "ISceneOutlinerTreeItem.h"
 #include "UObject/ObjectKey.h"
 #include "Components/SceneComponent.h"
 
-
-namespace SceneOutliner
+/** A tree item that represents an Component in the world */
+struct SCENEOUTLINER_API FComponentTreeItem : ISceneOutlinerTreeItem
 {
-	/** A tree item that represents an Component in the world */
-	struct SCENEOUTLINER_API FComponentTreeItem : ITreeItem
+public:
+	DECLARE_DELEGATE_RetVal_OneParam(bool, FFilterPredicate, const UActorComponent*);
+	DECLARE_DELEGATE_RetVal_OneParam(bool, FInteractivePredicate, const UActorComponent*);
+
+	bool Filter(FFilterPredicate Pred) const
+	{
+		return Pred.Execute(Component.Get());
+	}
+
+	bool GetInteractiveState(FInteractivePredicate Pred) const
 	{
-	public:
-		DECLARE_DELEGATE_RetVal_OneParam(bool, FFilterPredicate, const UActorComponent*);
-		DECLARE_DELEGATE_RetVal_OneParam(bool, FInteractivePredicate, const UActorComponent*);
-
-		bool Filter(FFilterPredicate Pred) const
-		{
-			return Pred.Execute(Component.Get());
-		}
-
-		bool GetInteractiveState(FInteractivePredicate Pred) const
-		{
-			return Pred.Execute(Component.Get());
-		}
-
-		/** The Component this tree item is associated with. */
-		mutable TWeakObjectPtr<UActorComponent> Component;
-
-		/** Constant identifier for this tree item */
-		const FObjectKey ID;
-
-		/** Static type identifier for this tree item class */
-		static const FTreeItemType Type;
-
-		/** Construct this item from an Component */
-		FComponentTreeItem(UActorComponent* InComponent);
-
-		/* Begin ITreeItem Implementation */
-		virtual bool IsValid() const override { return Component.IsValid(); }
-		virtual FTreeItemID GetID() const override;
-		virtual FString GetDisplayString() const override;
-		virtual bool CanInteract() const override;
-		virtual TSharedRef<SWidget> GenerateLabelWidget(ISceneOutliner& Outliner, const STableRow<FTreeItemPtr>& InRow) override;
-		/* End ITreeItem Implementation */
-	public:
-		/** true if this item exists in both the current world and PIE. */
-		bool bExistsInCurrentWorldAndPIE;
-		/** Cache the string displayed */
-		FString CachedDisplayString;
-	};
-
-} // namespace SceneOutliner
+		return Pred.Execute(Component.Get());
+	}
+
+	/** The Component this tree item is associated with. */
+	mutable TWeakObjectPtr<UActorComponent> Component;
+
+	/** Constant identifier for this tree item */
+	const FObjectKey ID;
+
+	/** Static type identifier for this tree item class */
+	static const FSceneOutlinerTreeItemType Type;
+
+	/** Construct this item from an Component */
+	FComponentTreeItem(UActorComponent* InComponent);
+
+	/* Begin ISceneOutlinerTreeItem Implementation */
+	virtual bool IsValid() const override { return Component.IsValid(); }
+	virtual FSceneOutlinerTreeItemID GetID() const override;
+	virtual FString GetDisplayString() const override;
+	virtual bool CanInteract() const override;
+	virtual TSharedRef<SWidget> GenerateLabelWidget(ISceneOutliner& Outliner, const STableRow<FSceneOutlinerTreeItemPtr>& InRow) override;
+	/* End ISceneOutlinerTreeItem Implementation */
+public:
+	/** true if this item exists in both the current world and PIE. */
+	bool bExistsInCurrentWorldAndPIE;
+	/** Cache the string displayed */
+	FString CachedDisplayString;
+};
