Removed SceneOutliner namespace

diff --git a/Engine/Source/Editor/SceneOutliner/Public/FolderTreeItem.h b/Engine/Source/Editor/SceneOutliner/Public/FolderTreeItem.h
--- a/Engine/Source/Editor/SceneOutliner/Public/FolderTreeItem.h
+++ b/Engine/Source/Editor/SceneOutliner/Public/FolderTreeItem.h
@@ -4,7 +4,7 @@
 
 #include "CoreMinimal.h"
 #include "Widgets/SWidget.h"
-#include "ITreeItem.h"
+#include "ISceneOutlinerTreeItem.h"
 
 class UToolMenu;
 
@@ -12,57 +12,57 @@ namespace SceneOutliner
 {
 	struct FFolderPathSelector
 	{
-		bool operator()(TWeakPtr<ITreeItem> Item, FName& DataOut) const;
+		bool operator()(TWeakPtr<ISceneOutlinerTreeItem> Item, FName& DataOut) const;
 	};
+}	// namespace SceneOutliner
 
-	/** A tree item that represents a folder in the world */
-	struct SCENEOUTLINER_API FFolderTreeItem : ITreeItem
-	{
-	public:
-		/** Static type identifier for this tree item class */
-		static const FTreeItemType Type;
 
-		DECLARE_DELEGATE_RetVal_OneParam(bool, FFilterPredicate, FName);
-		DECLARE_DELEGATE_RetVal_OneParam(bool, FInteractivePredicate, FName);
+/** A tree item that represents a folder in the world */
+struct SCENEOUTLINER_API FFolderTreeItem : ISceneOutlinerTreeItem
+{
+public:
+	/** Static type identifier for this tree item class */
+	static const FSceneOutlinerTreeItemType Type;
 
-		bool Filter(FFilterPredicate Pred) const
-		{
-			return Pred.Execute(Path);
-		}
+	DECLARE_DELEGATE_RetVal_OneParam(bool, FFilterPredicate, FName);
+	DECLARE_DELEGATE_RetVal_OneParam(bool, FInteractivePredicate, FName);
 
-		bool GetInteractiveState(FInteractivePredicate Pred) const
-		{
-			return Pred.Execute(Path);
-		}
+	bool Filter(FFilterPredicate Pred) const
+	{
+		return Pred.Execute(Path);
+	}
 
-		/** The path of this folder. / separated. */
-		FName Path;
+	bool GetInteractiveState(FInteractivePredicate Pred) const
+	{
+		return Pred.Execute(Path);
+	}
 
-		/** The leaf name of this folder */
-		FName LeafName;
+	/** The path of this folder. / separated. */
+	FName Path;
 
-		/** Constructor that takes a path to this folder (including leaf-name) */
-		FFolderTreeItem(FName InPath);
-		/** Constructor that takes a path to this folder and a subclass tree item type (used for subclassing FFolderTreeItem) */
-		FFolderTreeItem(FName InPath, FTreeItemType Type);
+	/** The leaf name of this folder */
+	FName LeafName;
 
-		/* Begin ITreeItem Implementation */
-		virtual bool IsValid() const override { return true; }
-		virtual FTreeItemID GetID() const override;
-		virtual FString GetDisplayString() const override;
-		virtual bool CanInteract() const override;
-		virtual void GenerateContextMenu(UToolMenu* Menu, SSceneOutliner& Outliner) override;
-		/** Delete this folder, children will be reparented to provided new parent path */
-		virtual void Delete(FName InNewParentPath) {}
-		/* End ITreeItem Implementation */
+	/** Constructor that takes a path to this folder (including leaf-name) */
+	FFolderTreeItem(FName InPath);
+	/** Constructor that takes a path to this folder and a subclass tree item type (used for subclassing FFolderTreeItem) */
+	FFolderTreeItem(FName InPath, FSceneOutlinerTreeItemType Type);
 
-		/** Move this folder to a new parent */
-		virtual FName MoveTo(const FName& NewParent) { return FName(); }
-	private:
-		/** Create a new folder as a child of this one */
-		virtual void CreateSubFolder(TWeakPtr<SSceneOutliner> WeakOutliner) {}
-		/** Duplicate folder hierarchy */
-		void DuplicateHierarchy(TWeakPtr<SSceneOutliner> WeakOutliner);
-	};
+	/* Begin ISceneOutlinerTreeItem Implementation */
+	virtual bool IsValid() const override { return true; }
+	virtual FSceneOutlinerTreeItemID GetID() const override;
+	virtual FString GetDisplayString() const override;
+	virtual bool CanInteract() const override;
+	virtual void GenerateContextMenu(UToolMenu* Menu, SSceneOutliner& Outliner) override;
+	/** Delete this folder, children will be reparented to provided new parent path */
+	virtual void Delete(FName InNewParentPath) {}
+	/* End ISceneOutlinerTreeItem Implementation */
 
-}		// namespace SceneOutliner
+	/** Move this folder to a new parent */
+	virtual FName MoveTo(const FName& NewParent) { return FName(); }
+private:
+	/** Create a new folder as a child of this one */
+	virtual void CreateSubFolder(TWeakPtr<SSceneOutliner> WeakOutliner) {}
+	/** Duplicate folder hierarchy */
+	void DuplicateHierarchy(TWeakPtr<SSceneOutliner> WeakOutliner);
+};
diff --git a/Engine/Source/Editor/SceneOutliner/Public/ISceneOutliner.h b/Engine/Source/Editor/SceneOutliner/Public/ISceneOutliner.h
--- a/Engine/Source/Editor/SceneOutliner/Public/ISceneOutliner.h
+++ b/Engine/Source/Editor/SceneOutliner/Public/ISceneOutliner.h
@@ -10,7 +10,7 @@
 
 // Forward declaration
 template<typename ItemType> class STreeView;
-namespace SceneOutliner { struct FColumnInfo; }
+struct FSceneOutlinerColumnInfo;
 
 /**
  * The public interface for the Scene Outliner widget
@@ -29,10 +29,10 @@ public:
 	virtual TAttribute<FText> GetFilterHighlightText() const = 0;
 
 	/** @return Returns the common data for this outliner */
-	virtual const SceneOutliner::FSharedOutlinerData& GetSharedData() const = 0;
+	virtual const FSharedSceneOutlinerData& GetSharedData() const = 0;
 
 	/** Get a const reference to the actual tree hierarchy */
-	virtual const STreeView<SceneOutliner::FTreeItemPtr>& GetTree() const = 0;
+	virtual const STreeView<FSceneOutlinerTreeItemPtr>& GetTree() const = 0;
 
 	/** Set the keyboard focus to the outliner */
 	virtual void SetKeyboardFocus() = 0;
@@ -44,28 +44,28 @@ public:
 	virtual void CacheIconForClass(FName InClassName, const FSlateBrush* InSlateBrush) = 0;
 
 	/** Should the scene outliner accept a request to rename a item of the tree */
-	virtual bool CanExecuteRenameRequest(const SceneOutliner::ITreeItem& ItemPtr) const = 0;
+	virtual bool CanExecuteRenameRequest(const ISceneOutlinerTreeItem& ItemPtr) const = 0;
 
 	/** 
 	 * Add a filter to the scene outliner 
 	 * @param Filter The filter to apply to the scene outliner
 	 * @return The index of the filter.
 	 */
-	virtual int32 AddFilter(const TSharedRef<SceneOutliner::FOutlinerFilter>& Filter) = 0;
+	virtual int32 AddFilter(const TSharedRef<FSceneOutlinerFilter>& Filter) = 0;
 
 	/** 
 	 * Remove a filter from the scene outliner
 	 * @param Filter The Filter to remove
 	 * @return True if the filter was removed.
 	 */
-	virtual bool RemoveFilter(const TSharedRef<SceneOutliner::FOutlinerFilter>& Filter) = 0;
+	virtual bool RemoveFilter(const TSharedRef<FSceneOutlinerFilter>& Filter) = 0;
 
 	/** 
 	 * Retrieve the filter at the specified index
 	 * @param Index The index of the filter to retrive
 	 * @return A valid poiter to a filter if the index was valid
 	 */
-	virtual TSharedPtr<SceneOutliner::FOutlinerFilter> GetFilterAtIndex(int32 Index) = 0;
+	virtual TSharedPtr<FSceneOutlinerFilter> GetFilterAtIndex(int32 Index) = 0;
 
 	/** Get number of filters applied to the scene outliner */
 	virtual int32 GetFilterCount() const = 0;
@@ -76,7 +76,7 @@ public:
 	 * @param ColumnId The id of the column to add
 	 * @param ColumInfo The struct that contains the information on how to present and retrieve the column
 	 */
-	virtual void AddColumn(FName ColumnId, const SceneOutliner::FColumnInfo& ColumnInfo) = 0;
+	virtual void AddColumn(FName ColumnId, const FSceneOutlinerColumnInfo& ColumnInfo) = 0;
 
 	/**
 	 * Remove a column of the scene outliner
@@ -90,7 +90,7 @@ public:
 	/** Return the sorting mode for the specified ColumnId */
 	virtual EColumnSortMode::Type GetColumnSortMode( const FName ColumnId ) const = 0;
 
-	virtual uint32 GetTypeSortPriority(const SceneOutliner::ITreeItem& Item) const = 0;
+	virtual uint32 GetTypeSortPriority(const ISceneOutlinerTreeItem& Item) const = 0;
 
 	/** Request that the tree be sorted at a convenient time */
 	virtual void RequestSort() = 0;
@@ -117,10 +117,10 @@ public:
 	virtual bool CanSupportDragAndDrop() const = 0;
 
 	/** Set the item selection of the outliner based on a selector function. Any items which return true will be added */
-	virtual void SetSelection(const TFunctionRef<bool(SceneOutliner::ITreeItem&)> Selector) = 0;
+	virtual void SetSelection(const TFunctionRef<bool(ISceneOutlinerTreeItem&)> Selector) = 0;
 
 	/** Get the active SceneOutlinerMode */
-	const SceneOutliner::ISceneOutlinerMode* GetMode() const { return Mode; }
+	const ISceneOutlinerMode* GetMode() const { return Mode; }
 protected:
-	SceneOutliner::ISceneOutlinerMode* Mode;
+	ISceneOutlinerMode* Mode;
 };
diff --git a/Engine/Source/Editor/SceneOutliner/Public/ISceneOutlinerColumn.h b/Engine/Source/Editor/SceneOutliner/Public/ISceneOutlinerColumn.h
--- a/Engine/Source/Editor/SceneOutliner/Public/ISceneOutlinerColumn.h
+++ b/Engine/Source/Editor/SceneOutliner/Public/ISceneOutlinerColumn.h
@@ -21,16 +21,16 @@ public:
 
 	virtual SHeaderRow::FColumn::FArguments ConstructHeaderRowColumn() = 0;
 
-	virtual const TSharedRef< SWidget > ConstructRowWidget( SceneOutliner::FTreeItemRef TreeItem, const STableRow<SceneOutliner::FTreeItemPtr>& Row ) = 0;
+	virtual const TSharedRef< SWidget > ConstructRowWidget(FSceneOutlinerTreeItemRef TreeItem, const STableRow<FSceneOutlinerTreeItemPtr>& Row) = 0;
 
 	virtual void Tick(double InCurrentTime, float InDeltaTime) {}
 
 public:
 	/** Optionally overridden interface methods */
-	virtual void PopulateSearchStrings( const SceneOutliner::ITreeItem& Item, TArray< FString >& OutSearchStrings ) const {}
+	virtual void PopulateSearchStrings( const ISceneOutlinerTreeItem& Item, TArray< FString >& OutSearchStrings ) const {}
 
 	virtual bool SupportsSorting() const { return false; }
 
-	virtual void SortItems(TArray<SceneOutliner::FTreeItemPtr>& OutItems, const EColumnSortMode::Type SortMode) const {}
+	virtual void SortItems(TArray<FSceneOutlinerTreeItemPtr>& OutItems, const EColumnSortMode::Type SortMode) const {}
 };
 
diff --git a/Engine/Source/Editor/SceneOutliner/Public/ISceneOutlinerHierarchy.h b/Engine/Source/Editor/SceneOutliner/Public/ISceneOutlinerHierarchy.h
--- a/Engine/Source/Editor/SceneOutliner/Public/ISceneOutlinerHierarchy.h
+++ b/Engine/Source/Editor/SceneOutliner/Public/ISceneOutlinerHierarchy.h
@@ -6,33 +6,30 @@
 
 #include "SceneOutlinerStandaloneTypes.h"
 #include "SceneOutlinerFwd.h"
-#include "ITreeItem.h"
+#include "ISceneOutlinerTreeItem.h"
 
-namespace SceneOutliner
+class ISceneOutlinerHierarchy : public TSharedFromThis<ISceneOutlinerHierarchy>
 {
-	class ISceneOutlinerHierarchy : public TSharedFromThis<ISceneOutlinerHierarchy>
-	{
-	public:
-		ISceneOutlinerHierarchy(ISceneOutlinerMode* InMode) : Mode(InMode) {}
-		virtual ~ISceneOutlinerHierarchy() {}
+public:
+	ISceneOutlinerHierarchy(ISceneOutlinerMode* InMode) : Mode(InMode) {}
+	virtual ~ISceneOutlinerHierarchy() {}
 
-		/** Find all direct children of a tree item in an existing item map, if any exist. */
-		virtual void FindChildren(const ITreeItem& Item, const TMap<FTreeItemID, FTreeItemPtr>& Items, TArray<FTreeItemPtr>& OutChildItems) const = 0;
-		/** Find the parent of a tree item in an existing item map, it if exists. */
-		virtual FTreeItemPtr FindParent(const ITreeItem& Item, const TMap<FTreeItemID, FTreeItemPtr>& Items) const = 0;
+	/** Find all direct children of a tree item in an existing item map, if any exist. */
+	virtual void FindChildren(const ISceneOutlinerTreeItem& Item, const TMap<FSceneOutlinerTreeItemID, FSceneOutlinerTreeItemPtr>& Items, TArray<FSceneOutlinerTreeItemPtr>& OutChildItems) const = 0;
+	/** Find the parent of a tree item in an existing item map, it if exists. */
+	virtual FSceneOutlinerTreeItemPtr FindParent(const ISceneOutlinerTreeItem& Item, const TMap<FSceneOutlinerTreeItemID, FSceneOutlinerTreeItemPtr>& Items) const = 0;
 		
-		/** Create a linearization of all applicable items in the hierarchy */
-		virtual void CreateItems(TArray<FTreeItemPtr>& OutItems) const = 0;
-		/** Create a linearization of all direct and indirect children of a given item in the hierarchy */
-		virtual void CreateChildren(const FTreeItemPtr& Item, TArray<FTreeItemPtr>& OutChildren) const = 0;
-		/** Forcibly create a parent item for a given tree item */
-		virtual FTreeItemPtr CreateParentItem(const FTreeItemPtr& Item) const = 0;
+	/** Create a linearization of all applicable items in the hierarchy */
+	virtual void CreateItems(TArray<FSceneOutlinerTreeItemPtr>& OutItems) const = 0;
+	/** Create a linearization of all direct and indirect children of a given item in the hierarchy */
+	virtual void CreateChildren(const FSceneOutlinerTreeItemPtr& Item, TArray<FSceneOutlinerTreeItemPtr>& OutChildren) const = 0;
+	/** Forcibly create a parent item for a given tree item */
+	virtual FSceneOutlinerTreeItemPtr CreateParentItem(const FSceneOutlinerTreeItemPtr& Item) const = 0;
 
-		DECLARE_EVENT_OneParam(ISceneOutlinerHierarchy, FHierarchyChangedEvent, FHierarchyChangedData)
-		FHierarchyChangedEvent& OnHierarchyChanged() { return HierarchyChangedEvent; }
+	DECLARE_EVENT_OneParam(ISceneOutlinerHierarchy, FHierarchyChangedEvent, FSceneOutlinerHierarchyChangedData)
+	FHierarchyChangedEvent& OnHierarchyChanged() { return HierarchyChangedEvent; }
 
-	protected:
-		ISceneOutlinerMode* Mode;
-		FHierarchyChangedEvent HierarchyChangedEvent;
-	};
-}
\ No newline at end of file
+protected:
+	ISceneOutlinerMode* Mode;
+	FHierarchyChangedEvent HierarchyChangedEvent;
+};
diff --git a/Engine/Source/Editor/SceneOutliner/Public/ISceneOutlinerMode.h b/Engine/Source/Editor/SceneOutliner/Public/ISceneOutlinerMode.h
--- a/Engine/Source/Editor/SceneOutliner/Public/ISceneOutlinerMode.h
+++ b/Engine/Source/Editor/SceneOutliner/Public/ISceneOutlinerMode.h
@@ -5,152 +5,150 @@
 #include "CoreMinimal.h"
 
 #include "SceneOutlinerFwd.h"
-#include "ITreeItem.h"
+#include "ISceneOutlinerTreeItem.h"
 #include "SceneOutlinerFilters.h"
 #include "SSceneOutliner.h"
 #include "SceneOutlinerDragDrop.h"
 
-namespace SceneOutliner
+
+class ISceneOutlinerMode
 {
-	class ISceneOutlinerMode
+public:
+	ISceneOutlinerMode(SSceneOutliner* InSceneOutliner) : SceneOutliner(InSceneOutliner) {}
+
+	virtual ~ISceneOutlinerMode() {}
+
+	/** Rebuild all mode data */
+	virtual void Rebuild() = 0;
+
+	/** 
+		* Attempt to create a tree item for a given data type.
+		* @param bForce should the item be created regardless if it passes the filters or not
+		* @return returns a shared pointer to the new item if it was created
+				may be invalid if the item was not created
+		*/
+	template <typename TreeItemType, typename TreeItemData>
+	FSceneOutlinerTreeItemPtr CreateItemFor(const TreeItemData& Data, bool bForce = false)
 	{
-	public:
-		ISceneOutlinerMode(SSceneOutliner* InSceneOutliner) : SceneOutliner(InSceneOutliner) {}
-
-		virtual ~ISceneOutlinerMode() {}
-
-		/** Rebuild all mode data */
-		virtual void Rebuild() = 0;
-
-		/** 
-		 * Attempt to create a tree item for a given data type.
-		 * @param bForce should the item be created regardless if it passes the filters or not
-		 * @return returns a shared pointer to the new item if it was created
-				   may be invalid if the item was not created
-		 */
-		template <typename TreeItemType, typename TreeItemData>
-		FTreeItemPtr CreateItemFor(const TreeItemData& Data, bool bForce = false)
-		{
-			return SceneOutliner->CreateItemFor<TreeItemType>(Data, 
-				[this](const ITreeItem& Item) { OnItemPassesFilters(Item); }, bForce);
-		}
-
-		/** Fill the view button menu with content */
-		virtual void CreateViewContent(FMenuBuilder& MenuBuilder) {}
-
-		/** Returns a factory to create a Folder Picker mode which matches this mode */
-		virtual FCreateOutlinerMode CreateFolderPickerMode() const { return FCreateOutlinerMode(); }
-	public:
-		/** Construct a new Drag and drop operation */
-		virtual TSharedPtr<FDragDropOperation> CreateDragDropOperation(const TArray<FTreeItemPtr>& InTreeItems) const { return nullptr; }
-
-		/**
-		 * Populates a drag/drop operation if the mode supports that type of operation.
-		 * @param OutPayload:	The payload data will be stored here after parsing
-		 * @return				Returns true if the operation is supported by the mode
-		 */
-		virtual bool ParseDragDrop(FDragDropPayload& OutPayload, const FDragDropOperation& Operation) const { return false; /* Default does not support drag/drop */ }
-
-		/** Test whether the specified payload can be dropped onto this tree item */
-		virtual FDragValidationInfo ValidateDrop(const ITreeItem& DropTarget, const FDragDropPayload& Payload) const { return FDragValidationInfo::Invalid(); }
-
-		/** Called when a payload is dropped onto a target */
-		virtual void OnDrop(ITreeItem& DropTarget, const FDragDropPayload& Payload, const FDragValidationInfo& ValidationInfo) const {}
-
-		/** Called when a payload is dragged over an item */
-		virtual FReply OnDragOverItem(const FDragDropEvent& Event, const ITreeItem& Item) const { return FReply::Handled(); }
-	public:
-		/* Events */
-
-		/** Called by the outliner when a new item is added to the outliner tree */
-		virtual void OnItemAdded(FTreeItemPtr Item) {}
-		/** Called by the outliner when an item is removed from the outliner tree */
-		virtual void OnItemRemoved(FTreeItemPtr Item) {}
-		/** Called by the outliner when an item is selected in the tree. Used to partially synchronize selection */
-		virtual void OnItemSelectionChanged(FTreeItemPtr Item, ESelectInfo::Type SelectionType, const FItemSelection& Selection) {}
-		/** Called by the outliner when a tree item is double clicked */
-		virtual void OnItemDoubleClick(FTreeItemPtr Item) {}
-		/** Called by the outliner when search box text is changed */
-		virtual void OnFilterTextChanged(const FText& InFilterText) {}
-		/** Called by the outliner when search box text is committed (by pressing enter) */
-		virtual void OnFilterTextCommited(FItemSelection& Selection, ETextCommit::Type CommitType) {}
-		/** Called by the outliner when an item passes filters but isn't yet added to the tree */
-		virtual void OnItemPassesFilters(const ITreeItem& Item) {}
-		/** Called by the outliner when a key is pressed */
-		virtual FReply OnKeyDown(const FKeyEvent& InKeyEvent) { return FReply::Unhandled(); }
-
-		/** Returns a context menu widget based on the current selection */
-		virtual TSharedPtr<SWidget> CreateContextMenu() { return TSharedPtr<SWidget>(); }
-		/** Check if an item can be renamed */
-		virtual bool CanRenameItem(const ITreeItem& Item) const { return false; }
-
-		/** Synchronize the mode specific selection with the tree view */
-		virtual void SynchronizeSelection() {}
-		/** Trigger a duplication of selected items */
-		virtual void OnDuplicateSelected() {}
-	public:
-		/** Get the current selection mode */
-		virtual ESelectionMode::Type GetSelectionMode() const { return ESelectionMode::Single; }
-		/** Get the sort priority of an item */
-		virtual int32 GetTypeSortPriority(const ITreeItem& Item) const { return 0; }
-		/** Does the mode support keyboard focus */
-		virtual bool SupportsKeyboardFocus() const { return false; }
-		/** Get the current status text */
-		virtual FText GetStatusText() const { return FText(); }
-		/** Get the color of the status text */
-		virtual FSlateColor GetStatusTextColor() const { return FSlateColor(); }
-		/** Should folders be shown in the outliner */
-		virtual bool ShouldShowFolders() const { return false; }
-		/** Should the new folder button be displayed by the outliner */
-		virtual bool SupportsCreateNewFolder() const { return false; }
-		/** Is the mode interactive */
-		virtual bool IsInteractive() const { return false; }
-		/** Does the mode want to display a status bar */
-		virtual bool ShowStatusBar() const { return false; }
-		/** Does the mode want to display a view button */
-		virtual bool ShowViewButton() const { return false; }
-		/** Does the mode allow the user to enable and disable filters through a menu */
-		virtual bool ShowFilterOptions() const { return false; }
-
-		/** Can the current selection be deleted */
-		virtual bool CanDelete() const { return false; }
-		/** Can the current selection be renamed */
-		virtual bool CanRename() const { return false; }
-		/** Can the current selection be cut */
-		virtual bool CanCut() const { return false; }
-		/** Can the current selection be copied */
-		virtual bool CanCopy() const { return false; }
-		/** Can the current selection be pasted */
-		virtual bool CanPaste() const { return false; }
-
-		/** Does the mode support drag and drop */
-		virtual bool CanSupportDragAndDrop() const { return false; }
-	public:
-		/* Folder management */
+		return SceneOutliner->CreateItemFor<TreeItemType>(Data,
+			[this](const ISceneOutlinerTreeItem& Item) { OnItemPassesFilters(Item); }, bForce);
+	}
+
+	/** Fill the view button menu with content */
+	virtual void CreateViewContent(FMenuBuilder& MenuBuilder) {}
+
+	/** Returns a factory to create a Folder Picker mode which matches this mode */
+	virtual FCreateSceneOutlinerMode CreateFolderPickerMode() const { return FCreateSceneOutlinerMode(); }
+public:
+	/** Construct a new Drag and drop operation */
+	virtual TSharedPtr<FDragDropOperation> CreateDragDropOperation(const TArray<FSceneOutlinerTreeItemPtr>& InTreeItems) const { return nullptr; }
+
+	/**
+		* Populates a drag/drop operation if the mode supports that type of operation.
+		* @param OutPayload:	The payload data will be stored here after parsing
+		* @return				Returns true if the operation is supported by the mode
+		*/
+	virtual bool ParseDragDrop(FSceneOutlinerDragDropPayload& OutPayload, const FDragDropOperation& Operation) const { return false; /* Default does not support drag/drop */ }
+
+	/** Test whether the specified payload can be dropped onto this tree item */
+	virtual FSceneOutlinerDragValidationInfo ValidateDrop(const ISceneOutlinerTreeItem& DropTarget, const FSceneOutlinerDragDropPayload& Payload) const { return FSceneOutlinerDragValidationInfo::Invalid(); }
+
+	/** Called when a payload is dropped onto a target */
+	virtual void OnDrop(ISceneOutlinerTreeItem& DropTarget, const FSceneOutlinerDragDropPayload& Payload, const FSceneOutlinerDragValidationInfo& ValidationInfo) const {}
+
+	/** Called when a payload is dragged over an item */
+	virtual FReply OnDragOverItem(const FDragDropEvent& Event, const ISceneOutlinerTreeItem& Item) const { return FReply::Handled(); }
+public:
+	/* Events */
+
+	/** Called by the outliner when a new item is added to the outliner tree */
+	virtual void OnItemAdded(FSceneOutlinerTreeItemPtr Item) {}
+	/** Called by the outliner when an item is removed from the outliner tree */
+	virtual void OnItemRemoved(FSceneOutlinerTreeItemPtr Item) {}
+	/** Called by the outliner when an item is selected in the tree. Used to partially synchronize selection */
+	virtual void OnItemSelectionChanged(FSceneOutlinerTreeItemPtr Item, ESelectInfo::Type SelectionType, const FSceneOutlinerItemSelection& Selection) {}
+	/** Called by the outliner when a tree item is double clicked */
+	virtual void OnItemDoubleClick(FSceneOutlinerTreeItemPtr Item) {}
+	/** Called by the outliner when search box text is changed */
+	virtual void OnFilterTextChanged(const FText& InFilterText) {}
+	/** Called by the outliner when search box text is committed (by pressing enter) */
+	virtual void OnFilterTextCommited(FSceneOutlinerItemSelection& Selection, ETextCommit::Type CommitType) {}
+	/** Called by the outliner when an item passes filters but isn't yet added to the tree */
+	virtual void OnItemPassesFilters(const ISceneOutlinerTreeItem& Item) {}
+	/** Called by the outliner when a key is pressed */
+	virtual FReply OnKeyDown(const FKeyEvent& InKeyEvent) { return FReply::Unhandled(); }
+
+	/** Returns a context menu widget based on the current selection */
+	virtual TSharedPtr<SWidget> CreateContextMenu() { return TSharedPtr<SWidget>(); }
+	/** Check if an item can be renamed */
+	virtual bool CanRenameItem(const ISceneOutlinerTreeItem& Item) const { return false; }
+
+	/** Synchronize the mode specific selection with the tree view */
+	virtual void SynchronizeSelection() {}
+	/** Trigger a duplication of selected items */
+	virtual void OnDuplicateSelected() {}
+public:
+	/** Get the current selection mode */
+	virtual ESelectionMode::Type GetSelectionMode() const { return ESelectionMode::Single; }
+	/** Get the sort priority of an item */
+	virtual int32 GetTypeSortPriority(const ISceneOutlinerTreeItem& Item) const { return 0; }
+	/** Does the mode support keyboard focus */
+	virtual bool SupportsKeyboardFocus() const { return false; }
+	/** Get the current status text */
+	virtual FText GetStatusText() const { return FText(); }
+	/** Get the color of the status text */
+	virtual FSlateColor GetStatusTextColor() const { return FSlateColor(); }
+	/** Should folders be shown in the outliner */
+	virtual bool ShouldShowFolders() const { return false; }
+	/** Should the new folder button be displayed by the outliner */
+	virtual bool SupportsCreateNewFolder() const { return false; }
+	/** Is the mode interactive */
+	virtual bool IsInteractive() const { return false; }
+	/** Does the mode want to display a status bar */
+	virtual bool ShowStatusBar() const { return false; }
+	/** Does the mode want to display a view button */
+	virtual bool ShowViewButton() const { return false; }
+	/** Does the mode allow the user to enable and disable filters through a menu */
+	virtual bool ShowFilterOptions() const { return false; }
+
+	/** Can the current selection be deleted */
+	virtual bool CanDelete() const { return false; }
+	/** Can the current selection be renamed */
+	virtual bool CanRename() const { return false; }
+	/** Can the current selection be cut */
+	virtual bool CanCut() const { return false; }
+	/** Can the current selection be copied */
+	virtual bool CanCopy() const { return false; }
+	/** Can the current selection be pasted */
+	virtual bool CanPaste() const { return false; }
+
+	/** Does the mode support drag and drop */
+	virtual bool CanSupportDragAndDrop() const { return false; }
+public:
+	/* Folder management */
 		
-		/** Creates a new folder item at the root with a valid name*/
-		virtual FName CreateNewFolder() { return FName(NAME_None); }
-		/** Create a folder under a specific parent with a given leaf name */
-		virtual FName CreateFolder(const FName& ParentPath, const FName& LeafName) { return FName(NAME_None); }
-		/** Reparent an item to a given folder path. Returns true if the operation is sucessful */
-		virtual bool ReparentItemToFolder(const FName& FolderPath, const FTreeItemPtr& Item) { return false; }
-		/** Select all descendants of a folder. Optionally select only immediate descendants. */
-		virtual void SelectFoldersDescendants(const TArray<FFolderTreeItem*>& FolderItems, bool bSelectImmediateChildrenOnly) {}
-	public:
-		/* Getters */
-
-		ISceneOutlinerHierarchy* GetHierarchy() { return Hierarchy.Get(); }
-		TMap<FName, FOutlinerFilterInfo>& GetFilterInfos() { return FilterInfoMap; }
-	protected:
-		virtual TUniquePtr<ISceneOutlinerHierarchy> CreateHierarchy() = 0;
-	protected:
-		TMap<FName, FOutlinerFilterInfo> FilterInfoMap;
-		TUniquePtr<ISceneOutlinerHierarchy> Hierarchy;
-
-		SSceneOutliner* SceneOutliner;
-	private:
-		ISceneOutlinerMode() = delete;
-		ISceneOutlinerMode(const ISceneOutliner&) = delete;
-		ISceneOutlinerMode& operator=(const ISceneOutliner&) = delete;
-	};
-}
\ No newline at end of file
+	/** Creates a new folder item at the root with a valid name*/
+	virtual FName CreateNewFolder() { return FName(NAME_None); }
+	/** Create a folder under a specific parent with a given leaf name */
+	virtual FName CreateFolder(const FName& ParentPath, const FName& LeafName) { return FName(NAME_None); }
+	/** Reparent an item to a given folder path. Returns true if the operation is sucessful */
+	virtual bool ReparentItemToFolder(const FName& FolderPath, const FSceneOutlinerTreeItemPtr& Item) { return false; }
+	/** Select all descendants of a folder. Optionally select only immediate descendants. */
+	virtual void SelectFoldersDescendants(const TArray<FFolderTreeItem*>& FolderItems, bool bSelectImmediateChildrenOnly) {}
+public:
+	/* Getters */
+
+	ISceneOutlinerHierarchy* GetHierarchy() { return Hierarchy.Get(); }
+	TMap<FName, FSceneOutlinerFilterInfo>& GetFilterInfos() { return FilterInfoMap; }
+protected:
+	virtual TUniquePtr<ISceneOutlinerHierarchy> CreateHierarchy() = 0;
+protected:
+	TMap<FName, FSceneOutlinerFilterInfo> FilterInfoMap;
+	TUniquePtr<ISceneOutlinerHierarchy> Hierarchy;
+
+	SSceneOutliner* SceneOutliner;
+private:
+	ISceneOutlinerMode() = delete;
+	ISceneOutlinerMode(const ISceneOutliner&) = delete;
+	ISceneOutlinerMode& operator=(const ISceneOutliner&) = delete;
+};
diff --git a/Engine/Source/Editor/SceneOutliner/Public/SOutlinerTreeView.h b/Engine/Source/Editor/SceneOutliner/Public/SOutlinerTreeView.h
--- a/Engine/Source/Editor/SceneOutliner/Public/SOutlinerTreeView.h
+++ b/Engine/Source/Editor/SceneOutliner/Public/SOutlinerTreeView.h
@@ -10,109 +10,106 @@
 #include "Widgets/Views/STreeView.h"
 #include "Animation/SlateSprings.h"
 
-namespace SceneOutliner
+class SSceneOutlinerTreeView : public STreeView< FSceneOutlinerTreeItemPtr >
 {
-
-	class SOutlinerTreeView : public STreeView< FTreeItemPtr >
-	{
-	public:
+public:
 		
-		/** Construct this widget */
-		void Construct(const FArguments& InArgs, TSharedRef<SSceneOutliner> Owner);
+	/** Construct this widget */
+	void Construct(const FArguments& InArgs, TSharedRef<SSceneOutliner> Owner);
+
+	void FlashHighlightOnItem( FSceneOutlinerTreeItemPtr FlashHighlightOnItem );
 
-		void FlashHighlightOnItem( FTreeItemPtr FlashHighlightOnItem );
+	const TWeakPtr<SSceneOutliner>& GetOutlinerPtr() { return SceneOutlinerWeak; }
+protected:
 
-		const TWeakPtr<SSceneOutliner>& GetOutlinerPtr() { return SceneOutlinerWeak; }
-	protected:
+	virtual FReply OnDragOver( const FGeometry& MyGeometry, const FDragDropEvent& DragDropEvent ) override;
 
-		virtual FReply OnDragOver( const FGeometry& MyGeometry, const FDragDropEvent& DragDropEvent ) override;
+	virtual void OnDragLeave( const FDragDropEvent& DragDropEvent ) override;
 
-		virtual void OnDragLeave( const FDragDropEvent& DragDropEvent ) override;
+	virtual FReply OnDrop( const FGeometry& MyGeometry, const FDragDropEvent& DragDropEvent ) override;
 
-		virtual FReply OnDrop( const FGeometry& MyGeometry, const FDragDropEvent& DragDropEvent ) override;
+	/** Weak reference to the outliner widget that owns this list */
+	TWeakPtr<SSceneOutliner> SceneOutlinerWeak;
+};
 
-		/** Weak reference to the outliner widget that owns this list */
-		TWeakPtr<SSceneOutliner> SceneOutlinerWeak;
-	};
+/** Widget that represents a row in the outliner's tree control.  Generates widgets for each column on demand. */
+class SSceneOutlinerTreeRow
+	: public SMultiColumnTableRow< FSceneOutlinerTreeItemPtr >
+{
 
-	/** Widget that represents a row in the outliner's tree control.  Generates widgets for each column on demand. */
-	class SSceneOutlinerTreeRow
-		: public SMultiColumnTableRow< FTreeItemPtr >
-	{
+public:
 
-	public:
+	SLATE_BEGIN_ARGS( SSceneOutlinerTreeRow ) {}
 
-		SLATE_BEGIN_ARGS( SSceneOutlinerTreeRow ) {}
+		/** The list item for this row */
+		SLATE_ARGUMENT( FSceneOutlinerTreeItemPtr, Item )
 
-			/** The list item for this row */
-			SLATE_ARGUMENT( FTreeItemPtr, Item )
+	SLATE_END_ARGS()
 
-		SLATE_END_ARGS()
 
+	/** Construct function for this widget */
+	void Construct( const FArguments& InArgs, const TSharedRef<SSceneOutlinerTreeView>& OutlinerTreeView, TSharedRef<SSceneOutliner> SceneOutliner );
 
-		/** Construct function for this widget */
-		void Construct( const FArguments& InArgs, const TSharedRef<SOutlinerTreeView>& OutlinerTreeView, TSharedRef<SSceneOutliner> SceneOutliner );
+	/** Overridden from SMultiColumnTableRow.  Generates a widget for this column of the tree row. */
+	virtual TSharedRef<SWidget> GenerateWidgetForColumn( const FName& ColumnName ) override;
 
-		/** Overridden from SMultiColumnTableRow.  Generates a widget for this column of the tree row. */
-		virtual TSharedRef<SWidget> GenerateWidgetForColumn( const FName& ColumnName ) override;
+	void FlashHighlight();
 
-		void FlashHighlight();
+protected:
 
-	protected:
+	virtual FReply OnMouseButtonDown(const FGeometry& MyGeometry, const FPointerEvent& MouseEvent) override;
+	virtual FReply OnMouseButtonUp(const FGeometry& MyGeometry, const FPointerEvent& MouseEvent) override;
 
-		virtual FReply OnMouseButtonDown(const FGeometry& MyGeometry, const FPointerEvent& MouseEvent) override;
-		virtual FReply OnMouseButtonUp(const FGeometry& MyGeometry, const FPointerEvent& MouseEvent) override;
+	virtual void OnDragEnter( const FGeometry& MyGeometry, const FDragDropEvent& DragDropEvent ) override;
 
-		virtual void OnDragEnter( const FGeometry& MyGeometry, const FDragDropEvent& DragDropEvent ) override;
+	virtual void OnDragLeave( const FDragDropEvent& DragDropEvent ) override;
 
-		virtual void OnDragLeave( const FDragDropEvent& DragDropEvent ) override;
+	virtual FReply OnDragOver( const FGeometry& MyGeometry, const FDragDropEvent& DragDropEvent ) override;
 
-		virtual FReply OnDragOver( const FGeometry& MyGeometry, const FDragDropEvent& DragDropEvent ) override;
+	virtual FReply OnDrop( const FGeometry& MyGeometry, const FDragDropEvent& DragDropEvent ) override;
 
-		virtual FReply OnDrop( const FGeometry& MyGeometry, const FDragDropEvent& DragDropEvent ) override;
+	bool IsRowEnabled() const;
 
-		bool IsRowEnabled() const;
+private:
 
-	private:
+	/** Weak reference to the outliner widget that owns our list */
+	TWeakPtr< SSceneOutliner > SceneOutlinerWeak;
 
-		/** Weak reference to the outliner widget that owns our list */
-		TWeakPtr< SSceneOutliner > SceneOutlinerWeak;
+	/** The item associated with this row of data */
+	TWeakPtr<ISceneOutlinerTreeItem> Item;
 
-		/** The item associated with this row of data */
-		TWeakPtr<ITreeItem> Item;
+protected:
 
-	protected:
+	virtual void Tick( const FGeometry& AllottedGeometry, const double InCurrentTime, const float InDeltaTime ) override;
 
-		virtual void Tick( const FGeometry& AllottedGeometry, const double InCurrentTime, const float InDeltaTime ) override;
+	virtual int32 OnPaint( const FPaintArgs& Args, const FGeometry& AllottedGeometry, const FSlateRect& MyCullingRect, FSlateWindowElementList& OutDrawElements, int32 LayerId, const FWidgetStyle& InWidgetStyle, bool bParentEnabled ) const override;
 
-		virtual int32 OnPaint( const FPaintArgs& Args, const FGeometry& AllottedGeometry, const FSlateRect& MyCullingRect, FSlateWindowElementList& OutDrawElements, int32 LayerId, const FWidgetStyle& InWidgetStyle, bool bParentEnabled ) const override;
+private:
+	/** How many pixels to extend the highlight rectangle's left side horizontally */
+	static const float HighlightRectLeftOffset;
 
-	private:
-		/** How many pixels to extend the highlight rectangle's left side horizontally */
-		static const float HighlightRectLeftOffset;
+	/** How many pixels to extend the highlight rectangle's right side horizontally */
+	static const float HighlightRectRightOffset;
 
-		/** How many pixels to extend the highlight rectangle's right side horizontally */
-		static const float HighlightRectRightOffset;
+	/** How quickly the highlight 'targeting' rectangle will slide around.  Larger is faster. */
+	static const float HighlightTargetSpringConstant;
 
-		/** How quickly the highlight 'targeting' rectangle will slide around.  Larger is faster. */
-		static const float HighlightTargetSpringConstant;
+	/** Duration of animation highlight target effects */
+	static const float HighlightTargetEffectDuration;
 
-		/** Duration of animation highlight target effects */
-		static const float HighlightTargetEffectDuration;
+	/** Opacity of the highlight target effect overlay */
+	static const float HighlightTargetOpacity;
 
-		/** Opacity of the highlight target effect overlay */
-		static const float HighlightTargetOpacity;
+	/** How large the highlight target effect will be when highlighting, as a scalar percentage of font height */
+	static const float LabelChangedAnimOffsetPercent;
 
-		/** How large the highlight target effect will be when highlighting, as a scalar percentage of font height */
-		static const float LabelChangedAnimOffsetPercent;
+	/** Highlight "targeting" visual effect left position */
+	FFloatSpring1D HighlightTargetLeftSpring;
 
-		/** Highlight "targeting" visual effect left position */
-		FFloatSpring1D HighlightTargetLeftSpring;
+	/** Highlight "targeting" visual effect right position */
+	FFloatSpring1D HighlightTargetRightSpring;
 
-		/** Highlight "targeting" visual effect right position */
-		FFloatSpring1D HighlightTargetRightSpring;
+	/** Last time that the user had a major interaction with the highlight */
+	double LastHighlightInteractionTime;
+};
 
-		/** Last time that the user had a major interaction with the highlight */
-		double LastHighlightInteractionTime;
-	};
-}		// namespace SceneOutliner
diff --git a/Engine/Source/Editor/SceneOutliner/Public/SSceneOutliner.h b/Engine/Source/Editor/SceneOutliner/Public/SSceneOutliner.h
--- a/Engine/Source/Editor/SceneOutliner/Public/SSceneOutliner.h
+++ b/Engine/Source/Editor/SceneOutliner/Public/SSceneOutliner.h
@@ -42,23 +42,23 @@ template<typename ItemType> class STreeView;
  */
 namespace SceneOutliner
 {
-	DECLARE_EVENT_OneParam(SSceneOutliner, FTreeItemPtrEvent, FTreeItemPtr);
+	DECLARE_EVENT_OneParam(SSceneOutliner, FTreeItemPtrEvent, FSceneOutlinerTreeItemPtr);
 
-	DECLARE_EVENT_TwoParams(SSceneOutliner, FOnItemSelectionChanged, FTreeItemPtr, ESelectInfo::Type);
+	DECLARE_EVENT_TwoParams(SSceneOutliner, FOnItemSelectionChanged, FSceneOutlinerTreeItemPtr, ESelectInfo::Type);
 
-	typedef TTextFilter< const ITreeItem& > TreeItemTextFilter;
+	typedef TTextFilter< const ISceneOutlinerTreeItem& > TreeItemTextFilter;
 
 	/** Structure that defines an operation that should be applied to the tree */
 	struct FPendingTreeOperation
 	{
 		enum EType { Added, Removed, Moved };
-		FPendingTreeOperation(EType InType, TSharedRef<ITreeItem> InItem) : Type(InType), Item(InItem) { }
+		FPendingTreeOperation(EType InType, TSharedRef<ISceneOutlinerTreeItem> InItem) : Type(InType), Item(InItem) { }
 
 		/** The type of operation that is to be applied */
 		EType Type;
 
 		/** The tree item to which this operation relates */
-		FTreeItemRef Item;
+		FSceneOutlinerTreeItemRef Item;
 	};
 
 	/** Set of actions to apply to new tree items */
@@ -67,388 +67,386 @@ namespace SceneOutliner
 		enum Type
 		{
 			/** Do nothing when it is created */
-			None			= 0,
+			None = 0,
 			/** Select the item when it is created */
-			Select			= 1 << 0,
+			Select = 1 << 0,
 			/** Scroll the item into view when it is created */
-			ScrollIntoView	= 1 << 1,
+			ScrollIntoView = 1 << 1,
 			/** Interactively rename the item when it is created (implies the above) */
-			Rename			= 1 << 2,
+			Rename = 1 << 2,
 		};
 	}
+}
 
-	/**
+/**
 	 * Stores a set of selected items with parsing functions for the scene outliner
 	 */
-	struct FItemSelection
-	{
-		/** Set of selected items */
-		mutable TArray<TWeakPtr<ITreeItem>> SelectedItems;
+struct FSceneOutlinerItemSelection
+{
+	/** Set of selected items */
+	mutable TArray<TWeakPtr<ISceneOutlinerTreeItem>> SelectedItems;
 
-		FItemSelection() {}
+	FSceneOutlinerItemSelection() {}
 
-		FItemSelection(const TArray<FTreeItemPtr>& InSelectedItems)
-			: SelectedItems(InSelectedItems) {}
+	FSceneOutlinerItemSelection(const TArray<FSceneOutlinerTreeItemPtr>& InSelectedItems)
+		: SelectedItems(InSelectedItems) {}
 
-		FItemSelection(SOutlinerTreeView& Tree)
-			: FItemSelection(Tree.GetSelectedItems()) {}
+	FSceneOutlinerItemSelection(SSceneOutlinerTreeView& Tree)
+		: FSceneOutlinerItemSelection(Tree.GetSelectedItems()) {}
 
-		/** Returns true if the selection has an item of a specified type */
-		template <typename TreeType>
-		bool Has() const
+	/** Returns true if the selection has an item of a specified type */
+	template <typename TreeType>
+	bool Has() const
+	{
+		for (const TWeakPtr<ISceneOutlinerTreeItem>& Item : SelectedItems)
 		{
-			for (const TWeakPtr<ITreeItem>& Item : SelectedItems)
+			if (const auto ItemPtr = Item.Pin())
 			{
-				if (const auto ItemPtr = Item.Pin())
+				if (ItemPtr->IsA<TreeType>())
 				{
-					if (ItemPtr->IsA<TreeType>())
-					{
-						return true;
-					}
+					return true;
 				}
 			}
-			return false;
 		}
+		return false;
+	}
 
-		/** Returns the total number of items in the selection */
-		uint32 Num() const
-		{
-			return SelectedItems.Num();
-		}
+	/** Returns the total number of items in the selection */
+	uint32 Num() const
+	{
+		return SelectedItems.Num();
+	}
 
-		/** Returns the number of items of a specific type in the selection */
-		template <typename TreeType>
-		uint32 Num() const
+	/** Returns the number of items of a specific type in the selection */
+	template <typename TreeType>
+	uint32 Num() const
+	{
+		uint32 Result = 0;
+		for (const TWeakPtr<ISceneOutlinerTreeItem>& Item : SelectedItems)
 		{
-			uint32 Result = 0;
-			for (const TWeakPtr<ITreeItem>& Item : SelectedItems)
+			if (const auto ItemPtr = Item.Pin())
 			{
-				if (const auto ItemPtr = Item.Pin())
+				if (ItemPtr->IsA<TreeType>())
 				{
-					if (ItemPtr->IsA<TreeType>())
-					{
-						++Result;
-					}
+					++Result;
 				}
 			}
-			return Result;
 		}
+		return Result;
+	}
 
-		/** Add a new item to the selection */
-		void Add(FTreeItemPtr NewItem)
-		{
-			SelectedItems.Add(NewItem);
-		}
+	/** Add a new item to the selection */
+	void Add(FSceneOutlinerTreeItemPtr NewItem)
+	{
+		SelectedItems.Add(NewItem);
+	}
 
-		/** Get all items of a specified type */
-		template <typename TreeType>
-		void Get(TArray<TreeType*>& OutArray) const
+	/** Get all items of a specified type */
+	template <typename TreeType>
+	void Get(TArray<TreeType*>& OutArray) const
+	{
+		for (const TWeakPtr<ISceneOutlinerTreeItem>& Item : SelectedItems)
 		{
-			for (const TWeakPtr<ITreeItem>& Item : SelectedItems)
+			if (const auto ItemPtr = Item.Pin())
 			{
-				if (const auto ItemPtr = Item.Pin())
+				if (TreeType* CastedItem = ItemPtr->CastTo<TreeType>())
 				{
-					if (TreeType* CastedItem = ItemPtr->CastTo<TreeType>())
-					{
-						OutArray.Add(CastedItem);
-					}
+					OutArray.Add(CastedItem);
 				}
 			}
 		}
+	}
 
-		/** Apply a function to each item of a specified type */
-		template <typename TreeType>
-		void ForEachItem(TFunctionRef<void(TreeType&)> Func) const
+	/** Apply a function to each item of a specified type */
+	template <typename TreeType>
+	void ForEachItem(TFunctionRef<void(TreeType&)> Func) const
+	{
+		for (const TWeakPtr<ISceneOutlinerTreeItem>& Item : SelectedItems)
 		{
-			for (const TWeakPtr<ITreeItem>& Item : SelectedItems)
+			if (const auto ItemPtr = Item.Pin())
 			{
-				if (const auto ItemPtr = Item.Pin())
+				if (TreeType* CastedItem = ItemPtr->CastTo<TreeType>())
 				{
-					if (TreeType* CastedItem = ItemPtr->CastTo<TreeType>())
-					{
-						Func(*CastedItem);
-					}
+					Func(*CastedItem);
 				}
 			}
 		}
+	}
 
-		/** Use a selector to retrieve a specific data type from items in the selection. Will only add an item's data if the selector returns true for that item. */
-		template <typename DataType>
-		TArray<DataType> GetData(TFunctionRef<bool(const TWeakPtr<ITreeItem>&, DataType&)> Selector) const
+	/** Use a selector to retrieve a specific data type from items in the selection. Will only add an item's data if the selector returns true for that item. */
+	template <typename DataType>
+	TArray<DataType> GetData(TFunctionRef<bool(const TWeakPtr<ISceneOutlinerTreeItem>&, DataType&)> Selector) const
+	{
+		TArray<DataType> Result;
+		for (TWeakPtr<ISceneOutlinerTreeItem>& Item : SelectedItems)
 		{
-			TArray<DataType> Result;
-			for (TWeakPtr<ITreeItem>& Item : SelectedItems)
+			DataType Data;
+			if (Selector(Item, Data))
 			{
-				DataType Data;
-				if (Selector(Item, Data))
-				{
-					Result.Add(Data);
-				}
+				Result.Add(Data);
 			}
-			return Result;
 		}
-	};
-
-	/**
-	 * Scene Outliner widget
-	 */
-	class SSceneOutliner : public ISceneOutliner, public FEditorUndoClient, public FGCObject
-	{
+		return Result;
+	}
+};
 
-	public:
+/**
+	* Scene Outliner widget
+	*/
+class SSceneOutliner : public ISceneOutliner, public FEditorUndoClient, public FGCObject
+{
 
-		SLATE_BEGIN_ARGS(SSceneOutliner)
-		{}
-		SLATE_END_ARGS()
+public:
 
-			/**
-			 * Construct this widget.  Called by the SNew() Slate macro.
-			 *
-			 * @param	InArgs		Declaration used by the SNew() macro to construct this widget
-			 * @param	InitOptions	Programmer-driven initialization options for this widget
-			 */
-			void Construct(const FArguments& InArgs, const FInitializationOptions& InitOptions);
+	SLATE_BEGIN_ARGS(SSceneOutliner)
+	{}
+	SLATE_END_ARGS()
 
-		/** Default constructor - initializes data that is shared between all tree items */
-		SSceneOutliner() : SharedData(MakeShareable(new FSharedOutlinerData)) {}
+		/**
+			* Construct this widget.  Called by the SNew() Slate macro.
+			*
+			* @param	InArgs		Declaration used by the SNew() macro to construct this widget
+			* @param	InitOptions	Programmer-driven initialization options for this widget
+			*/
+		void Construct(const FArguments& InArgs, const FSceneOutlinerInitializationOptions& InitOptions);
 
-		/** SSceneOutliner destructor */
-		~SSceneOutliner();
+	/** Default constructor - initializes data that is shared between all tree items */
+	SSceneOutliner() : SharedData(MakeShareable(new FSharedSceneOutlinerData)) {}
 
-		/** SWidget interface */
-		virtual void Tick(const FGeometry& AllottedGeometry, const double InCurrentTime, const float InDeltaTime) override;
-		virtual bool SupportsKeyboardFocus() const override;
-		virtual FReply OnKeyDown(const FGeometry& MyGeometry, const FKeyEvent& InKeyEvent) override;
+	/** SSceneOutliner destructor */
+	~SSceneOutliner();
 
-		/** Sends a requests to the Scene Outliner to refresh itself the next chance it gets */
-		virtual void Refresh() override;
+	/** SWidget interface */
+	virtual void Tick(const FGeometry& AllottedGeometry, const double InCurrentTime, const float InDeltaTime) override;
+	virtual bool SupportsKeyboardFocus() const override;
+	virtual FReply OnKeyDown(const FGeometry& MyGeometry, const FKeyEvent& InKeyEvent) override;
 
-		void RefreshSelection();
+	/** Sends a requests to the Scene Outliner to refresh itself the next chance it gets */
+	virtual void Refresh() override;
 
-		//~ Begin FEditorUndoClient Interface
-		virtual void PostUndo(bool bSuccess) override;
-		virtual void PostRedo(bool bSuccess) override { PostUndo(bSuccess); }
-		// End of FEditorUndoClient
+	void RefreshSelection();
 
-		/** @return Returns the common data for this outliner */
-		virtual const FSharedOutlinerData& GetSharedData() const override
-		{
-			return *SharedData;
-		}
+	//~ Begin FEditorUndoClient Interface
+	virtual void PostUndo(bool bSuccess) override;
+	virtual void PostRedo(bool bSuccess) override { PostUndo(bSuccess); }
+	// End of FEditorUndoClient
 
-		/** Get a const reference to the actual tree hierarchy */
-		virtual const STreeView<FTreeItemPtr>& GetTree() const override
-		{
-			return *OutlinerTreeView;
-		}
+	/** @return Returns the common data for this outliner */
+	virtual const FSharedSceneOutlinerData& GetSharedData() const override
+	{
+		return *SharedData;
+	}
 
-		virtual const TSharedPtr< SOutlinerTreeView>& GetTreeView() const
-		{
-			return OutlinerTreeView;
-		}
+	/** Get a const reference to the actual tree hierarchy */
+	virtual const STreeView<FSceneOutlinerTreeItemPtr>& GetTree() const override
+	{
+		return *OutlinerTreeView;
+	}
 
-		/** @return Returns a string to use for highlighting results in the outliner list */
-		virtual TAttribute<FText> GetFilterHighlightText() const override;
+	virtual const TSharedPtr< SSceneOutlinerTreeView>& GetTreeView() const
+	{
+		return OutlinerTreeView;
+	}
 
-		/** Set the keyboard focus to the outliner */
-		virtual void SetKeyboardFocus() override;
+	/** @return Returns a string to use for highlighting results in the outliner list */
+	virtual TAttribute<FText> GetFilterHighlightText() const override;
 
-		/** Gets the cached icon for this class name */
-		virtual const FSlateBrush* GetCachedIconForClass(FName InClassName) const override;
+	/** Set the keyboard focus to the outliner */
+	virtual void SetKeyboardFocus() override;
 
-		/** Sets the cached icon for this class name */
-		virtual void CacheIconForClass(FName InClassName, const FSlateBrush* InSlateBrush) override;
+	/** Gets the cached icon for this class name */
+	virtual const FSlateBrush* GetCachedIconForClass(FName InClassName) const override;
 
-		/** Should the scene outliner accept a request to rename a object */
-		virtual bool CanExecuteRenameRequest(const ITreeItem& ItemPtr) const override;
+	/** Sets the cached icon for this class name */
+	virtual void CacheIconForClass(FName InClassName, const FSlateBrush* InSlateBrush) override;
 
-		/**
-		 * Add a filter to the scene outliner
-		 * @param Filter The filter to apply to the scene outliner
-		 * @return The index of the filter.
-		 */
-		virtual int32 AddFilter(const TSharedRef<SceneOutliner::FOutlinerFilter>& Filter) override;
+	/** Should the scene outliner accept a request to rename a object */
+	virtual bool CanExecuteRenameRequest(const ISceneOutlinerTreeItem& ItemPtr) const override;
 
-		/**
-		 * Remove a filter from the scene outliner
-		 * @param Filter The Filter to remove
-		 * @return True if the filter was removed.
-		 */
-		virtual bool RemoveFilter(const TSharedRef<SceneOutliner::FOutlinerFilter>& Filter) override;
+	/**
+		* Add a filter to the scene outliner
+		* @param Filter The filter to apply to the scene outliner
+		* @return The index of the filter.
+		*/
+	virtual int32 AddFilter(const TSharedRef<FSceneOutlinerFilter>& Filter) override;
 
-		/**
-		 * Retrieve the filter at the specified index
-		 * @param Index The index of the filter to retrive
-		 * @return A valid poiter to a filter if the index was valid
-		 */
-		virtual TSharedPtr<SceneOutliner::FOutlinerFilter> GetFilterAtIndex(int32 Index) override;
+	/**
+		* Remove a filter from the scene outliner
+		* @param Filter The Filter to remove
+		* @return True if the filter was removed.
+		*/
+	virtual bool RemoveFilter(const TSharedRef<FSceneOutlinerFilter>& Filter) override;
 
-		/** Get number of filters applied to the scene outliner */
-		virtual int32 GetFilterCount() const override;
+	/**
+		* Retrieve the filter at the specified index
+		* @param Index The index of the filter to retrive
+		* @return A valid poiter to a filter if the index was valid
+		*/
+	virtual TSharedPtr<FSceneOutlinerFilter> GetFilterAtIndex(int32 Index) override;
 
-		/**
-		 * Add or replace a column of the scene outliner
-		 * Note: The column id must match the id of the column returned by the factory
-		 * @param ColumnId The id of the column to add
-		 * @param ColumInfo The struct that contains the information on how to present and retrieve the column
-		 */
-		virtual void AddColumn(FName ColumId, const SceneOutliner::FColumnInfo& ColumInfo) override;
+	/** Get number of filters applied to the scene outliner */
+	virtual int32 GetFilterCount() const override;
 
-		/**
-		 * Remove a column of the scene outliner
-		 * @param ColumnId The name of the column to remove
-		 */
-		virtual void RemoveColumn(FName ColumId) override;
+	/**
+		* Add or replace a column of the scene outliner
+		* Note: The column id must match the id of the column returned by the factory
+		* @param ColumnId The id of the column to add
+		* @param ColumInfo The struct that contains the information on how to present and retrieve the column
+		*/
+	virtual void AddColumn(FName ColumId, const FSceneOutlinerColumnInfo& ColumInfo) override;
 
-		/** Return the name/Id of the columns of the scene outliner */
-		virtual TArray<FName> GetColumnIds() const override;
+	/**
+		* Remove a column of the scene outliner
+		* @param ColumnId The name of the column to remove
+		*/
+	virtual void RemoveColumn(FName ColumId) override;
 
-		/** @return Returns the current sort mode of the specified column */
-		virtual EColumnSortMode::Type GetColumnSortMode(const FName ColumnId) const;
+	/** Return the name/Id of the columns of the scene outliner */
+	virtual TArray<FName> GetColumnIds() const override;
 
-		/** Request that the tree be sorted at a convenient time */
-		virtual void RequestSort();
+	/** @return Returns the current sort mode of the specified column */
+	virtual EColumnSortMode::Type GetColumnSortMode(const FName ColumnId) const;
 
-		/** Returns true if edit delete can be executed */
-		virtual bool Delete_CanExecute();
+	/** Request that the tree be sorted at a convenient time */
+	virtual void RequestSort();
 
-		/** Returns true if edit rename can be executed */
-		virtual bool Rename_CanExecute();
+	/** Returns true if edit delete can be executed */
+	virtual bool Delete_CanExecute();
 
-		/** Executes rename. */
-		virtual void Rename_Execute();
+	/** Returns true if edit rename can be executed */
+	virtual bool Rename_CanExecute();
 
-		/** Returns true if edit cut can be executed */
-		virtual bool Cut_CanExecute();
+	/** Executes rename. */
+	virtual void Rename_Execute();
 
-		/** Returns true if edit copy can be executed */
-		virtual bool Copy_CanExecute();
+	/** Returns true if edit cut can be executed */
+	virtual bool Cut_CanExecute();
 
-		/** Returns true if edit paste can be executed */
-		virtual bool Paste_CanExecute();
+	/** Returns true if edit copy can be executed */
+	virtual bool Copy_CanExecute();
 
-		/** Can the scene outliner rows generated on drag event */
-		virtual bool CanSupportDragAndDrop() const override;
+	/** Returns true if edit paste can be executed */
+	virtual bool Paste_CanExecute();
 
-		/** Tells the scene outliner that it should do a full refresh, which will clear the entire tree and rebuild it from scratch. */
-		virtual void FullRefresh() override;
+	/** Can the scene outliner rows generated on drag event */
+	virtual bool CanSupportDragAndDrop() const override;
 
-	public:
-		/** Event to react to a user double click on a item */
-		FTreeItemPtrEvent& GetDoubleClickEvent() { return OnDoubleClickOnTreeEvent; }
+	/** Tells the scene outliner that it should do a full refresh, which will clear the entire tree and rebuild it from scratch. */
+	virtual void FullRefresh() override;
 
-		/**
-		 * Allow the system that use the scene outliner to react when it's selection is changed
-		 * Note: This event will only be broadcast on a user input.
-		 */
-		FOnItemSelectionChanged& GetOnItemSelectionChanged() { return OnItemSelectionChanged; }
+public:
+	/** Event to react to a user double click on a item */
+	SceneOutliner::FTreeItemPtrEvent& GetDoubleClickEvent() { return OnDoubleClickOnTreeEvent; }
 
-		/** Set the item selection of the outliner based on a selector function. Any items which return true will be added */
-		virtual void SetSelection(const TFunctionRef<bool(SceneOutliner::ITreeItem&)> Selector) override;
+	/**
+		* Allow the system that use the scene outliner to react when it's selection is changed
+		* Note: This event will only be broadcast on a user input.
+		*/
+	SceneOutliner::FOnItemSelectionChanged& GetOnItemSelectionChanged() { return OnItemSelectionChanged; }
 
-		/** Set the selection status of a set of items in the scene outliner */
-		void SetItemSelection(const TArray<FTreeItemPtr>& InItems, bool bSelected, ESelectInfo::Type SelectInfo = ESelectInfo::Direct);
+	/** Set the item selection of the outliner based on a selector function. Any items which return true will be added */
+	virtual void SetSelection(const TFunctionRef<bool(ISceneOutlinerTreeItem&)> Selector) override;
 
-		/** Set the selection status of a single item in the scene outliner */
-		void SetItemSelection(const FTreeItemPtr& InItem, bool bSelected, ESelectInfo::Type SelectInfo = ESelectInfo::Direct);
+	/** Set the selection status of a set of items in the scene outliner */
+	void SetItemSelection(const TArray<FSceneOutlinerTreeItemPtr>& InItems, bool bSelected, ESelectInfo::Type SelectInfo = ESelectInfo::Direct);
 
-		/** Adds a set of items to the current selection */
-		void AddToSelection(const TArray<FTreeItemPtr>& InItems, ESelectInfo::Type SelectInfo = ESelectInfo::Direct);
+	/** Set the selection status of a single item in the scene outliner */
+	void SetItemSelection(const FSceneOutlinerTreeItemPtr& InItem, bool bSelected, ESelectInfo::Type SelectInfo = ESelectInfo::Direct);
 
-		/** Remove a set of items from the current selection */
-		void RemoveFromSelection(const TArray<FTreeItemPtr>& InItems);
+	/** Adds a set of items to the current selection */
+	void AddToSelection(const TArray<FSceneOutlinerTreeItemPtr>& InItems, ESelectInfo::Type SelectInfo = ESelectInfo::Direct);
 
-		/** Remove an item from the current selection */
-		void RemoveFromSelection(const FTreeItemPtr& InItem);
+	/** Remove a set of items from the current selection */
+	void RemoveFromSelection(const TArray<FSceneOutlinerTreeItemPtr>& InItems, ESelectInfo::Type SelectInfo = ESelectInfo::Direct);
 
-		/**
-		 * Returns the list of currently selected tree items
-		 */
-		virtual TArray<FTreeItemPtr> GetSelectedItems() const { return OutlinerTreeView->GetSelectedItems(); }
+	/**
+		* Returns the list of currently selected tree items
+		*/
+	virtual TArray<FSceneOutlinerTreeItemPtr> GetSelectedItems() const { return OutlinerTreeView->GetSelectedItems(); }
 
-		/**
-		 * Returns the currently selected items.
-		 */
-		virtual FItemSelection GetSelection() const { return FItemSelection(*OutlinerTreeView); }
+	/**
+		* Returns the currently selected items.
+		*/
+	virtual FSceneOutlinerItemSelection GetSelection() const { return FSceneOutlinerItemSelection(*OutlinerTreeView); }
 
-		/**
-		 * Add a folder to the selection of the scene outliner
-		 * @param FolderName The name of the folder to add to selection
-		 */
-		void AddFolderToSelection(const FName& FolderName);
+	/**
+		* Add a folder to the selection of the scene outliner
+		* @param FolderName The name of the folder to add to selection
+		*/
+	void AddFolderToSelection(const FName& FolderName);
 
-		/**
-		 * Remove a folder from the selection of the scene outliner
-		 * @param FolderName The name of the folder to remove from the selection
-		 */
-		void RemoveFolderFromSelection(const FName& FolderName);
+	/**
+		* Remove a folder from the selection of the scene outliner
+		* @param FolderName The name of the folder to remove from the selection
+		*/
+	void RemoveFolderFromSelection(const FName& FolderName);
 
-		/** Deselect all selected items */
-		void ClearSelection();
+	/** Deselect all selected items */
+	void ClearSelection();
 
-		/** Sets the next item to rename */
-		void SetPendingRenameItem(const FTreeItemPtr& InItem) { PendingRenameItem = InItem; Refresh(); }
+	/** Sets the next item to rename */
+	void SetPendingRenameItem(const FSceneOutlinerTreeItemPtr& InItem) { PendingRenameItem = InItem; Refresh(); }
 
-		/** Retrieve an ITreeItem by its ID if it exists in the tree */
-		FTreeItemPtr GetTreeItem(FTreeItemID, bool bIncludePending = false);
+	/** Retrieve an ISceneOutlinerTreeItem by its ID if it exists in the tree */
+	FSceneOutlinerTreeItemPtr GetTreeItem(FSceneOutlinerTreeItemID, bool bIncludePending = false);
 
-		/** Get the outliner filter collection */
-		TSharedPtr<FOutlinerFilters>& GetFilters() { return Filters; }
+	/** Get the outliner filter collection */
+	TSharedPtr<FSceneOutlinerFilters>& GetFilters() { return Filters; }
 
-		/** Create a drag drop operation */
-		TSharedPtr<FDragDropOperation> CreateDragDropOperation(const TArray<FTreeItemPtr>& InTreeItems) const;
+	/** Create a drag drop operation */
+	TSharedPtr<FDragDropOperation> CreateDragDropOperation(const TArray<FSceneOutlinerTreeItemPtr>& InTreeItems) const;
 
-		/** Parse a drag drop operation into a payload */
-		bool ParseDragDrop(FDragDropPayload& OutPayload, const FDragDropOperation& Operation) const;
+	/** Parse a drag drop operation into a payload */
+	bool ParseDragDrop(FSceneOutlinerDragDropPayload& OutPayload, const FDragDropOperation& Operation) const;
 
-		/** Validate a drag drop operation on a drop target */
-		FDragValidationInfo ValidateDrop(const ITreeItem& DropTarget, const FDragDropPayload& Payload) const;
+	/** Validate a drag drop operation on a drop target */
+	FSceneOutlinerDragValidationInfo ValidateDrop(const ISceneOutlinerTreeItem& DropTarget, const FSceneOutlinerDragDropPayload& Payload) const;
 
-		/** Called when a payload is dropped onto a target */
-		void OnDropPayload(ITreeItem& DropTarget, const FDragDropPayload& Payload, const FDragValidationInfo& ValidationInfo) const;
+	/** Called when a payload is dropped onto a target */
+	void OnDropPayload(ISceneOutlinerTreeItem& DropTarget, const FSceneOutlinerDragDropPayload& Payload, const FSceneOutlinerDragValidationInfo& ValidationInfo) const;
 
-		/** Called when a payload is dragged over an item */
-		FReply OnDragOverItem(const FDragDropEvent& Event, const ITreeItem& Item) const;
-	private:
-		/** Methods that implement structural modification logic for the tree */
+	/** Called when a payload is dragged over an item */
+	FReply OnDragOverItem(const FDragDropEvent& Event, const ISceneOutlinerTreeItem& Item) const;
+private:
+	/** Methods that implement structural modification logic for the tree */
 
-		/** Empty all the tree item containers maintained by this outliner */
-		void EmptyTreeItems();
+	/** Empty all the tree item containers maintained by this outliner */
+	void EmptyTreeItems();
 
-		/** Apply incremental changes to, or a complete repopulation of the tree  */
-		void Populate();
+	/** Apply incremental changes to, or a complete repopulation of the tree  */
+	void Populate();
 
-		/** Repopulates the entire tree */
-		void RepopulateEntireTree();
+	/** Repopulates the entire tree */
+	void RepopulateEntireTree();
 
-		/** Adds a single new item to the pending map and creates an add operation for it */
-		void AddPendingItem(FTreeItemPtr Item);
+	/** Adds a single new item to the pending map and creates an add operation for it */
+	void AddPendingItem(FSceneOutlinerTreeItemPtr Item);
 
-		/** Adds a new item and all of its children to the pending items. */
-		void AddPendingItemAndChildren(FTreeItemPtr Item);
+	/** Adds a new item and all of its children to the pending items. */
+	void AddPendingItemAndChildren(FSceneOutlinerTreeItemPtr Item);
 
-		/** Attempts to add a pending item to the current tree. Will add any parents if required. */
-		bool AddItemToTree(FTreeItemRef InItem);
+	/** Attempts to add a pending item to the current tree. Will add any parents if required. */
+	bool AddItemToTree(FSceneOutlinerTreeItemRef InItem);
 
-		/** Add an item to the tree, even if it doesn't match the filter terms. Used to add parent's that would otherwise be filtered out */
-		void AddUnfilteredItemToTree(FTreeItemRef Item);
+	/** Add an item to the tree, even if it doesn't match the filter terms. Used to add parent's that would otherwise be filtered out */
+	void AddUnfilteredItemToTree(FSceneOutlinerTreeItemRef Item);
 
-		/** Ensure that the specified item's parent is added to the tree, if applicable */
-		FTreeItemPtr EnsureParentForItem(FTreeItemRef Item);
+	/** Ensure that the specified item's parent is added to the tree, if applicable */
+	FSceneOutlinerTreeItemPtr EnsureParentForItem(FSceneOutlinerTreeItemRef Item);
 
-		/** Remove the specified item from the tree */
-		void RemoveItemFromTree(FTreeItemRef InItem);
+	/** Remove the specified item from the tree */
+	void RemoveItemFromTree(FSceneOutlinerTreeItemRef InItem);
 
-		/** Called when a child has been removed from the specified parent. Will potentially remove the parent from the tree */
-		void OnChildRemovedFromParent(ITreeItem& Parent);
+	/** Called when a child has been removed from the specified parent. Will potentially remove the parent from the tree */
+	void OnChildRemovedFromParent(ISceneOutlinerTreeItem& Parent);
 
-		/** Called when a child has been moved in the tree hierarchy */
-		void OnItemMoved(const FTreeItemRef& Item);
+	/** Called when a child has been moved in the tree hierarchy */
+	void OnItemMoved(const FSceneOutlinerTreeItemRef& Item);
 
 	public:
 		// Test the filters using stack-allocated data to prevent unnecessary heap allocations
 		template <typename TreeItemType, typename TreeItemData>
-		FTreeItemPtr CreateItemFor(const TreeItemData& Data, TFunctionRef<void(const TreeItemType&)> OnItemPassesFilters, bool bForce = false)
+		FSceneOutlinerTreeItemPtr CreateItemFor(const TreeItemData& Data, TFunctionRef<void(const TreeItemType&)> OnItemPassesFilters, bool bForce = false)
 		{
 			const TreeItemType Temporary(Data);
 			bool bPassesFilters = Filters->PassesAllFilters(Temporary);
@@ -457,334 +455,332 @@ namespace SceneOutliner
 				OnItemPassesFilters(Temporary);
 			}
 
-			bPassesFilters &= SearchBoxFilter->PassesFilter(Temporary);
-
-			if (bForce || bPassesFilters)
-			{
-				FTreeItemPtr Result = MakeShareable(new TreeItemType(Data));
-				Result->Flags.bIsFilteredOut = !bPassesFilters;
-				Result->Flags.bInteractive = Filters->GetInteractiveState(*Result);
-				return Result;
-			}
+		bPassesFilters &= SearchBoxFilter->PassesFilter(Temporary);
 
-			return nullptr;
+		if (bForce || bPassesFilters)
+		{
+			FSceneOutlinerTreeItemPtr Result = MakeShareable(new TreeItemType(Data));
+			Result->Flags.bIsFilteredOut = !bPassesFilters;
+			Result->Flags.bInteractive = Filters->GetInteractiveState(*Result);
+			return Result;
 		}
 
-		/** Instruct the outliner to perform an action on the specified item when it is created */
-		void OnItemAdded(const FTreeItemID& ItemID, uint8 ActionMask);
+		return nullptr;
+	}
 
-		/** Get the columns to be displayed in this outliner */
-		const TMap<FName, TSharedPtr<ISceneOutlinerColumn>>& GetColumns() const
-		{
-			return Columns;
-		}
+	/** Instruct the outliner to perform an action on the specified item when it is created */
+	void OnItemAdded(const FSceneOutlinerTreeItemID& ItemID, uint8 ActionMask);
 
-		bool PassesFilters(const ITreeItem& Item) const
-		{
-			return Filters->PassesAllFilters(Item);
-		}
+	/** Get the columns to be displayed in this outliner */
+	const TMap<FName, TSharedPtr<ISceneOutlinerColumn>>& GetColumns() const
+	{
+		return Columns;
+	}
 
-		/** @return	Returns true if the text filter is currently active */
-		bool IsTextFilterActive() const;
+	bool PassesFilters(const ISceneOutlinerTreeItem& Item) const
+	{
+		return Filters->PassesAllFilters(Item);
+	}
 
-		bool PassesTextFilter(const FTreeItemPtr& Item) const
-		{
-			return SearchBoxFilter->PassesFilter(*Item);
-		}
+	/** @return	Returns true if the text filter is currently active */
+	bool IsTextFilterActive() const;
 
-		bool HasSelectorFocus(FTreeItemPtr Item) const
-		{
-			return OutlinerTreeView->Private_HasSelectorFocus(Item);
-		}
+	bool PassesTextFilter(const FSceneOutlinerTreeItemPtr& Item) const
+	{
+		return SearchBoxFilter->PassesFilter(*Item);
+	}
 
-		/** Handler for when a property changes on any item. Called by the mode */
-		void OnItemLabelChanged(FTreeItemPtr ChangedItem);
-	private:
+	bool HasSelectorFocus(FSceneOutlinerTreeItemPtr Item) const
+	{
+		return OutlinerTreeView->Private_HasSelectorFocus(Item);
+	}
 
-		/** Map of columns that are shown on this outliner. */
-		TMap<FName, TSharedPtr<ISceneOutlinerColumn>> Columns;
+	/** Handler for when a property changes on any item. Called by the mode */
+	void OnItemLabelChanged(FSceneOutlinerTreeItemPtr ChangedItem);
+private:
 
-		/** Set up the columns required for this outliner */
-		void SetupColumns(SHeaderRow& HeaderRow);
+	/** Map of columns that are shown on this outliner. */
+	TMap<FName, TSharedPtr<ISceneOutlinerColumn>> Columns;
 
-		/** Refresh the scene outliner for when a colum was added or removed */
-		void RefreshColums();
+	/** Set up the columns required for this outliner */
+	void SetupColumns(SHeaderRow& HeaderRow);
 
-		/** Populates OutSearchStrings with the strings associated with TreeItem that should be used in searching */
-		void PopulateSearchStrings( const ITreeItem& TreeItem, OUT TArray< FString >& OutSearchStrings ) const;
+	/** Refresh the scene outliner for when a colum was added or removed */
+	void RefreshColums();
 
+	/** Populates OutSearchStrings with the strings associated with TreeItem that should be used in searching */
+	void PopulateSearchStrings( const ISceneOutlinerTreeItem& TreeItem, OUT TArray< FString >& OutSearchStrings ) const;
 
-	public:
-		/** Miscellaneous helper functions */
 
-		/** Scroll the specified item into view */
-		void ScrollItemIntoView(const FTreeItemPtr& Item);
+public:
+	/** Miscellaneous helper functions */
 
-		void SetItemExpansion(const FTreeItemPtr& Item, bool bIsExpanded);
-		
-		bool IsItemExpanded(const FTreeItemPtr& Item) const;
+	/** Scroll the specified item into view */
+	void ScrollItemIntoView(const FSceneOutlinerTreeItemPtr& Item);
 
-	private:
+	void SetItemExpansion(const FSceneOutlinerTreeItemPtr& Item, bool bIsExpanded);
+		
+	bool IsItemExpanded(const FSceneOutlinerTreeItemPtr& Item) const;
 
-		/** Check whether we should be showing folders or not in this scene outliner */
-		bool ShouldShowFolders() const;
+private:
 
-		/** Get an array of selected folders */
-		void GetSelectedFolders(TArray<FFolderTreeItem*>& OutFolders) const;
+	/** Check whether we should be showing folders or not in this scene outliner */
+	bool ShouldShowFolders() const;
 
-		/** Get an array of selected folder names */
-		TArray<FName> GetSelectedFolderNames() const;
-	private:
-		/** Tree view event bindings */
+	/** Get an array of selected folders */
+	void GetSelectedFolders(TArray<FFolderTreeItem*>& OutFolders) const;
 
-		/** Called by STreeView to generate a table row for the specified item */
-		TSharedRef< ITableRow > OnGenerateRowForOutlinerTree( FTreeItemPtr Item, const TSharedRef< STableViewBase >& OwnerTable );
+	/** Get an array of selected folder names */
+	TArray<FName> GetSelectedFolderNames() const;
+private:
+	/** Tree view event bindings */
 
-		/** Called by STreeView to get child items for the specified parent item */
-		void OnGetChildrenForOutlinerTree( FTreeItemPtr InParent, TArray< FTreeItemPtr >& OutChildren );
+	/** Called by STreeView to generate a table row for the specified item */
+	TSharedRef< ITableRow > OnGenerateRowForOutlinerTree( FSceneOutlinerTreeItemPtr Item, const TSharedRef< STableViewBase >& OwnerTable );
 
-		/** Called by STreeView when the tree's selection has changed */
-		void OnOutlinerTreeSelectionChanged( FTreeItemPtr TreeItem, ESelectInfo::Type SelectInfo );
+	/** Called by STreeView to get child items for the specified parent item */
+	void OnGetChildrenForOutlinerTree( FSceneOutlinerTreeItemPtr InParent, TArray< FSceneOutlinerTreeItemPtr >& OutChildren );
 
-		/** Called by STreeView when the user double-clicks on an item in the tree */
-		void OnOutlinerTreeDoubleClick( FTreeItemPtr TreeItem );
+	/** Called by STreeView when the tree's selection has changed */
+	void OnOutlinerTreeSelectionChanged( FSceneOutlinerTreeItemPtr TreeItem, ESelectInfo::Type SelectInfo );
 
-		/** Called by STreeView when an item is scrolled into view */
-		void OnOutlinerTreeItemScrolledIntoView( FTreeItemPtr TreeItem, const TSharedPtr<ITableRow>& Widget );
+	/** Called by STreeView when the user double-clicks on an item in the tree */
+	void OnOutlinerTreeDoubleClick( FSceneOutlinerTreeItemPtr TreeItem );
 
-		/** Called when an item in the tree has been collapsed or expanded */
-		void OnItemExpansionChanged(FTreeItemPtr TreeItem, bool bIsExpanded) const;
+	/** Called by STreeView when an item is scrolled into view */
+	void OnOutlinerTreeItemScrolledIntoView( FSceneOutlinerTreeItemPtr TreeItem, const TSharedPtr<ITableRow>& Widget );
 
-	private:
-		/** Level, editor and other global event hooks required to keep the outliner up to date */
+	/** Called when an item in the tree has been collapsed or expanded */
+	void OnItemExpansionChanged(FSceneOutlinerTreeItemPtr TreeItem, bool bIsExpanded) const;
 
-		void OnHierarchyChangedEvent(FHierarchyChangedData Event);
+private:
+	/** Level, editor and other global event hooks required to keep the outliner up to date */
 
-		/** Handler for when an asset is reloaded */
-		void OnAssetReloaded(const EPackageReloadPhase InPackageReloadPhase, FPackageReloadedEvent* InPackageReloadedEvent);
+	void OnHierarchyChangedEvent(FSceneOutlinerHierarchyChangedData Event);
 
-	public:
-		/** Copy specified folders to clipboard, keeping current clipboard contents if they differ from previous clipboard contents (meaning items were copied) */
-		void CopyFoldersToClipboard(const TArray<FName>& InFolders, const FString& InPrevClipboardContents);
+	/** Handler for when an asset is reloaded */
+	void OnAssetReloaded(const EPackageReloadPhase InPackageReloadPhase, FPackageReloadedEvent* InPackageReloadedEvent);
 
-		/** Called by copy and duplicate */
-		void CopyFoldersBegin();
+public:
+	/** Copy specified folders to clipboard, keeping current clipboard contents if they differ from previous clipboard contents (meaning items were copied) */
+	void CopyFoldersToClipboard(const TArray<FName>& InFolders, const FString& InPrevClipboardContents);
 
-		/** Called by copy and duplicate */
-		void CopyFoldersEnd();
+	/** Called by copy and duplicate */
+	void CopyFoldersBegin();
 
-		/** Called by paste and duplicate */
-		void PasteFoldersBegin(TArray<FName> InFolders);
+	/** Called by copy and duplicate */
+	void CopyFoldersEnd();
 
-		/** Paste folders end logic */
-		void PasteFoldersEnd();
+	/** Called by paste and duplicate */
+	void PasteFoldersBegin(TArray<FName> InFolders);
 
-		/** Called by cut and delete */
-		void DeleteFoldersBegin();
+	/** Paste folders end logic */
+	void PasteFoldersEnd();
 
-		/** Called by cut and delete */
-		void DeleteFoldersEnd();
+	/** Called by cut and delete */
+	void DeleteFoldersBegin();
 
-		/** Get an array of folders to paste */
-		TArray<FName> GetClipboardPasteFolders() const;
+	/** Called by cut and delete */
+	void DeleteFoldersEnd();
 
-		/** Construct folders export string to be used in clipboard */
-		FString ExportFolderList(TArray<FName> InFolders) const;
+	/** Get an array of folders to paste */
+	TArray<FName> GetClipboardPasteFolders() const;
 
-		/** Construct array of folders to be created based on input clipboard string */
-		TArray<FName> ImportFolderList(const FString& InStrBuffer) const;
+	/** Construct folders export string to be used in clipboard */
+	FString ExportFolderList(TArray<FName> InFolders) const;
 
-	public:
-		/** Duplicates current folder and all descendants */
-		void DuplicateFoldersHierarchy();
+	/** Construct array of folders to be created based on input clipboard string */
+	TArray<FName> ImportFolderList(const FString& InStrBuffer) const;
 
-	private:
-		/** Cache selected folders during edit delete */
-		TArray<FFolderTreeItem*> CacheFoldersDelete;
+public:
+	/** Duplicates current folder and all descendants */
+	void DuplicateFoldersHierarchy();
 
-		/** Cache folders for cut/copy/paste/duplicate */
-		TArray<FName> CacheFoldersEdit;
+private:
+	/** Cache selected folders during edit delete */
+	TArray<FFolderTreeItem*> CacheFoldersDelete;
 
-		/** Cache clipboard contents for cut/copy */
-		FString CacheClipboardContents;
+	/** Cache folders for cut/copy/paste/duplicate */
+	TArray<FName> CacheFoldersEdit;
 
-		/** Maps pre-existing children during paste or duplicate */
-		TMap<FName, TArray<FTreeItemID>> CachePasteFolderExistingChildrenMap;
+	/** Cache clipboard contents for cut/copy */
+	FString CacheClipboardContents;
 
-	private:
-		/** Miscellaneous bindings required by the UI */
+	/** Maps pre-existing children during paste or duplicate */
+	TMap<FName, TArray<FSceneOutlinerTreeItemID>> CachePasteFolderExistingChildrenMap;
 
-		/** Called by the editable text control when the filter text is changed by the user */
-		void OnFilterTextChanged( const FText& InFilterText );
+private:
+	/** Miscellaneous bindings required by the UI */
 
-		/** Called by the editable text control when a user presses enter or commits their text change */
-		void OnFilterTextCommitted( const FText& InFilterText, ETextCommit::Type CommitInfo );
+	/** Called by the editable text control when the filter text is changed by the user */
+	void OnFilterTextChanged( const FText& InFilterText );
 
-		/** Called by the filter button to get the image to display in the button */
-		const FSlateBrush* GetFilterButtonGlyph() const;
+	/** Called by the editable text control when a user presses enter or commits their text change */
+	void OnFilterTextCommitted( const FText& InFilterText, ETextCommit::Type CommitInfo );
 
-		/** @return	The filter button tool-tip text */
-		FString GetFilterButtonToolTip() const;
+	/** Called by the filter button to get the image to display in the button */
+	const FSlateBrush* GetFilterButtonGlyph() const;
 
-		/** @return	Returns whether the filter status line should be drawn */
-		EVisibility GetFilterStatusVisibility() const;
+	/** @return	The filter button tool-tip text */
+	FString GetFilterButtonToolTip() const;
 
-		/** @return	Returns the filter status text */
-		FText GetFilterStatusText() const;
+	/** @return	Returns whether the filter status line should be drawn */
+	EVisibility GetFilterStatusVisibility() const;
 
-		/** @return Returns color for the filter status text message, based on success of search filter */
-		FSlateColor GetFilterStatusTextColor() const;
+	/** @return	Returns the filter status text */
+	FText GetFilterStatusText() const;
 
-		/**	Returns the current visibility of the Empty label */
-		EVisibility GetEmptyLabelVisibility() const;
+	/** @return Returns color for the filter status text message, based on success of search filter */
+	FSlateColor GetFilterStatusTextColor() const;
 
-		/** @return the selection mode; disabled entirely if in PIE/SIE mode */
-		ESelectionMode::Type GetSelectionMode() const;
+	/**	Returns the current visibility of the Empty label */
+	EVisibility GetEmptyLabelVisibility() const;
 
-		/** @return the content for the view button */
-		TSharedRef<SWidget> GetViewButtonContent(bool bShowFilters);
+	/** @return the selection mode; disabled entirely if in PIE/SIE mode */
+	ESelectionMode::Type GetSelectionMode() const;
 
-		/** @return the foreground color for the view button */
-		FSlateColor GetViewButtonForegroundColor() const;
+	/** @return the content for the view button */
+	TSharedRef<SWidget> GetViewButtonContent(bool bShowFilters);
 
-	public:
+	/** @return the foreground color for the view button */
+	FSlateColor GetViewButtonForegroundColor() const;
 
-		/** Open a context menu for this scene outliner */
-		TSharedPtr<SWidget> OnOpenContextMenu();
+public:
 
-		void FillFoldersSubMenu(UToolMenu* Menu) const;
-		void AddMoveToFolderOutliner(UToolMenu* Menu) const;
-		void FillSelectionSubMenu(UToolMenu* Menun) const;
-		TSharedRef<TSet<FName>> GatherInvalidMoveToDestinations() const;
+	/** Open a context menu for this scene outliner */
+	TSharedPtr<SWidget> OnOpenContextMenu();
 
-		/** Called to select descendants of the currently selected folders */
-		void SelectFoldersDescendants(bool bSelectImmediateChildrenOnly = false);
+	void FillFoldersSubMenu(UToolMenu* Menu) const;
+	void AddMoveToFolderOutliner(UToolMenu* Menu) const;
+	void FillSelectionSubMenu(UToolMenu* Menun) const;
+	TSharedRef<TSet<FName>> GatherInvalidMoveToDestinations() const;
 
-		/** Moves the current selection to the specified folder path */
-		void MoveSelectionTo(FName NewParent);
+	/** Called to select descendants of the currently selected folders */
+	void SelectFoldersDescendants(bool bSelectImmediateChildrenOnly = false);
 
-		/** Create a new folder under the specified parent name (NAME_None for root) */
-		void CreateFolder();
+	/** Moves the current selection to the specified folder path */
+	void MoveSelectionTo(FName NewParent);
 
-	private:
-		/** Called when the user has clicked the button to add a new folder */
-		FReply OnCreateFolderClicked();
+	/** Create a new folder under the specified parent name (NAME_None for root) */
+	void CreateFolder();
 
-	private:
+private:
+	/** Called when the user has clicked the button to add a new folder */
+	FReply OnCreateFolderClicked();
 
-		/** Context menu opening delegate provided by the client */
-		FOnContextMenuOpening OnContextMenuOpening;
+private:
 
-		TSharedPtr<FSharedOutlinerData> SharedData;
+	/** Context menu opening delegate provided by the client */
+	FOnContextMenuOpening OnContextMenuOpening;
 
-		/** List of pending operations to be applied to the tree */
-		TArray<FPendingTreeOperation> PendingOperations;
+	TSharedPtr<FSharedSceneOutlinerData> SharedData;
 
-		/** Map of actions to apply to new tree items */
-		TMap<FTreeItemID, uint8> NewItemActions;
+	/** List of pending operations to be applied to the tree */
+	TArray<SceneOutliner::FPendingTreeOperation> PendingOperations;
 
-		/** Our tree view */
-		TSharedPtr< SOutlinerTreeView > OutlinerTreeView;
+	/** Map of actions to apply to new tree items */
+	TMap<FSceneOutlinerTreeItemID, uint8> NewItemActions;
 
-		/** A map of all items we have in the tree */
-		FTreeItemMap TreeItemMap;
+	/** Our tree view */
+	TSharedPtr< SSceneOutlinerTreeView > OutlinerTreeView;
 
-		/** Pending tree items that are yet to be added the tree */
-		FTreeItemMap PendingTreeItemMap;
+	/** A map of all items we have in the tree */
+	FSceneOutlinerTreeItemMap TreeItemMap;
 
-		/** Folders pending selection */
-		TArray<FName> PendingFoldersSelect;
+	/** Pending tree items that are yet to be added the tree */
+	FSceneOutlinerTreeItemMap PendingTreeItemMap;
 
-		/** Root level tree items */
-		TArray<FTreeItemPtr> RootTreeItems;
+	/** Folders pending selection */
+	TArray<FName> PendingFoldersSelect;
 
-		/** The button that displays view options */
-		TSharedPtr<SComboButton> ViewOptionsComboButton;
+	/** Root level tree items */
+	TArray<FSceneOutlinerTreeItemPtr> RootTreeItems;
 
-	private:
+	/** The button that displays view options */
+	TSharedPtr<SComboButton> ViewOptionsComboButton;
 
-		/** Structure containing information relating to the expansion state of parent items in the tree */
-		typedef TMap<FTreeItemID, bool> FParentsExpansionState;
+private:
 
-		/** Gets the current expansion state of parent items */
-		FParentsExpansionState GetParentsExpansionState() const;
+	/** Structure containing information relating to the expansion state of parent items in the tree */
+	typedef TMap<FSceneOutlinerTreeItemID, bool> FParentsExpansionState;
 
-		/** Updates the expansion state of parent items after a repopulate, according to the previous state */
-		void SetParentsExpansionState(const FParentsExpansionState& ExpansionStateInfo) const;
-	private:
+	/** Gets the current expansion state of parent items */
+	FParentsExpansionState GetParentsExpansionState() const;
 
-		/** True if the outliner needs to be repopulated at the next appropriate opportunity, usually because our
-		    item set has changed in some way. */
-		uint8 bNeedsRefresh : 1;
+	/** Updates the expansion state of parent items after a repopulate, according to the previous state */
+	void SetParentsExpansionState(const FParentsExpansionState& ExpansionStateInfo) const;
+private:
 
-		/** true if the Scene Outliner should do a full refresh. */
-		uint8 bFullRefresh : 1;
+	/** True if the outliner needs to be repopulated at the next appropriate opportunity, usually because our
+		item set has changed in some way. */
+	uint8 bNeedsRefresh : 1;
 
-		/** true if the Scene Outliner should refresh selection */
-		uint8 bSelectionDirty : 1;
+	/** true if the Scene Outliner should do a full refresh. */
+	uint8 bFullRefresh : 1;
 
-		/** True if the Scene Outliner is currently responding to a level visibility change */
-		uint8 bDisableIntermediateSorting : 1;
+	/** true if the Scene Outliner should refresh selection */
+	uint8 bSelectionDirty : 1;
 
-		uint8 bNeedsColumRefresh : 1;
+	/** True if the Scene Outliner is currently responding to a level visibility change */
+	uint8 bDisableIntermediateSorting : 1;
 
-		/** Reentrancy guard */
-		bool bIsReentrant;
+	uint8 bNeedsColumRefresh : 1;
 
-		/* Widget containing the filtering text box */
-		TSharedPtr< SSearchBox > FilterTextBoxWidget;
+	/** Reentrancy guard */
+	bool bIsReentrant;
 
-		/** The header row of the scene outliner */
-		TSharedPtr< SHeaderRow > HeaderRowWidget;
+	/* Widget containing the filtering text box */
+	TSharedPtr< SSearchBox > FilterTextBoxWidget;
 
-		/** A collection of filters used to filter the displayed items and folders in the scene outliner */
-		TSharedPtr< FOutlinerFilters > Filters;
+	/** The header row of the scene outliner */
+	TSharedPtr< SHeaderRow > HeaderRowWidget;
 
-		/** The TextFilter attached to the SearchBox widget of the Scene Outliner */
-		TSharedPtr< TreeItemTextFilter > SearchBoxFilter;
+	/** A collection of filters used to filter the displayed items and folders in the scene outliner */
+	TSharedPtr< FSceneOutlinerFilters > Filters;
 
-		/** True if the search box will take keyboard focus next frame */
-		bool bPendingFocusNextFrame;
+	/** The TextFilter attached to the SearchBox widget of the Scene Outliner */
+	TSharedPtr< SceneOutliner::TreeItemTextFilter > SearchBoxFilter;
 
-		/** The tree item that is currently pending a rename */
-		TWeakPtr<ITreeItem> PendingRenameItem;
+	/** True if the search box will take keyboard focus next frame */
+	bool bPendingFocusNextFrame;
 
-		TMap<FName, const FSlateBrush*> CachedIcons;
+	/** The tree item that is currently pending a rename */
+	TWeakPtr<ISceneOutlinerTreeItem> PendingRenameItem;
 
-		/** Maintain a count of the number of folders active in the outliner */
-		uint32 FolderCount = 0;
+	TMap<FName, const FSlateBrush*> CachedIcons;
 
-		FTreeItemPtrEvent OnDoubleClickOnTreeEvent;
+	/** Maintain a count of the number of folders active in the outliner */
+	uint32 FolderCount = 0;
 
-		FOnItemSelectionChanged OnItemSelectionChanged;
-	private:
+	SceneOutliner::FTreeItemPtrEvent OnDoubleClickOnTreeEvent;
 
-		virtual void AddReferencedObjects(FReferenceCollector& Collector) override {};
+	SceneOutliner::FOnItemSelectionChanged OnItemSelectionChanged;
+private:
 
-		/** Functions relating to sorting */
+	virtual void AddReferencedObjects(FReferenceCollector& Collector) override {};
 
-		/** Timer for PIE/SIE mode to sort the outliner. */
-		float SortOutlinerTimer;
+	/** Functions relating to sorting */
 
-		/** true if the outliner currently needs to be sorted */
-		bool bSortDirty;
+	/** Timer for PIE/SIE mode to sort the outliner. */
+	float SortOutlinerTimer;
 
-		/** Specify which column to sort with */
-		FName SortByColumn;
+	/** true if the outliner currently needs to be sorted */
+	bool bSortDirty;
 
-		/** Currently selected sorting mode */
-		EColumnSortMode::Type SortMode;
+	/** Specify which column to sort with */
+	FName SortByColumn;
 
-		/** Handles column sorting mode change */
-		void OnColumnSortModeChanged( const EColumnSortPriority::Type SortPriority, const FName& ColumnId, const EColumnSortMode::Type InSortMode );
+	/** Currently selected sorting mode */
+	EColumnSortMode::Type SortMode;
 
-		/** Sort the specified array of items based on the current sort column */
-		void SortItems(TArray<FTreeItemPtr>& Items) const;
+	/** Handles column sorting mode change */
+	void OnColumnSortModeChanged( const EColumnSortPriority::Type SortPriority, const FName& ColumnId, const EColumnSortMode::Type InSortMode );
 
-		virtual uint32 GetTypeSortPriority(const ITreeItem& Item) const override;
+	/** Sort the specified array of items based on the current sort column */
+	void SortItems(TArray<FSceneOutlinerTreeItemPtr>& Items) const;
 
-		/** Handler for recursively expanding/collapsing items */
-		void SetItemExpansionRecursive(FTreeItemPtr Model, bool bInExpansionState);
-	};
+	virtual uint32 GetTypeSortPriority(const ISceneOutlinerTreeItem& Item) const override;
 
-}		// namespace SceneOutliner
+	/** Handler for recursively expanding/collapsing items */
+	void SetItemExpansionRecursive(FSceneOutlinerTreeItemPtr Model, bool bInExpansionState);
+};
diff --git a/Engine/Source/Editor/SceneOutliner/Public/SceneOutliner.h b/Engine/Source/Editor/SceneOutliner/Public/SceneOutliner.h
--- a/Engine/Source/Editor/SceneOutliner/Public/SceneOutliner.h
+++ b/Engine/Source/Editor/SceneOutliner/Public/SceneOutliner.h
@@ -14,7 +14,7 @@ MONOLITHIC_HEADER_BOILERPLATE()
 
 #include "SceneOutlinerModule.h"
 
-#include "ITreeItem.h"
+#include "ISceneOutlinerTreeItem.h"
 
 #include "ActorTreeItem.h"
 #include "FolderTreeItem.h"
diff --git a/Engine/Source/Editor/SceneOutliner/Public/SceneOutlinerDelegates.h b/Engine/Source/Editor/SceneOutliner/Public/SceneOutlinerDelegates.h
--- a/Engine/Source/Editor/SceneOutliner/Public/SceneOutlinerDelegates.h
+++ b/Engine/Source/Editor/SceneOutliner/Public/SceneOutlinerDelegates.h
@@ -4,30 +4,25 @@
 
 #include "UnrealEdMisc.h"
 
-namespace SceneOutliner
+class FSceneOutlinerDelegates
 {
-
-	class FSceneOutlinerDelegates
+public:
+	/** Return a single FSceneOutlinerDelegates object */
+	SCENEOUTLINER_API static FSceneOutlinerDelegates& Get()
 	{
-	public:
-		/** Return a single FSceneOutlinerDelegates object */
-		SCENEOUTLINER_API static FSceneOutlinerDelegates& Get()
-		{
-			// return the singleton object
-			static FSceneOutlinerDelegates Singleton;
-			return Singleton;
-		}
-
-		/**	Broadcasts whenever the current selection changes */
-		FSimpleMulticastDelegate SelectionChanged;
+		// return the singleton object
+		static FSceneOutlinerDelegates Singleton;
+		return Singleton;
+	}
 
-		/** Broadcasts whenever a Component Selection Changes */
-		DECLARE_MULTICAST_DELEGATE_OneParam(FOnComponentSelectionChanged, class UActorComponent*);
-		FOnComponentSelectionChanged OnComponentSelectionChanged;
+	/**	Broadcasts whenever the current selection changes */
+	FSimpleMulticastDelegate SelectionChanged;
 
-		/** Broadcasts whenever a Component has been modified */
-		FSimpleMulticastDelegate OnComponentsUpdated;
+	/** Broadcasts whenever a Component Selection Changes */
+	DECLARE_MULTICAST_DELEGATE_OneParam(FOnComponentSelectionChanged, class UActorComponent*);
+	FOnComponentSelectionChanged OnComponentSelectionChanged;
 
-	};
+	/** Broadcasts whenever a Component has been modified */
+	FSimpleMulticastDelegate OnComponentsUpdated;
 
-} // namespace SceneOutliner
+};
