Removed SceneOutliner namespace

diff --git a/Engine/Source/Editor/SceneOutliner/Public/SceneOutlinerDragDrop.h b/Engine/Source/Editor/SceneOutliner/Public/SceneOutlinerDragDrop.h
--- a/Engine/Source/Editor/SceneOutliner/Public/SceneOutlinerDragDrop.h
+++ b/Engine/Source/Editor/SceneOutliner/Public/SceneOutlinerDragDrop.h
@@ -7,195 +7,191 @@
 #include "Layout/Visibility.h"
 #include "Input/DragAndDrop.h"
 #include "DragAndDrop/CompositeDragDropOp.h"
-#include "ITreeItem.h"
+#include "ISceneOutlinerTreeItem.h"
 
-namespace SceneOutliner
+/** Enum to describe the compatibility of a drag drop operation */
+enum class ESceneOutlinerDropCompatibility : uint8
 {
-	/** Enum to describe the compatibility of a drag drop operation */
-	enum DropCompatibility
-	{
-		Drop_Compatible,
-		Drop_Incompatible,
-		Drop_MultipleSelection_Incompatible,
-		Drop_CompatibleAttach,
-		Drop_IncompatibleGeneric,
-		Drop_CompatibleGeneric,
-		Drop_CompatibleMultipleAttach,
-		Drop_IncompatibleMultipleAttach,
-		Drop_CompatibleDetach,
-		Drop_CompatibleMultipleDetach
-	};
-
-	/** Consilidated drag/drop with parsing functions for the scene outliner */
-	struct FDragDropPayload
-	{
-		/** Default constructor, resulting in unset contents */
-		FDragDropPayload() {}
+	Compatible,
+	Incompatible,
+	MultipleSelection_Incompatible,
+	CompatibleAttach,
+	IncompatibleGeneric,
+	CompatibleGeneric,
+	CompatibleMultipleAttach,
+	IncompatibleMultipleAttach,
+	CompatibleDetach,
+	CompatibleMultipleDetach
+};
+
+/** Consilidated drag/drop with parsing functions for the scene outliner */
+struct FSceneOutlinerDragDropPayload
+{
+	/** Default constructor, resulting in unset contents */
+	FSceneOutlinerDragDropPayload() {}
 
-		/** Populate this payload from an array of tree items */
-		template<typename TreeType>
-		FDragDropPayload(const TArray<TreeType>& InDraggedItems)
+	/** Populate this payload from an array of tree items */
+	template<typename TreeType>
+	FSceneOutlinerDragDropPayload(const TArray<TreeType>& InDraggedItems)
+	{
+		for (const auto& Item : InDraggedItems)
 		{
-			for (const auto& Item : InDraggedItems)
-			{
-				DraggedItems.Add(Item);
-			}
+			DraggedItems.Add(Item);
 		}
+	}
 
-		/** Returns true if the payload has an item of a specified type */
-		template <typename TreeType>
-		bool Has() const
+	/** Returns true if the payload has an item of a specified type */
+	template <typename TreeType>
+	bool Has() const
+	{
+		for (const TWeakPtr<ISceneOutlinerTreeItem>& Item : DraggedItems)
 		{
-			for (const TWeakPtr<ITreeItem>& Item : DraggedItems)
+			if (const auto ItemPtr = Item.Pin())
 			{
-				if (const auto ItemPtr = Item.Pin())
+				if (ItemPtr->IsA<TreeType>())
 				{
-					if (ItemPtr->IsA<TreeType>())
-					{
-						return true;
-					}
+					return true;
 				}
 			}
-			return false;
 		}
+		return false;
+	}
 
-		/** Return an array of all tree items in the payload which are of a specified type */
-		template <typename TreeType>
-		TArray<TreeType*> Get() const
+	/** Return an array of all tree items in the payload which are of a specified type */
+	template <typename TreeType>
+	TArray<TreeType*> Get() const
+	{
+		TArray<TreeType*> Result;
+		for (const TWeakPtr<ISceneOutlinerTreeItem>& Item : DraggedItems)
 		{
-			TArray<TreeType*> Result;
-			for (const TWeakPtr<ITreeItem>& Item : DraggedItems)
+			if (const auto ItemPtr = Item.Pin())
 			{
-				if (const auto ItemPtr = Item.Pin())
+				if (TreeType* CastedItem = ItemPtr->CastTo<TreeType>())
 				{
-					if (TreeType* CastedItem = ItemPtr->CastTo<TreeType>())
-					{
-						Result.Add(CastedItem);
-					}
+					Result.Add(CastedItem);
 				}
 			}
-			return Result;
 		}
+		return Result;
+	}
 
-		/** Apply a function to each item in the payload */
-		template <typename TreeType>
-		void ForEachItem(TFunctionRef<void(TreeType&)> Func) const
+	/** Apply a function to each item in the payload */
+	template <typename TreeType>
+	void ForEachItem(TFunctionRef<void(TreeType&)> Func) const
+	{
+		for (const TWeakPtr<ISceneOutlinerTreeItem>& Item : DraggedItems)
 		{
-			for (const TWeakPtr<ITreeItem>& Item : DraggedItems)
+			if (const auto ItemPtr = Item.Pin())
 			{
-				if (const auto ItemPtr = Item.Pin())
+				if (TreeType* CastedItem = ItemPtr->CastTo<TreeType>())
 				{
-					if (TreeType* CastedItem = ItemPtr->CastTo<TreeType>())
-					{
-						Func(*CastedItem);
-					}
+					Func(*CastedItem);
 				}
 			}
 		}
+	}
 		
-		/** Use a selector to retrieve an array of a specific data type from the items in the payload */
-		template <typename DataType>
-		TArray<DataType> GetData(TFunctionRef<bool(TWeakPtr<ITreeItem>, DataType&)> Selector) const
+	/** Use a selector to retrieve an array of a specific data type from the items in the payload */
+	template <typename DataType>
+	TArray<DataType> GetData(TFunctionRef<bool(TWeakPtr<ISceneOutlinerTreeItem>, DataType&)> Selector) const
+	{
+		TArray<DataType> Result;
+		for (TWeakPtr<ISceneOutlinerTreeItem>& Item : DraggedItems)
 		{
-			TArray<DataType> Result;
-			for (TWeakPtr<ITreeItem>& Item : DraggedItems)
+			DataType Data;
+			if (Selector(Item, Data))
 			{
-				DataType Data;
-				if (Selector(Item, Data))
-				{
-					Result.Add(Data);
-				}
+				Result.Add(Data);
 			}
-			return Result;
 		}
+		return Result;
+	}
 
-		/** List of all dragged items */
-		mutable TArray<TWeakPtr<ITreeItem>> DraggedItems;
-	};
+	/** List of all dragged items */
+	mutable TArray<TWeakPtr<ISceneOutlinerTreeItem>> DraggedItems;
+};
 
-	/** Struct used for validation of a drag/drop operation in the scene outliner */
-	struct FDragValidationInfo
-	{
-		/** The tooltip type to display on the operation */
-		DropCompatibility CompatibilityType;
+/** Struct used for validation of a drag/drop operation in the scene outliner */
+struct FSceneOutlinerDragValidationInfo
+{
+	/** The tooltip type to display on the operation */
+	ESceneOutlinerDropCompatibility CompatibilityType;
 
-		/** The tooltip text to display on the operation */
-		FText ValidationText;
+	/** The tooltip text to display on the operation */
+	FText ValidationText;
 
-		/** Construct this validation information out of a tootip type and some text */
-		FDragValidationInfo(const DropCompatibility InCompatibilityType, const FText InValidationText)
-			: CompatibilityType(InCompatibilityType)
-			, ValidationText(InValidationText)
-		{}
+	/** Construct this validation information out of a tootip type and some text */
+	FSceneOutlinerDragValidationInfo(const ESceneOutlinerDropCompatibility InCompatibilityType, const FText InValidationText)
+		: CompatibilityType(InCompatibilityType)
+		, ValidationText(InValidationText)
+	{}
 
-		/** Return a generic invalid result */
-		static FDragValidationInfo Invalid()
-		{
-			return FDragValidationInfo(DropCompatibility::Drop_IncompatibleGeneric, FText());
-		}
+	/** Return a generic invalid result */
+	static FSceneOutlinerDragValidationInfo Invalid()
+	{
+		return FSceneOutlinerDragValidationInfo(ESceneOutlinerDropCompatibility::IncompatibleGeneric, FText());
+	}
 		
-		/** @return true if this operation is valid, false otheriwse */ 
-		bool IsValid() const
+	/** @return true if this operation is valid, false otheriwse */ 
+	bool IsValid() const
+	{
+		switch(CompatibilityType)
 		{
-			switch(CompatibilityType)
-			{
-			case DropCompatibility::Drop_Compatible:
-			case DropCompatibility::Drop_CompatibleAttach:
-			case DropCompatibility::Drop_CompatibleGeneric:
-			case DropCompatibility::Drop_CompatibleMultipleAttach:
-			case DropCompatibility::Drop_CompatibleDetach:
-			case DropCompatibility::Drop_CompatibleMultipleDetach:
-				return true;
-			default:
-				return false;
-			}
+		case ESceneOutlinerDropCompatibility::Compatible:
+		case ESceneOutlinerDropCompatibility::CompatibleAttach:
+		case ESceneOutlinerDropCompatibility::CompatibleGeneric:
+		case ESceneOutlinerDropCompatibility::CompatibleMultipleAttach:
+		case ESceneOutlinerDropCompatibility::CompatibleDetach:
+		case ESceneOutlinerDropCompatibility::CompatibleMultipleDetach:
+			return true;
+		default:
+			return false;
 		}
-	};
+	}
+};
 
-	/* A drag/drop operation when dragging folders in the scene outliner */
-	struct FFolderDragDropOp: public FDecoratedDragDropOp
-	{
-		DRAG_DROP_OPERATOR_TYPE(FFolderDragDropOp, FDecoratedDragDropOp)
+/* A drag/drop operation when dragging folders in the scene outliner */
+struct FFolderDragDropOp: public FDecoratedDragDropOp
+{
+	DRAG_DROP_OPERATOR_TYPE(FFolderDragDropOp, FDecoratedDragDropOp)
 
-		/** Array of folders that we are dragging */
-		FFolderPaths Folders;
+	/** Array of folders that we are dragging */
+	TArray<FName> Folders;
 
-		void Init(FFolderPaths InFolders);
-	};
+	void Init(TArray<FName> InFolders);
+};
 
-	/** A drag/drop operation that was started from the scene outliner */
-	struct FSceneOutlinerDragDropOp : public FCompositeDragDropOp
-	{
-		DRAG_DROP_OPERATOR_TYPE(FSceneOutlinerDragDropOp, FDecoratedDragDropOp);
+/** A drag/drop operation that was started from the scene outliner */
+struct FSceneOutlinerDragDropOp : public FCompositeDragDropOp
+{
+	DRAG_DROP_OPERATOR_TYPE(FSceneOutlinerDragDropOp, FDecoratedDragDropOp);
 		
-		FSceneOutlinerDragDropOp();
-
-		using FDragDropOperation::Construct;
+	FSceneOutlinerDragDropOp();
 
-		void ResetTooltip()
-		{
-			OverrideText = FText();
-			OverrideIcon = nullptr;
-		}
+	using FDragDropOperation::Construct;
 
-		void SetTooltip(FText InOverrideText, const FSlateBrush* InOverrideIcon)
-		{
-			OverrideText = InOverrideText;
-			OverrideIcon = InOverrideIcon;
-		}
+	void ResetTooltip()
+	{
+		OverrideText = FText();
+		OverrideIcon = nullptr;
+	}
 
-		virtual TSharedPtr<SWidget> GetDefaultDecorator() const override;
+	void SetTooltip(FText InOverrideText, const FSlateBrush* InOverrideIcon)
+	{
+		OverrideText = InOverrideText;
+		OverrideIcon = InOverrideIcon;
+	}
 
-	private:
+	virtual TSharedPtr<SWidget> GetDefaultDecorator() const override;
 
-		EVisibility GetOverrideVisibility() const;
-		EVisibility GetDefaultVisibility() const;
+private:
 
-		FText OverrideText;
-		FText GetOverrideText() const { return OverrideText; }
+	EVisibility GetOverrideVisibility() const;
+	EVisibility GetDefaultVisibility() const;
 
-		const FSlateBrush* OverrideIcon;
-		const FSlateBrush* GetOverrideIcon() const { return OverrideIcon; }
-	};
+	FText OverrideText;
+	FText GetOverrideText() const { return OverrideText; }
 
-}
+	const FSlateBrush* OverrideIcon;
+	const FSlateBrush* GetOverrideIcon() const { return OverrideIcon; }
+};
diff --git a/Engine/Source/Editor/SceneOutliner/Public/SceneOutlinerFilters.h b/Engine/Source/Editor/SceneOutliner/Public/SceneOutlinerFilters.h
--- a/Engine/Source/Editor/SceneOutliner/Public/SceneOutlinerFilters.h
+++ b/Engine/Source/Editor/SceneOutliner/Public/SceneOutlinerFilters.h
@@ -6,200 +6,196 @@
 #include "Misc/IFilter.h"
 #include "Misc/FilterCollection.h"
 #include "SceneOutlinerFwd.h"
-#include "ITreeItem.h"
+#include "ISceneOutlinerTreeItem.h"
 
 class FMenuBuilder;
 
-namespace SceneOutliner
-{
-	/**
-	 * Contains information used to create a filter which will be displayed as user toggleable filter
-	 */
-	class FOutlinerFilterInfo
-	{ 
-	public:
-		FOutlinerFilterInfo(const FText& InFilterTitle, const FText& InFilterTooltip, bool bInActive, const FCreateSceneOutlinerFilter& InFactory = FCreateSceneOutlinerFilter())
-			: FilterTitle(InFilterTitle)
-			, FilterTooltip(InFilterTooltip)
-			, bActive(bInActive)
-			, Factory(InFactory)
-		{}
+/**
+	* Contains information used to create a filter which will be displayed as user toggleable filter
+	*/
+class FSceneOutlinerFilterInfo
+{ 
+public:
+	FSceneOutlinerFilterInfo(const FText& InFilterTitle, const FText& InFilterTooltip, bool bInActive, const FCreateSceneOutlinerFilter& InFactory = FCreateSceneOutlinerFilter())
+		: FilterTitle(InFilterTitle)
+		, FilterTooltip(InFilterTooltip)
+		, bActive(bInActive)
+		, Factory(InFactory)
+	{}
 
-		/** Initialize and apply a new filter */
-		void InitFilter(TSharedPtr<FOutlinerFilters> InFilters);
+	/** Initialize and apply a new filter */
+	void InitFilter(TSharedPtr<FSceneOutlinerFilters> InFilters);
 
-		/** Add menu for this filter */
-		void AddMenu(FMenuBuilder& InMenuBuilder);
+	/** Add menu for this filter */
+	void AddMenu(FMenuBuilder& InMenuBuilder);
 
-		bool IsFilterActive() const;
+	bool IsFilterActive() const;
 
-		DECLARE_EVENT_OneParam(FOutlinerFilterInfo, FOnToggle, bool);
-		FOnToggle& OnToggle() { return OnToggleEvent; }
+	DECLARE_EVENT_OneParam(FSceneOutlinerFilterInfo, FOnToggle, bool);
+	FOnToggle& OnToggle() { return OnToggleEvent; }
 
-	private:
-		void ApplyFilter(bool bActive);
-		void ToggleFilterActive();
+private:
+	void ApplyFilter(bool bActive);
+	void ToggleFilterActive();
 
-		TWeakPtr<FOutlinerFilters> Filters;
+	TWeakPtr<FSceneOutlinerFilters> Filters;
 
-		TSharedPtr<FOutlinerFilter> Filter;
+	TSharedPtr<FSceneOutlinerFilter> Filter;
 
-		FText FilterTitle;
-		FText FilterTooltip;
-		bool bActive;
+	FText FilterTitle;
+	FText FilterTooltip;
+	bool bActive;
 
-		FOnToggle OnToggleEvent;
+	FOnToggle OnToggleEvent;
 
-		FCreateSceneOutlinerFilter Factory;
-	};
+	FCreateSceneOutlinerFilter Factory;
+};
 
+/** A filter that can be applied to any type in the tree */
+class FSceneOutlinerFilter : public IFilter<const ISceneOutlinerTreeItem&>
+{
+public:
 	/** Enum to specify how items that are not explicitly handled by this filter should be managed */
-	enum class EDefaultFilterBehaviour : uint8 { Pass, Fail };
-
-	/** A filter that can be applied to any type in the tree */
-	class FOutlinerFilter : public IFilter<const ITreeItem&>
-	{
-	public:
+	enum class EDefaultBehaviour : uint8 { Pass, Fail };
 
-		/** Event that is fired if this filter changes */
-		DECLARE_DERIVED_EVENT(FOutlinerFilter, IFilter<const ITreeItem&>::FChangedEvent, FChangedEvent);
-		virtual FChangedEvent& OnChanged() override { return ChangedEvent; }
+	/** Event that is fired if this filter changes */
+	DECLARE_DERIVED_EVENT(FSceneOutlinerFilter, IFilter<const ISceneOutlinerTreeItem&>::FChangedEvent, FChangedEvent);
+	virtual FChangedEvent& OnChanged() override { return ChangedEvent; }
 
-	protected:
+protected:
 
-		/**	The event that broadcasts whenever a change occurs to the filter */
-		FChangedEvent ChangedEvent;
+	/**	The event that broadcasts whenever a change occurs to the filter */
+	FChangedEvent ChangedEvent;
 
-		/** Default result of the filter when not overridden in derived classes */
-		const EDefaultFilterBehaviour DefaultBehaviour;
+	/** Default result of the filter when not overridden in derived classes */
+	const EDefaultBehaviour DefaultBehaviour;
 
-		/** Constructor to specify the default result of a filter */
-		FOutlinerFilter(EDefaultFilterBehaviour InDefaultBehaviour)
-			: DefaultBehaviour(InDefaultBehaviour)
-		{}
+	/** Constructor to specify the default result of a filter */
+	FSceneOutlinerFilter(EDefaultBehaviour InDefaultBehaviour)
+		: DefaultBehaviour(InDefaultBehaviour)
+	{}
 
-	private:
-		/** Transient result from the filter operation. Only valid until the next invocation of the filter. */
-		mutable bool bTransientFilterResult;
+private:
+	/** Transient result from the filter operation. Only valid until the next invocation of the filter. */
+	mutable bool bTransientFilterResult;
 
-		/** Check whether the specified item passes our filter */
-		virtual bool PassesFilter(const ITreeItem& InItem) const override
-		{
-			bTransientFilterResult = (DefaultBehaviour == EDefaultFilterBehaviour::Pass);
-			return bTransientFilterResult;
-		}
-	public:
-		/** 
-		 * Check if an item should be interactive according to this filter.
-		 * Default behavior just checks if it passes the filter or not.
-		 */
-		virtual bool GetInteractiveState(const ITreeItem& InItem) const
-		{
-			return PassesFilter(InItem);
-		}
-	};
-
-	/** Outliner filter which will be applied on items which match the specified type */
-	template <class T>
-	struct TOutlinerFilter : public FOutlinerFilter
+	/** Check whether the specified item passes our filter */
+	virtual bool PassesFilter(const ISceneOutlinerTreeItem& InItem) const override
 	{
-		TOutlinerFilter(EDefaultFilterBehaviour InDefaultBehaviour)
-			: FOutlinerFilter(InDefaultBehaviour)
-		{}
+		bTransientFilterResult = (DefaultBehaviour == EDefaultBehaviour::Pass);
+		return bTransientFilterResult;
+	}
+public:
+	/** 
+		* Check if an item should be interactive according to this filter.
+		* Default behavior just checks if it passes the filter or not.
+		*/
+	virtual bool GetInteractiveState(const ISceneOutlinerTreeItem& InItem) const
+	{
+		return PassesFilter(InItem);
+	}
+};
 
-		virtual bool PassesFilter(const ITreeItem& InItem) const override
-		{
-			if (const T* CastedItem = InItem.CastTo<T>())
-			{
-				return PassesFilterImpl(*CastedItem);
-			}
-			return DefaultBehaviour == EDefaultFilterBehaviour::Pass;
-		}
+/** Outliner filter which will be applied on items which match the specified type */
+template <class T>
+struct TSceneOutlinerFilter : public FSceneOutlinerFilter
+{
+	TSceneOutlinerFilter(EDefaultBehaviour InDefaultBehaviour)
+		: FSceneOutlinerFilter(InDefaultBehaviour)
+	{}
 
-		virtual bool PassesFilterImpl(const T& InItem) const
+	virtual bool PassesFilter(const ISceneOutlinerTreeItem& InItem) const override
+	{
+		if (const T* CastedItem = InItem.CastTo<T>())
 		{
-			return DefaultBehaviour == EDefaultFilterBehaviour::Pass;
+			return PassesFilterImpl(*CastedItem);
 		}
+		return DefaultBehaviour == EDefaultBehaviour::Pass;
+	}
 
-		virtual bool GetInteractiveState(const ITreeItem& InItem) const
-		{
-			if (const T* CastedItem = InItem.CastTo<T>())
-			{
-				return GetInteractiveStateImpl(*CastedItem);
-			}
-			return DefaultBehaviour == EDefaultFilterBehaviour::Pass;
-		}
+	virtual bool PassesFilterImpl(const T& InItem) const
+	{
+		return DefaultBehaviour == EDefaultBehaviour::Pass;
+	}
 
-		// If not overriden will just default to testing against PassesFilter
-		virtual bool GetInteractiveStateImpl(const T& InItem) const
+	virtual bool GetInteractiveState(const ISceneOutlinerTreeItem& InItem) const
+	{
+		if (const T* CastedItem = InItem.CastTo<T>())
 		{
-			return PassesFilterImpl(InItem);
+			return GetInteractiveStateImpl(*CastedItem);
 		}
-	};
-
+		return DefaultBehaviour == EDefaultBehaviour::Pass;
+	}
 
-	/** Predicate based filter for the outliner */
-	template <class T>
-	struct TOutlinerPredicateFilter : public TOutlinerFilter<T>
+	// If not overriden will just default to testing against PassesFilter
+	virtual bool GetInteractiveStateImpl(const T& InItem) const
 	{
-		using TFilterPredicate = typename T::FFilterPredicate;
-		using TInteractivePredicate = typename T::FInteractivePredicate;
+		return PassesFilterImpl(InItem);
+	}
+};
 
-		/** Predicate used to filter tree items */
-		mutable TFilterPredicate FilterPred;
-		mutable TInteractivePredicate InteractivePred;
 
-		TOutlinerPredicateFilter(TFilterPredicate InFilterPred, EDefaultFilterBehaviour InDefaultBehaviour, TInteractivePredicate InInteractivePredicate = TInteractivePredicate())
-			: TOutlinerFilter<T>(InDefaultBehaviour)
-			, FilterPred(InFilterPred)
-			, InteractivePred(InInteractivePredicate)
-		{}
+/** Predicate based filter for the outliner */
+template <class T>
+struct TSceneOutlinerPredicateFilter : public TSceneOutlinerFilter<T>
+{
+	using TFilterPredicate = typename T::FFilterPredicate;
+	using TInteractivePredicate = typename T::FInteractivePredicate;
 
-		virtual bool PassesFilterImpl(const T& InItem) const override
-		{
-			return InItem.Filter(FilterPred);
-		}
+	/** Predicate used to filter tree items */
+	mutable TFilterPredicate FilterPred;
+	mutable TInteractivePredicate InteractivePred;
 
-		virtual bool GetInteractiveStateImpl(const T& InItem) const override 
-		{
-			if (InteractivePred.IsBound())
-			{
-				return InItem.GetInteractiveState(InteractivePred);
-			}
+	TSceneOutlinerPredicateFilter(TFilterPredicate InFilterPred, FSceneOutlinerFilter::EDefaultBehaviour InDefaultBehaviour, TInteractivePredicate InInteractivePredicate = TInteractivePredicate())
+		: TSceneOutlinerFilter<T>(InDefaultBehaviour)
+		, FilterPred(InFilterPred)
+		, InteractivePred(InInteractivePredicate)
+	{}
 
-			// If not interactive state impl is provided, default to interactive if filter passes
-			return PassesFilterImpl(InItem);
-		}
-	};
+	virtual bool PassesFilterImpl(const T& InItem) const override
+	{
+		return InItem.Filter(FilterPred);
+	}
 
-	/** Scene outliner filters class. This class wraps a collection of filters and allows items of any type to be tested against the entire set. */
-	struct FOutlinerFilters : public TFilterCollection<const ITreeItem&>
+	virtual bool GetInteractiveStateImpl(const T& InItem) const override 
 	{
-		/** Overridden to ensure we only ever have FOutlinerFilters added */
-		int32 Add(const TSharedPtr<FOutlinerFilter>& Filter)
+		if (InteractivePred.IsBound())
 		{
-			return TFilterCollection::Add(Filter);
+			return InItem.GetInteractiveState(InteractivePred);
 		}
 
-		/** Test whether this tree item passes all filters, and set its interactive state according to the filter it failed (if applicable) */
-		bool GetInteractiveState(const ITreeItem& InItem) const
+		// If not interactive state impl is provided, default to interactive if filter passes
+		return PassesFilterImpl(InItem);
+	}
+};
+
+/** Scene outliner filters class. This class wraps a collection of filters and allows items of any type to be tested against the entire set. */
+struct FSceneOutlinerFilters : public TFilterCollection<const ISceneOutlinerTreeItem&>
+{
+	/** Overridden to ensure we only ever have FSceneOutlinerFilters added */
+	int32 Add(const TSharedPtr<FSceneOutlinerFilter>& Filter)
+	{
+		return TFilterCollection::Add(Filter);
+	}
+
+	/** Test whether this tree item passes all filters, and set its interactive state according to the filter it failed (if applicable) */
+	bool GetInteractiveState(const ISceneOutlinerTreeItem& InItem) const
+	{
+		for (const auto& Filter : ChildFilters)
 		{
-			for (const auto& Filter : ChildFilters)
+			if (!StaticCastSharedPtr<FSceneOutlinerFilter>(Filter)->GetInteractiveState(InItem))
 			{
-				if (!StaticCastSharedPtr<FOutlinerFilter>(Filter)->GetInteractiveState(InItem))
-				{
-					return false;
-				}
+				return false;
 			}
-
-			return true;
 		}
 
-		/** Add a filter predicate to this filter collection */
-		template <typename T>
-		void AddFilterPredicate(typename T::FFilterPredicate InFilterPred, EDefaultFilterBehaviour InDefaultBehaviour = EDefaultFilterBehaviour::Pass, typename T::FInteractivePredicate InInteractivePred = typename T::FInteractivePredicate())
-		{
-			Add(MakeShareable(new TOutlinerPredicateFilter<T>(InFilterPred, InDefaultBehaviour, InInteractivePred)));
-		}
-	};
-}
+		return true;
+	}
+
+	/** Add a filter predicate to this filter collection */
+	template <typename T>
+	void AddFilterPredicate(typename T::FFilterPredicate InFilterPred, FSceneOutlinerFilter::EDefaultBehaviour InDefaultBehaviour = FSceneOutlinerFilter::EDefaultBehaviour::Pass, typename T::FInteractivePredicate InInteractivePred = typename T::FInteractivePredicate())
+	{
+		Add(MakeShareable(new TSceneOutlinerPredicateFilter<T>(InFilterPred, InDefaultBehaviour, InInteractivePred)));
+	}
+};
diff --git a/Engine/Source/Editor/SceneOutliner/Public/SceneOutlinerFwd.h b/Engine/Source/Editor/SceneOutliner/Public/SceneOutlinerFwd.h
--- a/Engine/Source/Editor/SceneOutliner/Public/SceneOutlinerFwd.h
+++ b/Engine/Source/Editor/SceneOutliner/Public/SceneOutlinerFwd.h
@@ -8,45 +8,39 @@
 class ISceneOutliner;
 class ISceneOutlinerColumn;
 
-namespace SceneOutliner
-{
-	struct FTreeItemID;
+struct FSceneOutlinerTreeItemID;
 
-	struct FInitializationOptions;
-	struct FSharedOutlinerData;
-	
-	struct ITreeItem;
-	struct FActorTreeItem;
-	struct FWorldTreeItem;
-	struct FFolderTreeItem;
-	struct FComponentTreeItem;
-	
-	typedef TSharedPtr<ITreeItem> FTreeItemPtr;
-	typedef TSharedRef<ITreeItem> FTreeItemRef;
+struct FSceneOutlinerInitializationOptions;
+struct FSharedSceneOutlinerData;
 
-	typedef TMap<FTreeItemID, FTreeItemPtr> FTreeItemMap;
+struct ISceneOutlinerTreeItem;
+struct FActorTreeItem;
+struct FWorldTreeItem;
+struct FFolderTreeItem;
+struct FComponentTreeItem;
 
-	class ISceneOutlinerHierarchy;
-	class ISceneOutlinerMode;
+typedef TSharedPtr<ISceneOutlinerTreeItem> FSceneOutlinerTreeItemPtr;
+typedef TSharedRef<ISceneOutlinerTreeItem> FSceneOutlinerTreeItemRef;
 
-	class SSceneOutliner;
+typedef TMap<FSceneOutlinerTreeItemID, FSceneOutlinerTreeItemPtr> FSceneOutlinerTreeItemMap;
 
-	class FOutlinerFilter;
-	struct FOutlinerFilters;
+class ISceneOutlinerHierarchy;
+class ISceneOutlinerMode;
 
-	struct FDragDropPayload;
-	struct FDragValidationInfo;
+class SSceneOutliner;
 
-	/** Typedef to define an array of folder names, used during dragging */
-	typedef TArray<FName> FFolderPaths;
-}
+class FSceneOutlinerFilter;
+struct FSceneOutlinerFilters;
 
-DECLARE_DELEGATE_OneParam( FOnSceneOutlinerItemPicked, TSharedRef<SceneOutliner::ITreeItem> );
+struct FSceneOutlinerDragDropPayload;
+struct FSceneOutlinerDragValidationInfo;
 
-DECLARE_DELEGATE_OneParam( FCustomSceneOutlinerDeleteDelegate, const TArray<TWeakPtr<SceneOutliner::ITreeItem>>&  )
+DECLARE_DELEGATE_OneParam( FOnSceneOutlinerItemPicked, TSharedRef<ISceneOutlinerTreeItem> );
+
+DECLARE_DELEGATE_OneParam( FCustomSceneOutlinerDeleteDelegate, const TArray<TWeakPtr<ISceneOutlinerTreeItem>>&  )
 
 /** A delegate used to factory a new column type */
-DECLARE_DELEGATE_RetVal_OneParam( TSharedRef< ISceneOutlinerColumn >, FCreateSceneOutlinerColumn, ISceneOutliner& );
+DECLARE_DELEGATE_RetVal_OneParam( TSharedRef<ISceneOutlinerColumn>, FCreateSceneOutlinerColumn, ISceneOutliner& );
 
 /** A delegate used to factory a new filter type */
-DECLARE_DELEGATE_RetVal( TSharedRef< SceneOutliner::FOutlinerFilter >, FCreateSceneOutlinerFilter );
+DECLARE_DELEGATE_RetVal( TSharedRef<FSceneOutlinerFilter>, FCreateSceneOutlinerFilter );
diff --git a/Engine/Source/Editor/SceneOutliner/Public/SceneOutlinerMenuContext.h b/Engine/Source/Editor/SceneOutliner/Public/SceneOutlinerMenuContext.h
--- a/Engine/Source/Editor/SceneOutliner/Public/SceneOutlinerMenuContext.h
+++ b/Engine/Source/Editor/SceneOutliner/Public/SceneOutlinerMenuContext.h
@@ -7,10 +7,7 @@
 
 #include "SceneOutlinerMenuContext.generated.h"
 
-namespace SceneOutliner
-{
-	class SSceneOutliner;
-}
+class SSceneOutliner;
 
 UCLASS()
 class USceneOutlinerMenuContext : public UObject
@@ -27,7 +24,7 @@ public:
 		NumWorldsSelected = 0;
 	}
 
-	TWeakPtr<SceneOutliner::SSceneOutliner> SceneOutliner;
+	TWeakPtr<SSceneOutliner> SceneOutliner;
 
 	bool bShowParentTree;
 	int32 NumSelectedItems;
diff --git a/Engine/Source/Editor/SceneOutliner/Public/SceneOutlinerModule.h b/Engine/Source/Editor/SceneOutliner/Public/SceneOutlinerModule.h
--- a/Engine/Source/Editor/SceneOutliner/Public/SceneOutlinerModule.h
+++ b/Engine/Source/Editor/SceneOutliner/Public/SceneOutlinerModule.h
@@ -34,35 +34,25 @@ public:
 	 * @return	New scene outliner widget
 	 */
 	virtual TSharedRef<ISceneOutliner> CreateSceneOutliner(
-		const SceneOutliner::FInitializationOptions& InitOptions) const;
-
-	/**
-	 * Creates a scene outliner widget
-	 *
-	 * @param	InitOptions						Programmer-driven configuration for this widget instance
-	 * @param	OnActorPickedDelegate			Optional callback when an actor is selected in 'actor picking' mode
-	 *
-	 * @return	New scene outliner widget
-	 */
-
+		const FSceneOutlinerInitializationOptions& InitOptions) const;
 
 	/* Some common scene outliners */
 
 	/** Creates an actor picker widget. Calls the OnActorPickedDelegate when an item is selected. */
 	virtual TSharedRef<ISceneOutliner> CreateActorPicker(
-		const SceneOutliner::FInitializationOptions& InInitOptions,
+		const FSceneOutlinerInitializationOptions& InInitOptions,
 		const FOnActorPicked& OnActorPickedDelegate,
 		TWeakObjectPtr<UWorld> SpecifiedWorld = nullptr) const;
 
 	/** Creates a component picker widget. Calls the OnComponentPickedDelegate when an item is selected. */
 	virtual TSharedRef<ISceneOutliner> CreateComponentPicker(
-		const SceneOutliner::FInitializationOptions& InInitOptions,
+		const FSceneOutlinerInitializationOptions& InInitOptions,
 		const FOnComponentPicked& OnComponentPickedDelegate,
 		TWeakObjectPtr<UWorld> SpecifiedWorld = nullptr) const;
 
 	/** Creates an actor browser widget (also known as a World Outliner). */
 	virtual TSharedRef<ISceneOutliner> CreateActorBrowser(
-		const SceneOutliner::FInitializationOptions& InInitOptions,
+		const FSceneOutlinerInitializationOptions& InInitOptions,
 		TWeakObjectPtr<UWorld> SpecifiedWorld = nullptr) const;
 
 public:
@@ -83,7 +73,7 @@ public:
 
 	/** Register a new type of default column available to all scene outliners */
 	template< typename T >
-	void RegisterDefaultColumnType(SceneOutliner::FColumnInfo InColumnInfo)
+	void RegisterDefaultColumnType(FSceneOutlinerColumnInfo InColumnInfo)
 	{
 		auto ID = T::GetID();
 		if ( !ColumnMap.Contains( ID ) )
@@ -117,7 +107,7 @@ public:
 	}
 
 	/** Map of column type name -> default column info */
-	TMap< FName, SceneOutliner::FColumnInfo> DefaultColumnMap;
+	TMap< FName, FSceneOutlinerColumnInfo> DefaultColumnMap;
 
 private:
 
diff --git a/Engine/Source/Editor/SceneOutliner/Public/SceneOutlinerPublicTypes.h b/Engine/Source/Editor/SceneOutliner/Public/SceneOutlinerPublicTypes.h
--- a/Engine/Source/Editor/SceneOutliner/Public/SceneOutlinerPublicTypes.h
+++ b/Engine/Source/Editor/SceneOutliner/Public/SceneOutlinerPublicTypes.h
@@ -14,153 +14,149 @@ struct FToolMenuContext;
 
 DECLARE_DELEGATE_TwoParams(FSceneOutlinerModifyContextMenu, FName& /* MenuName */, FToolMenuContext& /* MenuContext */);
 
-namespace SceneOutliner
+/** A delegate used as a factory to defer mode creation in the outliner */
+DECLARE_DELEGATE_RetVal_OneParam(ISceneOutlinerMode*, FCreateSceneOutlinerMode, SSceneOutliner*);
+
+/** Container for built in column types. Function-static so they are available without linking */
+struct FSceneOutlinerBuiltInColumnTypes
 {
-	/** A delegate used as a factory to defer mode creation in the outliner */
-	DECLARE_DELEGATE_RetVal_OneParam(ISceneOutlinerMode*, FCreateOutlinerMode, SSceneOutliner*);
+	/** The gutter column */
+	static const FName& Gutter()
+	{
+		static FName Gutter("Gutter");
+		return Gutter;
+	}
 
-	/** Container for built in column types. Function-static so they are available without linking */
-	struct FBuiltInColumnTypes
+	/** The item label column */
+	static const FName& Label()
 	{
-		/** The gutter column */
-		static const FName& Gutter()
-		{
-			static FName Gutter("Gutter");
-			return Gutter;
-		}
-
-		/** The item label column */
-		static const FName& Label()
-		{
-			static FName Label("ItemLabel");
-			return Label;
-		}
-
-		/** Generic actor info column */
-		static FName& ActorInfo()
-		{
-			static FName ActorInfo("ActorInfo");
-			return ActorInfo;
-		}
-	};
-
-	/** Visibility enum for scene outliner columns */
-	enum class EColumnVisibility : uint8
+		static FName Label("ItemLabel");
+		return Label;
+	}
+
+	/** Generic actor info column */
+	static FName& ActorInfo()
 	{
-		/** This column defaults to being visible on the scene outliner */
-		Visible,
+		static FName ActorInfo("ActorInfo");
+		return ActorInfo;
+	}
+};
+
+/** Visibility enum for scene outliner columns */
+enum class ESceneOutlinerColumnVisibility : uint8
+{
+	/** This column defaults to being visible on the scene outliner */
+	Visible,
 
-		/** This column defaults to being invisible, yet still available on the scene outliner */
-		Invisible,
-	};
+	/** This column defaults to being invisible, yet still available on the scene outliner */
+	Invisible,
+};
 
-	/** Column information for the scene outliner */
-	struct FColumnInfo
-	{
-		FColumnInfo(EColumnVisibility InVisibility, int32 InPriorityIndex, const FCreateSceneOutlinerColumn& InFactory = FCreateSceneOutlinerColumn())
-			: Visibility(InVisibility), PriorityIndex(InPriorityIndex), Factory(InFactory)
-		{}
+/** Column information for the scene outliner */
+struct FSceneOutlinerColumnInfo
+{
+	FSceneOutlinerColumnInfo(ESceneOutlinerColumnVisibility InVisibility, int32 InPriorityIndex, const FCreateSceneOutlinerColumn& InFactory = FCreateSceneOutlinerColumn())
+		: Visibility(InVisibility), PriorityIndex(InPriorityIndex), Factory(InFactory)
+	{}
 
-		FColumnInfo() {}
+	FSceneOutlinerColumnInfo() {}
 
-		FColumnInfo(const FColumnInfo& InColumnInfo)
-			: Visibility(InColumnInfo.Visibility), PriorityIndex(InColumnInfo.PriorityIndex), Factory(InColumnInfo.Factory)
-		{}
+	FSceneOutlinerColumnInfo(const FSceneOutlinerColumnInfo& InColumnInfo)
+		: Visibility(InColumnInfo.Visibility), PriorityIndex(InColumnInfo.PriorityIndex), Factory(InColumnInfo.Factory)
+	{}
 
-		EColumnVisibility 	Visibility;
-		uint8				PriorityIndex;
+	ESceneOutlinerColumnVisibility 	Visibility;
+	uint8				PriorityIndex;
 
-		FCreateSceneOutlinerColumn	Factory;
-	};
+	FCreateSceneOutlinerColumn	Factory;
+};
 
-	/** Settings for the scene outliner which can be quieried publicly */
-	struct FSharedOutlinerData
-	{
-		/**	Invoked whenever the user attempts to delete an actor from within a Scene Outliner in the actor browsing mode */
-		FCustomSceneOutlinerDeleteDelegate CustomDelete;
+/** Settings for the scene outliner which can be quieried publicly */
+struct FSharedSceneOutlinerData
+{
+	/**	Invoked whenever the user attempts to delete an actor from within a Scene Outliner in the actor browsing mode */
+	FCustomSceneOutlinerDeleteDelegate CustomDelete;
 
-		/** Modify context menu before display */
-		FSceneOutlinerModifyContextMenu ModifyContextMenu;
+	/** Modify context menu before display */
+	FSceneOutlinerModifyContextMenu ModifyContextMenu;
 
-		/** Map of column types available to the scene outliner, along with default ordering */
-		TMap<FName, FColumnInfo> ColumnMap;
+	/** Map of column types available to the scene outliner, along with default ordering */
+	TMap<FName, FSceneOutlinerColumnInfo> ColumnMap;
 		
-		/** Whether the Scene Outliner should display parent actors in a Tree */
-		bool bShowParentTree : 1;
-
-		/** True to only show folders in this outliner */
-		bool bOnlyShowFolders : 1;
-
-		/** Show transient objects */
-		bool bShowTransient : 1;
-
-	public:
-
-		/** Constructor */
-		FSharedOutlinerData()
-			: bShowParentTree( true )
-			, bOnlyShowFolders( false )
-			, bShowTransient( false )
-		{}
-
-		/** Set up a default array of columns for this outliner */
-		void UseDefaultColumns();
-	};
-
-	/**
-	 * Settings for the Scene Outliner set by the programmer before spawning an instance of the widget.  This
-	 * is used to modify the outliner's behavior in various ways, such as filtering in or out specific classes
-	 * of actors.
-	 */
-	struct FInitializationOptions : FSharedOutlinerData
-	{
-		/** True if we should draw the header row above the tree view */
-		bool bShowHeaderRow : 1;
-
-		/** Whether the Scene Outliner should expose its searchbox */
-		bool bShowSearchBox : 1;
+	/** Whether the Scene Outliner should display parent actors in a Tree */
+	bool bShowParentTree : 1;
+
+	/** True to only show folders in this outliner */
+	bool bOnlyShowFolders : 1;
+
+	/** Show transient objects */
+	bool bShowTransient : 1;
+
+public:
+
+	/** Constructor */
+	FSharedSceneOutlinerData()
+		: bShowParentTree( true )
+		, bOnlyShowFolders( false )
+		, bShowTransient( false )
+	{}
+
+	/** Set up a default array of columns for this outliner */
+	void UseDefaultColumns();
+};
+
+/**
+	* Settings for the Scene Outliner set by the programmer before spawning an instance of the widget.  This
+	* is used to modify the outliner's behavior in various ways, such as filtering in or out specific classes
+	* of actors.
+	*/
+struct FSceneOutlinerInitializationOptions : FSharedSceneOutlinerData
+{
+	/** True if we should draw the header row above the tree view */
+	bool bShowHeaderRow : 1;
 
-		/** If true, the search box will gain focus when the scene outliner is created */
-		bool bFocusSearchBoxWhenOpened : 1;
+	/** Whether the Scene Outliner should expose its searchbox */
+	bool bShowSearchBox : 1;
 
-		/** If true, the Scene Outliner will expose a Create New Folder button */
-		bool bShowCreateNewFolder : 1;
+	/** If true, the search box will gain focus when the scene outliner is created */
+	bool bFocusSearchBoxWhenOpened : 1;
 
-		/** Optional collection of filters to use when filtering in the Scene Outliner */
-		TSharedPtr<FOutlinerFilters> Filters;		
+	/** If true, the Scene Outliner will expose a Create New Folder button */
+	bool bShowCreateNewFolder : 1;
 
-		FCreateOutlinerMode ModeFactory;
+	/** Optional collection of filters to use when filtering in the Scene Outliner */
+	TSharedPtr<FSceneOutlinerFilters> Filters;		
 
-	public:
+	FCreateSceneOutlinerMode ModeFactory;
 
-		/** Constructor */
-		FInitializationOptions()
-			: bShowHeaderRow( true )
-			, bShowSearchBox( true )
-			, bFocusSearchBoxWhenOpened( false )
-			, bShowCreateNewFolder( true )
-			, Filters( new FOutlinerFilters )
-			, ModeFactory()
-		{}
-	};
+public:
 
-	/** Default metrics for outliner tree items */
-	struct FDefaultTreeItemMetrics
-	{
-		static int32	IconSize() { return 18; };
-		static FMargin	IconPadding() { return FMargin(0.f, 0.f, 6.f, 0.f); };
-	};
+	/** Constructor */
+	FSceneOutlinerInitializationOptions()
+		: bShowHeaderRow( true )
+		, bShowSearchBox( true )
+		, bFocusSearchBoxWhenOpened( false )
+		, bShowCreateNewFolder( true )
+		, Filters( new FSceneOutlinerFilters )
+		, ModeFactory()
+	{}
+};
 
-	/** A struct which gets, and caches the visibility of a tree item */
-	struct SCENEOUTLINER_API FGetVisibilityCache
-	{
-		/** Map of tree item to visibility */
-		mutable TMap<const ITreeItem*, bool> VisibilityInfo;
+/** Default metrics for outliner tree items */
+struct FSceneOutlinerDefaultTreeItemMetrics
+{
+	static int32	IconSize() { return 18; };
+	static FMargin	IconPadding() { return FMargin(0.f, 0.f, 6.f, 0.f); };
+};
 
-		/** Get an item's visibility based on its children */
-		bool RecurseChildren(const ITreeItem& Item) const;
+/** A struct which gets, and caches the visibility of a tree item */
+struct SCENEOUTLINER_API FSceneOutlinerVisibilityCache
+{
+	/** Map of tree item to visibility */
+	mutable TMap<const ISceneOutlinerTreeItem*, bool> VisibilityInfo;
 
-		bool GetVisibility(const ITreeItem& Item) const;
-	};
+	/** Get an item's visibility based on its children */
+	bool RecurseChildren(const ISceneOutlinerTreeItem& Item) const;
 
-}	// namespace SceneOutliner
+	bool GetVisibility(const ISceneOutlinerTreeItem& Item) const;
+};
diff --git a/Engine/Source/Editor/SceneOutliner/Public/SceneOutlinerStandaloneTypes.h b/Engine/Source/Editor/SceneOutliner/Public/SceneOutlinerStandaloneTypes.h
--- a/Engine/Source/Editor/SceneOutliner/Public/SceneOutlinerStandaloneTypes.h
+++ b/Engine/Source/Editor/SceneOutliner/Public/SceneOutlinerStandaloneTypes.h
@@ -10,209 +10,210 @@
 
 namespace SceneOutliner
 {
-	using FTreeItemUniqueID = uint32;
+	/** Folder-specific helper functions */
 
-	/** Variant type that defines an identifier for a tree item. Assumes 'trivial relocatability' as with many unreal containers. */
-	struct FTreeItemID
-	{
-	public:
-		enum class EType : uint8 { Object, Folder, UniqueID, Null };
+	/** Parse a new path (including leaf-name) into this tree item. Does not do any notification */
+	FName GetFolderLeafName(FName InPath);
 
-		/** Default constructed null item ID */
-		FTreeItemID() : Type(EType::Null), CachedHash(0) {}
+	/** Get the parent path for the specified folder path */
+	FORCEINLINE FName GetParentPath(FName Path)
+	{
+		return FName(*FPaths::GetPath(Path.ToString()));
+	}
 
-		/** ID representing a UObject */
-		FTreeItemID(const UObject* InObject) : Type(EType::Object)
-		{
-			check(InObject);
-			new (Data) FObjectKey(InObject);
-			CachedHash = CalculateTypeHash();
-		}
-		FTreeItemID(const FObjectKey& InKey) : Type(EType::Object)
-		{
-			new (Data) FObjectKey(InKey);
-			CachedHash = CalculateTypeHash();
-		}
+	bool PathIsChildOf(const FName& PotentialChild, const FName& Parent);
+}	// namespace SceneOutliner
 
-		/** ID representing a folder */
-		FTreeItemID(const FName& InFolder) : Type(EType::Folder)
-		{
-			new (Data) FName(InFolder);
-			CachedHash = CalculateTypeHash();
-		}
 
-		/** ID representing a generic tree item */
-		FTreeItemID(const FTreeItemUniqueID& CustomID) : Type(EType::UniqueID)
-		{
-			new (Data) FTreeItemUniqueID(CustomID);
-			CachedHash = CalculateTypeHash();
-		}
+/** Variant type that defines an identifier for a tree item. Assumes 'trivial relocatability' as with many unreal containers. */
+struct FSceneOutlinerTreeItemID
+{
+public:
+	using FUniqueID = uint32;
 
-		/** Copy construction / assignment */
-		FTreeItemID(const FTreeItemID& Other)
-		{
-			*this = Other;
-		}
-		FTreeItemID& operator=(const FTreeItemID& Other)
-		{
-			Type = Other.Type;
-			switch(Type)
-			{
-				case EType::Object:			new (Data) FObjectKey(Other.GetAsObjectKey());		break;
-				case EType::Folder:			new (Data) FName(Other.GetAsFolderRef());			break;
-				case EType::UniqueID:		new (Data) FTreeItemUniqueID(Other.GetAsHash());	break;
-				default:																		break;
-			}
-
-			CachedHash = CalculateTypeHash();
-			return *this;
-		}
+	enum class EType : uint8 { Object, Folder, UniqueID, Null };
 
-		/** Move construction / assignment */
-		FTreeItemID(FTreeItemID&& Other)
-		{
-			*this = MoveTemp(Other);
-		}
-		FTreeItemID& operator=(FTreeItemID&& Other)
-		{
-			FMemory::Memswap(this, &Other, sizeof(FTreeItemID));
-			return *this;
-		}
+	/** Default constructed null item ID */
+	FSceneOutlinerTreeItemID() : Type(EType::Null), CachedHash(0) {}
 
-		~FTreeItemID()
-		{
-			switch(Type)
-			{
-				case EType::Object:			GetAsObjectKey().~FObjectKey();							break;
-				case EType::Folder:			GetAsFolderRef().~FName();								break;
-				case EType::UniqueID:		/* NOP */												break;
-				default:																			break;
-			}
-		}
+	/** ID representing a UObject */
+	FSceneOutlinerTreeItemID(const UObject* InObject) : Type(EType::Object)
+	{
+		check(InObject);
+		new (Data) FObjectKey(InObject);
+		CachedHash = CalculateTypeHash();
+	}
+	FSceneOutlinerTreeItemID(const FObjectKey& InKey) : Type(EType::Object)
+	{
+		new (Data) FObjectKey(InKey);
+		CachedHash = CalculateTypeHash();
+	}
 
-		friend bool operator==(const FTreeItemID& One, const FTreeItemID& Other)
-		{
-			return One.Type == Other.Type && One.CachedHash == Other.CachedHash && One.Compare(Other);
-		}
-		friend bool operator!=(const FTreeItemID& One, const FTreeItemID& Other)
-		{
-			return One.Type != Other.Type || One.CachedHash != Other.CachedHash || !One.Compare(Other);
-		}
+	/** ID representing a folder */
+	FSceneOutlinerTreeItemID(const FName& InFolder) : Type(EType::Folder)
+	{
+		new (Data) FName(InFolder);
+		CachedHash = CalculateTypeHash();
+	}
 
-		uint32 CalculateTypeHash() const
-		{
-			uint32 Hash = 0;
-			switch(Type)
-			{
-				case EType::Object:			Hash = GetTypeHash(GetAsObjectKey());				break;
-				case EType::Folder:			Hash = GetTypeHash(GetAsFolderRef());				break;
-				case EType::UniqueID:		Hash = GetAsHash();									break;
-				default:																		break;
-			}
-
-			return HashCombine((uint8)Type, Hash);
-		}
+	/** ID representing a generic tree item */
+	FSceneOutlinerTreeItemID(const FUniqueID& CustomID) : Type(EType::UniqueID)
+	{
+		new (Data) FUniqueID(CustomID);
+		CachedHash = CalculateTypeHash();
+	}
 
-		friend uint32 GetTypeHash(const FTreeItemID& ItemID)
+	/** Copy construction / assignment */
+	FSceneOutlinerTreeItemID(const FSceneOutlinerTreeItemID& Other)
+	{
+		*this = Other;
+	}
+	FSceneOutlinerTreeItemID& operator=(const FSceneOutlinerTreeItemID& Other)
+	{
+		Type = Other.Type;
+		switch(Type)
 		{
-			return ItemID.CachedHash;
+			case EType::Object:			new (Data) FObjectKey(Other.GetAsObjectKey());		break;
+			case EType::Folder:			new (Data) FName(Other.GetAsFolderRef());			break;
+			case EType::UniqueID:		new (Data) FUniqueID(Other.GetAsHash());	break;
+			default:																		break;
 		}
 
-	private:
+		CachedHash = CalculateTypeHash();
+		return *this;
+	}
 
-		FObjectKey& 				GetAsObjectKey() const 			{ return *reinterpret_cast<FObjectKey*>(Data); }
-		FName& 						GetAsFolderRef() const			{ return *reinterpret_cast<FName*>(Data); }
-		FTreeItemUniqueID&			GetAsHash() const				{ return *reinterpret_cast<FTreeItemUniqueID*>(Data); }
+	/** Move construction / assignment */
+	FSceneOutlinerTreeItemID(FSceneOutlinerTreeItemID&& Other)
+	{
+		*this = MoveTemp(Other);
+	}
+	FSceneOutlinerTreeItemID& operator=(FSceneOutlinerTreeItemID&& Other)
+	{
+		FMemory::Memswap(this, &Other, sizeof(FSceneOutlinerTreeItemID));
+		return *this;
+	}
 
-		/** Compares the specified ID with this one - assumes matching types */
-		bool Compare(const FTreeItemID& Other) const
+	~FSceneOutlinerTreeItemID()
+	{
+		switch(Type)
 		{
-			switch(Type)
-			{
-				case EType::Object:			return GetAsObjectKey() == Other.GetAsObjectKey();
-				case EType::Folder:			return GetAsFolderRef() == Other.GetAsFolderRef();
-				case EType::UniqueID:		return GetAsHash() == Other.GetAsHash();
-				case EType::Null:			return true;
-				default: check(false);		return false;
-			}
+			case EType::Object:			GetAsObjectKey().~FObjectKey();							break;
+			case EType::Folder:			GetAsFolderRef().~FName();								break;
+			case EType::UniqueID:		/* NOP */												break;
+			default:																			break;
 		}
+	}
 
-		EType Type;
-
-		uint32 CachedHash;
-		static const uint32 MaxSize = TMaxSizeof<FObjectKey, FName, FTreeItemUniqueID>::Value;
-		mutable uint8 Data[MaxSize];
-	};
-	
-	struct FTreeItemType
+	friend bool operator==(const FSceneOutlinerTreeItemID& One, const FSceneOutlinerTreeItemID& Other)
 	{
-	public:
-		explicit FTreeItemType(const FTreeItemType* Parent = nullptr) : ID(++NextUniqueID), ParentType(Parent) {}
-		FTreeItemType(const FTreeItemType& Src) : ID(Src.ID), ParentType(Src.ParentType) {}
-
-		bool operator==(const FTreeItemType& Other) const
-		{
-			return ID == Other.ID || (ParentType != nullptr && *ParentType == Other);
-		}
+		return One.Type == Other.Type && One.CachedHash == Other.CachedHash && One.Compare(Other);
+	}
+	friend bool operator!=(const FSceneOutlinerTreeItemID& One, const FSceneOutlinerTreeItemID& Other)
+	{
+		return One.Type != Other.Type || One.CachedHash != Other.CachedHash || !One.Compare(Other);
+	}
 
-		bool IsA(const FTreeItemType& Other) const
+	uint32 CalculateTypeHash() const
+	{
+		uint32 Hash = 0;
+		switch(Type)
 		{
-			return (ID == Other.ID) || (ParentType && ParentType->IsA(Other));
+			case EType::Object:			Hash = GetTypeHash(GetAsObjectKey());				break;
+			case EType::Folder:			Hash = GetTypeHash(GetAsFolderRef());				break;
+			case EType::UniqueID:		Hash = GetAsHash();									break;
+			default:																		break;
 		}
 
-	private:
-		static uint32 NextUniqueID;
-		uint32 ID;
-		const FTreeItemType* ParentType;
-	};
+		return HashCombine((uint8)Type, Hash);
+	}
 
-	struct FCommonLabelData
+	friend uint32 GetTypeHash(const FSceneOutlinerTreeItemID& ItemID)
 	{
-		TWeakPtr<ISceneOutliner> WeakSceneOutliner;
-		static const FLinearColor DarkColor;
+		return ItemID.CachedHash;
+	}
 
-		TOptional<FLinearColor> GetForegroundColor(const ITreeItem& TreeItem) const;
+private:
 
-		bool CanExecuteRenameRequest(const ITreeItem& Item) const;
-	};
+	FObjectKey& 	GetAsObjectKey() const 		{ return *reinterpret_cast<FObjectKey*>(Data); }
+	FName& 			GetAsFolderRef() const		{ return *reinterpret_cast<FName*>(Data); }
+	FUniqueID&		GetAsHash() const			{ return *reinterpret_cast<FUniqueID*>(Data); }
 
-	/**
-	 * Contains hierarchy change data.
-	 * When an item is added, it will contain a pointer to the new item itself.
-	 * When an item is removed or moved, it will contain the unique ItemID to that item.
-	 * In the case that a folder is being moved, it will also contain the new path to that folder.
-	 */
-	struct FHierarchyChangedData
+	/** Compares the specified ID with this one - assumes matching types */
+	bool Compare(const FSceneOutlinerTreeItemID& Other) const
 	{
-		enum
+		switch(Type)
 		{
-			Added,
-			Removed,
-			Moved,
-			FolderMoved,
-			FullRefresh
-		} Type;
-
-		// This event may pass one of two kinds of data, depending on the type of event
-		FTreeItemPtr Item;
-
-		FTreeItemID ItemID;
-		// Used for FolderMoved events
-		FName NewPath;
-		/** Actions to apply to items */
-		uint8 ItemActions = 0;
-	};
+			case EType::Object:			return GetAsObjectKey() == Other.GetAsObjectKey();
+			case EType::Folder:			return GetAsFolderRef() == Other.GetAsFolderRef();
+			case EType::UniqueID:		return GetAsHash() == Other.GetAsHash();
+			case EType::Null:			return true;
+			default: check(false);		return false;
+		}
+	}
 
-	/** Folder-specific helper functions */
+	EType Type;
 
-	/** Parse a new path (including leaf-name) into this tree item. Does not do any notification */
-	FName GetFolderLeafName(FName InPath);
+	uint32 CachedHash;
+	static const uint32 MaxSize = TMaxSizeof<FObjectKey, FName, FUniqueID>::Value;
+	mutable uint8 Data[MaxSize];
+};
+	
+struct FSceneOutlinerTreeItemType
+{
+public:
+	explicit FSceneOutlinerTreeItemType(const FSceneOutlinerTreeItemType* Parent = nullptr) : ID(++NextUniqueID), ParentType(Parent) {}
+	FSceneOutlinerTreeItemType(const FSceneOutlinerTreeItemType& Src) : ID(Src.ID), ParentType(Src.ParentType) {}
 
-	/** Get the parent path for the specified folder path */
-	FORCEINLINE FName GetParentPath(FName Path)
+	bool operator==(const FSceneOutlinerTreeItemType& Other) const
 	{
-		return FName(*FPaths::GetPath(Path.ToString()));
+		return ID == Other.ID || (ParentType != nullptr && *ParentType == Other);
 	}
 
-	bool PathIsChildOf(const FName& PotentialChild, const FName& Parent);
-}	// namespace SceneOutliner
+	bool IsA(const FSceneOutlinerTreeItemType& Other) const
+	{
+		return (ID == Other.ID) || (ParentType && ParentType->IsA(Other));
+	}
+
+private:
+	static uint32 NextUniqueID;
+	uint32 ID;
+	const FSceneOutlinerTreeItemType* ParentType;
+};
+
+struct FSceneOutlinerCommonLabelData
+{
+	TWeakPtr<ISceneOutliner> WeakSceneOutliner;
+	static const FLinearColor DarkColor;
+
+	TOptional<FLinearColor> GetForegroundColor(const ISceneOutlinerTreeItem& TreeItem) const;
+
+	bool CanExecuteRenameRequest(const ISceneOutlinerTreeItem& Item) const;
+};
+
+/**
+	* Contains hierarchy change data.
+	* When an item is added, it will contain a pointer to the new item itself.
+	* When an item is removed or moved, it will contain the unique ItemID to that item.
+	* In the case that a folder is being moved, it will also contain the new path to that folder.
+	*/
+struct FSceneOutlinerHierarchyChangedData
+{
+	enum
+	{
+		Added,
+		Removed,
+		Moved,
+		FolderMoved,
+		FullRefresh
+	} Type;
+
+	// This event may pass one of two kinds of data, depending on the type of event
+	FSceneOutlinerTreeItemPtr Item;
+
+	FSceneOutlinerTreeItemID ItemID;
+	// Used for FolderMoved events
+	FName NewPath;
+	/** Actions to apply to items */
+	uint8 ItemActions = 0;
+};
diff --git a/Engine/Source/Editor/SceneOutliner/Public/SortHelper.h b/Engine/Source/Editor/SceneOutliner/Public/SortHelper.h
--- a/Engine/Source/Editor/SceneOutliner/Public/SortHelper.h
+++ b/Engine/Source/Editor/SceneOutliner/Public/SortHelper.h
@@ -4,129 +4,10 @@
 
 #include "CoreMinimal.h"
 #include "Widgets/Views/SHeaderRow.h"
-#include "ITreeItem.h"
+#include "ISceneOutlinerTreeItem.h"
 
 namespace SceneOutliner
 {
-	/**
-	 * Templated helper to alleviate performance problems with sorting based on complex predicates.
-	 * Example Usage:
-	 * 		FSortHelper<FString>().Primary([](const ITreeItem& Item){ return Item->GetString(); }).SortItems(Array);
-	 *
-	 * Or:
-	 * 		FSortHelper<FString, FString>()
-	 *			.Primary(FGetPrimaryString())
-	 *			.Secondary(FGetSecondaryString())
-	 *			.SortItems(Array);
-	 */
-	template<typename PrimaryKeyType, typename SecondaryKeyType = int32>
-	struct FSortHelper
-	{
-		typedef TFunction<PrimaryKeyType(const ITreeItem&)> FPrimaryFunction;
-		typedef TFunction<SecondaryKeyType(const ITreeItem&)> FSecondaryFunction;
-
-		FSortHelper()
-			: PrimarySortMode(EColumnSortMode::None), SecondarySortMode(EColumnSortMode::None)
-		{}
-
-		/** Sort primarily by the specified function and mode. Beware the function is a reference, so must be valid for the lifetime of this instance. */
-		FSortHelper& Primary(FPrimaryFunction&& Function, EColumnSortMode::Type SortMode)
-		{
-			PrimarySortMode = SortMode;
-			PrimaryFunction = MoveTemp(Function);
-			return *this;
-		}
-
-		/** Sort secondarily by the specified function and mode. Beware the function is a reference, so must be valid for the lifetime of this instance. */
-		FSortHelper& Secondary(FSecondaryFunction&& Function, EColumnSortMode::Type SortMode)
-		{
-			SecondarySortMode = SortMode;
-			SecondaryFunction = MoveTemp(Function);
-			return *this;
-		}
-
-		/** Sort the specified array using the current sort settings */
-		void Sort(TArray<FTreeItemPtr>& Array)
-		{
-			TArray<FSortPayload> SortData;
-			const auto End = Array.Num();
-			for (int32 Index = 0; Index < End; ++Index)
-			{
-				const auto& Element = Array[Index];
-
-				PrimaryKeyType PrimaryKey = PrimaryFunction(*Element);
-
-				SecondaryKeyType SecondaryKey;
-				if (SecondarySortMode != EColumnSortMode::None)
-				{
-					SecondaryKey = SecondaryFunction(*Element);
-				}
-
-				SortData.Emplace(Index, MoveTemp(PrimaryKey), MoveTemp(SecondaryKey));
-			}
-
-			SortData.Sort([&](const FSortPayload& One, const FSortPayload& Two){
-				if (PrimarySortMode == EColumnSortMode::Ascending && One.PrimaryKey != Two.PrimaryKey)
-					return One.PrimaryKey < Two.PrimaryKey;
-				else if (PrimarySortMode == EColumnSortMode::Descending && One.PrimaryKey != Two.PrimaryKey)
-					return One.PrimaryKey > Two.PrimaryKey;
-
-				if (SecondarySortMode == EColumnSortMode::Ascending)
-					return One.SecondaryKey < Two.SecondaryKey;
-				else if (SecondarySortMode == EColumnSortMode::Descending)
-					return One.SecondaryKey > Two.SecondaryKey;
-
-				return false;
-			});
-
-			TArray<FTreeItemPtr> NewArray;
-			NewArray.Reserve(Array.Num());
-
-			for (const auto& Element : SortData)
-			{
-				NewArray.Add(Array[Element.OriginalIndex]);
-			}
-			Array = MoveTemp(NewArray);
-		}
-
-	private:
-
-		EColumnSortMode::Type 	PrimarySortMode;
-		EColumnSortMode::Type 	SecondarySortMode;
-
-		FPrimaryFunction	PrimaryFunction;
-		FSecondaryFunction	SecondaryFunction;
-
-		/** Aggregated data from the sort methods. We extract the sort data from all elements first, then sort based on the extracted data. */
-		struct FSortPayload
-		{
-			int32 OriginalIndex;
-			
-			PrimaryKeyType PrimaryKey;
-			SecondaryKeyType SecondaryKey;
-
-			FSortPayload(int32 InOriginalIndex, PrimaryKeyType&& InPrimaryKey, SecondaryKeyType&& InSecondaryKey)
-				: OriginalIndex(InOriginalIndex), PrimaryKey(MoveTemp(InPrimaryKey)), SecondaryKey(MoveTemp(InSecondaryKey))
-			{}
-
-			FSortPayload(FSortPayload&& Other)
-			{
-				(*this) = MoveTemp(Other);
-			}
-			FSortPayload& operator=(FSortPayload&& rhs)
-			{
-				OriginalIndex = rhs.OriginalIndex;
-				PrimaryKey = MoveTemp(rhs.PrimaryKey);
-				SecondaryKey = MoveTemp(rhs.SecondaryKey);
-				return *this;
-			}
-
-		private:
-			FSortPayload(const FSortPayload&);
-			FSortPayload& operator=(const FSortPayload&);
-		};
-	};
-
 	/** Wrapper type that will sort FString's using a more natural comparison (correctly sorts numbers and ignores underscores) */
 	struct FNumericStringWrapper
 	{
@@ -253,3 +134,123 @@ namespace SceneOutliner
 		FNumericStringWrapper& operator=(const FNumericStringWrapper&);
 	};
 }
+
+
+/**
+	* Templated helper to alleviate performance problems with sorting based on complex predicates.
+	* Example Usage:
+	* 		FSceneOutlinerSortHelper<FString>().Primary([](const ISceneOutlinerTreeItem& Item){ return Item->GetString(); }).SortItems(Array);
+	*
+	* Or:
+	* 		FSceneOutlinerSortHelper<FString, FString>()
+	*			.Primary(FGetPrimaryString())
+	*			.Secondary(FGetSecondaryString())
+	*			.SortItems(Array);
+	*/
+template<typename PrimaryKeyType, typename SecondaryKeyType = int32>
+struct FSceneOutlinerSortHelper
+{
+	typedef TFunction<PrimaryKeyType(const ISceneOutlinerTreeItem&)> FPrimaryFunction;
+	typedef TFunction<SecondaryKeyType(const ISceneOutlinerTreeItem&)> FSecondaryFunction;
+
+	FSceneOutlinerSortHelper()
+		: PrimarySortMode(EColumnSortMode::None), SecondarySortMode(EColumnSortMode::None)
+	{}
+
+	/** Sort primarily by the specified function and mode. Beware the function is a reference, so must be valid for the lifetime of this instance. */
+	FSceneOutlinerSortHelper& Primary(FPrimaryFunction&& Function, EColumnSortMode::Type SortMode)
+	{
+		PrimarySortMode = SortMode;
+		PrimaryFunction = MoveTemp(Function);
+		return *this;
+	}
+
+	/** Sort secondarily by the specified function and mode. Beware the function is a reference, so must be valid for the lifetime of this instance. */
+	FSceneOutlinerSortHelper& Secondary(FSecondaryFunction&& Function, EColumnSortMode::Type SortMode)
+	{
+		SecondarySortMode = SortMode;
+		SecondaryFunction = MoveTemp(Function);
+		return *this;
+	}
+
+	/** Sort the specified array using the current sort settings */
+	void Sort(TArray<FSceneOutlinerTreeItemPtr>& Array)
+	{
+		TArray<FSortPayload> SortData;
+		const auto End = Array.Num();
+		for (int32 Index = 0; Index < End; ++Index)
+		{
+			const auto& Element = Array[Index];
+
+			PrimaryKeyType PrimaryKey = PrimaryFunction(*Element);
+
+			SecondaryKeyType SecondaryKey;
+			if (SecondarySortMode != EColumnSortMode::None)
+			{
+				SecondaryKey = SecondaryFunction(*Element);
+			}
+
+			SortData.Emplace(Index, MoveTemp(PrimaryKey), MoveTemp(SecondaryKey));
+		}
+
+		SortData.Sort([&](const FSortPayload& One, const FSortPayload& Two){
+			if (PrimarySortMode == EColumnSortMode::Ascending && One.PrimaryKey != Two.PrimaryKey)
+				return One.PrimaryKey < Two.PrimaryKey;
+			else if (PrimarySortMode == EColumnSortMode::Descending && One.PrimaryKey != Two.PrimaryKey)
+				return One.PrimaryKey > Two.PrimaryKey;
+
+			if (SecondarySortMode == EColumnSortMode::Ascending)
+				return One.SecondaryKey < Two.SecondaryKey;
+			else if (SecondarySortMode == EColumnSortMode::Descending)
+				return One.SecondaryKey > Two.SecondaryKey;
+
+			return false;
+		});
+
+		TArray<FSceneOutlinerTreeItemPtr> NewArray;
+		NewArray.Reserve(Array.Num());
+
+		for (const auto& Element : SortData)
+		{
+			NewArray.Add(Array[Element.OriginalIndex]);
+		}
+		Array = MoveTemp(NewArray);
+	}
+
+private:
+
+	EColumnSortMode::Type 	PrimarySortMode;
+	EColumnSortMode::Type 	SecondarySortMode;
+
+	FPrimaryFunction	PrimaryFunction;
+	FSecondaryFunction	SecondaryFunction;
+
+	/** Aggregated data from the sort methods. We extract the sort data from all elements first, then sort based on the extracted data. */
+	struct FSortPayload
+	{
+		int32 OriginalIndex;
+			
+		PrimaryKeyType PrimaryKey;
+		SecondaryKeyType SecondaryKey;
+
+		FSortPayload(int32 InOriginalIndex, PrimaryKeyType&& InPrimaryKey, SecondaryKeyType&& InSecondaryKey)
+			: OriginalIndex(InOriginalIndex), PrimaryKey(MoveTemp(InPrimaryKey)), SecondaryKey(MoveTemp(InSecondaryKey))
+		{}
+
+		FSortPayload(FSortPayload&& Other)
+		{
+			(*this) = MoveTemp(Other);
+		}
+		FSortPayload& operator=(FSortPayload&& rhs)
+		{
+			OriginalIndex = rhs.OriginalIndex;
+			PrimaryKey = MoveTemp(rhs.PrimaryKey);
+			SecondaryKey = MoveTemp(rhs.SecondaryKey);
+			return *this;
+		}
+
+	private:
+		FSortPayload(const FSortPayload&);
+		FSortPayload& operator=(const FSortPayload&);
+	};
+};
diff --git a/Engine/Source/Editor/SceneOutliner/Public/WorldTreeItem.h b/Engine/Source/Editor/SceneOutliner/Public/WorldTreeItem.h
--- a/Engine/Source/Editor/SceneOutliner/Public/WorldTreeItem.h
+++ b/Engine/Source/Editor/SceneOutliner/Public/WorldTreeItem.h
@@ -4,7 +4,7 @@
 
 #include "CoreMinimal.h"
 #include "Widgets/SWidget.h"
-#include "ITreeItem.h"
+#include "ISceneOutlinerTreeItem.h"
 #include "UObject/ObjectKey.h"
 
 class UToolMenu;
@@ -13,50 +13,49 @@ namespace SceneOutliner
 {
 	/** Get a description of a world to display in the scene outliner */
 	FText GetWorldDescription(UWorld* World);
+}
 
-	/** A tree item that represents an entire world */
-	struct SCENEOUTLINER_API FWorldTreeItem : ITreeItem
+/** A tree item that represents an entire world */
+struct SCENEOUTLINER_API FWorldTreeItem : ISceneOutlinerTreeItem
+{
+public:
+	DECLARE_DELEGATE_RetVal_OneParam(bool, FFilterPredicate, const UWorld*);
+	DECLARE_DELEGATE_RetVal_OneParam(bool, FInteractivePredicate, const UWorld*);
+
+	bool Filter(FFilterPredicate Pred) const
 	{
-	public:
-		DECLARE_DELEGATE_RetVal_OneParam(bool, FFilterPredicate, const UWorld*);
-		DECLARE_DELEGATE_RetVal_OneParam(bool, FInteractivePredicate, const UWorld*);
-
-		bool Filter(FFilterPredicate Pred) const
-		{
-			return Pred.Execute(World.Get());
-		}
-
-		bool GetInteractiveState(FInteractivePredicate Pred) const
-		{
-			return Pred.Execute(World.Get());
-		}
-
-		/** The world this tree item is associated with. */
-		mutable TWeakObjectPtr<UWorld> World;
+		return Pred.Execute(World.Get());
+	}
+
+	bool GetInteractiveState(FInteractivePredicate Pred) const
+	{
+		return Pred.Execute(World.Get());
+	}
+
+	/** The world this tree item is associated with. */
+	mutable TWeakObjectPtr<UWorld> World;
 		
-		/** Constant identifier for this tree item */
-		const FObjectKey ID;
-
-		/** Static type identifier for this tree item class */
-		static const FTreeItemType Type;
-
-		/** Construct this item from a world */
-		FWorldTreeItem(UWorld* InWorld);
-		FWorldTreeItem(TWeakObjectPtr<UWorld> InWorld);
-
-		/* Begin ITreeItem Implementation */
-		virtual bool IsValid() const override { return World.IsValid(); }
-		virtual FTreeItemID GetID() const override;
-		virtual FString GetDisplayString() const override;
-		virtual bool CanInteract() const override;
-		virtual void GenerateContextMenu(UToolMenu* Menu, SSceneOutliner& Outliner) override;
-		virtual TSharedRef<SWidget> GenerateLabelWidget(ISceneOutliner& Outliner, const STableRow<FTreeItemPtr>& InRow) override;
-		/* End ITreeItem Implementation */
-
-		/** Open the world settings for the contained world */
-		void OpenWorldSettings() const;
-		/** Get just the name of the world, for tooltip use */
-		FString GetWorldName() const;
-	};
-
-}		// namespace SceneOutliner
+	/** Constant identifier for this tree item */
+	const FObjectKey ID;
+
+	/** Static type identifier for this tree item class */
+	static const FSceneOutlinerTreeItemType Type;
+
+	/** Construct this item from a world */
+	FWorldTreeItem(UWorld* InWorld);
+	FWorldTreeItem(TWeakObjectPtr<UWorld> InWorld);
+
+	/* Begin ISceneOutlinerTreeItem Implementation */
+	virtual bool IsValid() const override { return World.IsValid(); }
+	virtual FSceneOutlinerTreeItemID GetID() const override;
+	virtual FString GetDisplayString() const override;
+	virtual bool CanInteract() const override;
+	virtual void GenerateContextMenu(UToolMenu* Menu, SSceneOutliner& Outliner) override;
+	virtual TSharedRef<SWidget> GenerateLabelWidget(ISceneOutliner& Outliner, const STableRow<FSceneOutlinerTreeItemPtr>& InRow) override;
+	/* End ISceneOutlinerTreeItem Implementation */
+
+	/** Open the world settings for the contained world */
+	void OpenWorldSettings() const;
+	/** Get just the name of the world, for tooltip use */
+	FString GetWorldName() const;
+};
diff --git a/Engine/Source/Editor/Sequencer/Private/LevelEditorSequencerIntegration.cpp b/Engine/Source/Editor/Sequencer/Private/LevelEditorSequencerIntegration.cpp
--- a/Engine/Source/Editor/Sequencer/Private/LevelEditorSequencerIntegration.cpp
+++ b/Engine/Source/Editor/Sequencer/Private/LevelEditorSequencerIntegration.cpp
@@ -994,7 +994,7 @@ void FLevelEditorSequencerIntegration::AttachOutlinerColumn()
 
 	FSceneOutlinerModule& SceneOutlinerModule = FModuleManager::LoadModuleChecked< FSceneOutlinerModule >("SceneOutliner");
 
-	SceneOutliner::FColumnInfo ColumnInfo(SceneOutliner::EColumnVisibility::Visible, 15, 
+	FSceneOutlinerColumnInfo ColumnInfo(ESceneOutlinerColumnVisibility::Visible, 15, 
 		FCreateSceneOutlinerColumn::CreateRaw( this, &FLevelEditorSequencerIntegration::CreateSequencerInfoColumn));
 
 	SceneOutlinerModule.RegisterDefaultColumnType< Sequencer::FSequencerInfoColumn >(ColumnInfo);
diff --git a/Engine/Source/Editor/Sequencer/Private/Sequencer.cpp b/Engine/Source/Editor/Sequencer/Private/Sequencer.cpp
--- a/Engine/Source/Editor/Sequencer/Private/Sequencer.cpp
+++ b/Engine/Source/Editor/Sequencer/Private/Sequencer.cpp
@@ -5609,11 +5609,9 @@ void FSequencer::UpdatePreviewLevelViewportClientFromCameraCut(FLevelEditorViewp
 	AActor* CameraActor = Cast<AActor>(InCameraObject);
 	AActor* PreviousCameraActor = Cast<AActor>(CameraCutParams.PreviousCameraObject);
 
-	const float BlendFactor = FMath::Clamp(CameraCutParams.PreviewBlendFactor, 0.f, 1.f);
-	const bool bIsBlending = (
-			CameraCutParams.BlendTime > 0.f &&
-			BlendFactor < 1.f - SMALL_NUMBER &&
+	const bool bIsBlending = (CameraCutParams.BlendTime > 0.f &&
 			(CameraActor != nullptr || PreviousCameraActor != nullptr));
+	const float BlendFactor = FMath::Clamp(CameraCutParams.PreviewBlendFactor, 0.f, 1.f);
 
 	// To preview blending we'll have to offset the viewport camera using the view modifiers API.
 	bApplyViewModifier = bIsBlending;
@@ -7062,10 +7060,8 @@ void FSequencer::AssignActor(FMenuBuilder& MenuBuilder, FGuid InObjectBinding)
 		return !BoundObjects.Contains(InActor);
 	};
 
-	using namespace SceneOutliner;
-
 	// Set up a menu entry to assign an actor to the object binding node
-	FInitializationOptions InitOptions;
+	FSceneOutlinerInitializationOptions InitOptions;
 	{
 		// We hide the header row to keep the UI compact.
 		InitOptions.bShowHeaderRow = false;
@@ -7073,7 +7069,7 @@ void FSequencer::AssignActor(FMenuBuilder& MenuBuilder, FGuid InObjectBinding)
 		InitOptions.bShowCreateNewFolder = false;
 		InitOptions.bFocusSearchBoxWhenOpened = true;
 		// Only want the actor label column
-		InitOptions.ColumnMap.Add(FBuiltInColumnTypes::Label(), FColumnInfo(EColumnVisibility::Visible, 0));
+		InitOptions.ColumnMap.Add(FSceneOutlinerBuiltInColumnTypes::Label(), FSceneOutlinerColumnInfo(ESceneOutlinerColumnVisibility::Visible, 0));
 
 		// Only display actors that are not possessed already
 		InitOptions.Filters->AddFilterPredicate<FActorTreeItem>(FActorTreeItem::FFilterPredicate::CreateLambda( IsActorValidForAssignment ) );
diff --git a/Engine/Source/Editor/Sequencer/Private/SequencerInfoColumn.cpp b/Engine/Source/Editor/Sequencer/Private/SequencerInfoColumn.cpp
--- a/Engine/Source/Editor/Sequencer/Private/SequencerInfoColumn.cpp
+++ b/Engine/Source/Editor/Sequencer/Private/SequencerInfoColumn.cpp
@@ -29,13 +29,13 @@ struct FGetActorInfo
 		: WeakColumn(StaticCastSharedRef<const FSequencerInfoColumn>(InColumn.AsShared()))
 	{}
 
-	FString operator()(const SceneOutliner::ITreeItem& Item) const
+	FString operator()(const ISceneOutlinerTreeItem& Item) const
 	{
 		if (!WeakColumn.IsValid())
 		{
 			return FString();
 		}
-		if (const SceneOutliner::FActorTreeItem* ActorItem = Item.CastTo<SceneOutliner::FActorTreeItem>())
+		if (const FActorTreeItem* ActorItem = Item.CastTo<FActorTreeItem>())
 		{
 			AActor* Actor = ActorItem->Actor.Get();
 			if (Actor)
@@ -90,7 +90,7 @@ SHeaderRow::FColumn::FArguments FSequencerInfoColumn::ConstructHeaderRowColumn()
 		.FillWidth( 5.0f );
 }
 
-const TSharedRef< SWidget > FSequencerInfoColumn::ConstructRowWidget( SceneOutliner::FTreeItemRef TreeItem, const STableRow<SceneOutliner::FTreeItemPtr>& Row )
+const TSharedRef< SWidget > FSequencerInfoColumn::ConstructRowWidget(FSceneOutlinerTreeItemRef TreeItem, const STableRow<FSceneOutlinerTreeItemPtr>& Row)
 {
 	auto SceneOutliner = WeakSceneOutliner.Pin();
 	check(SceneOutliner.IsValid());
@@ -98,7 +98,7 @@ const TSharedRef< SWidget > FSequencerInfoColumn::ConstructRowWidget( SceneOutli
 	TSharedRef<SHorizontalBox> HorizontalBox = SNew(SHorizontalBox);
 
 	TSharedRef<STextBlock> MainText = SNew( STextBlock )
-		.Text( this, &FSequencerInfoColumn::GetTextForItem, TWeakPtr<SceneOutliner::ITreeItem>(TreeItem) )
+		.Text( this, &FSequencerInfoColumn::GetTextForItem, TWeakPtr<ISceneOutlinerTreeItem>(TreeItem) )
 		.HighlightText( SceneOutliner->GetFilterHighlightText() )
 		.ColorAndOpacity( FSlateColor::UseSubduedForeground() );
 
@@ -113,17 +113,17 @@ const TSharedRef< SWidget > FSequencerInfoColumn::ConstructRowWidget( SceneOutli
 }
 
 
-void FSequencerInfoColumn::PopulateSearchStrings( const SceneOutliner::ITreeItem& Item, TArray< FString >& OutSearchStrings ) const
+void FSequencerInfoColumn::PopulateSearchStrings( const ISceneOutlinerTreeItem& Item, TArray< FString >& OutSearchStrings ) const
 {
 	OutSearchStrings.Add(Item.GetDisplayString());
 }
 
 
-void FSequencerInfoColumn::SortItems(TArray<SceneOutliner::FTreeItemPtr>& OutItems, const EColumnSortMode::Type SortMode) const
+void FSequencerInfoColumn::SortItems(TArray<FSceneOutlinerTreeItemPtr>& OutItems, const EColumnSortMode::Type SortMode) const
 {
 	if (WeakBindingData.IsValid())
 	{
-		SceneOutliner::FSortHelper<FString>()
+		FSceneOutlinerSortHelper<FString>()
 			.Primary(FGetActorInfo(*this), SortMode)
 			.Sort(OutItems);
 	}
@@ -139,7 +139,7 @@ FString FSequencerInfoColumn::GetTextForActor(AActor* InActor) const
 	return FString();
 }
 
-FText FSequencerInfoColumn::GetTextForItem( TWeakPtr<SceneOutliner::ITreeItem> TreeItem ) const
+FText FSequencerInfoColumn::GetTextForItem( TWeakPtr<ISceneOutlinerTreeItem> TreeItem ) const
 {
 	auto Item = TreeItem.Pin();
 	return Item.IsValid() && WeakBindingData.IsValid() ? FText::FromString(FGetActorInfo(*this)(*Item)) : FText::GetEmpty();
diff --git a/Engine/Source/Editor/Sequencer/Private/SequencerInfoColumn.h b/Engine/Source/Editor/Sequencer/Private/SequencerInfoColumn.h
--- a/Engine/Source/Editor/Sequencer/Private/SequencerInfoColumn.h
+++ b/Engine/Source/Editor/Sequencer/Private/SequencerInfoColumn.h
@@ -43,13 +43,13 @@ public:
 
 	virtual SHeaderRow::FColumn::FArguments ConstructHeaderRowColumn() override;
 	
-	virtual const TSharedRef< SWidget > ConstructRowWidget( SceneOutliner::FTreeItemRef TreeItem, const STableRow<SceneOutliner::FTreeItemPtr>& Row ) override;
+	virtual const TSharedRef< SWidget > ConstructRowWidget( FSceneOutlinerTreeItemRef TreeItem, const STableRow<FSceneOutlinerTreeItemPtr>& Row ) override;
 
-	virtual void PopulateSearchStrings( const SceneOutliner::ITreeItem& Item, TArray< FString >& OutSearchStrings ) const override;
+	virtual void PopulateSearchStrings( const ISceneOutlinerTreeItem& Item, TArray< FString >& OutSearchStrings ) const override;
 
 	virtual bool SupportsSorting() const override { return true; }
 
-	virtual void SortItems(TArray<SceneOutliner::FTreeItemPtr>& OutItems, const EColumnSortMode::Type SortMode) const override;
+	virtual void SortItems(TArray<FSceneOutlinerTreeItemPtr>& OutItems, const EColumnSortMode::Type SortMode) const override;
 	
 	// End ISceneOutlinerColumn Implementation
 	//////////////////////////////////////////////////////////////////////////
@@ -58,7 +58,7 @@ public:
 
 private:
 
-	FText GetTextForItem( const TWeakPtr<SceneOutliner::ITreeItem> TreeItem ) const;
+	FText GetTextForItem( const TWeakPtr<ISceneOutlinerTreeItem> TreeItem ) const;
 
 	/** Weak reference to the outliner widget that owns our list */
 	TWeakPtr< ISceneOutliner > WeakSceneOutliner;
