Added copyright notice

diff --git a/Engine/Source/Editor/SceneOutliner/Private/ActorFolderHierarchy.cpp b/Engine/Source/Editor/SceneOutliner/Private/ActorFolderHierarchy.cpp
--- a/Engine/Source/Editor/SceneOutliner/Private/ActorFolderHierarchy.cpp
+++ b/Engine/Source/Editor/SceneOutliner/Private/ActorFolderHierarchy.cpp
@@ -1,5 +1,6 @@
-#include "ActorFolderHierarchy.h"
+// Copyright Epic Games, Inc. All Rights Reserved.
 
+#include "ActorFolderHierarchy.h"
 #include "ISceneOutlinerMode.h"
 #include "WorldTreeItem.h"
 #include "ActorFolderTreeItem.h"
diff --git a/Engine/Source/Editor/SceneOutliner/Private/ActorFolderTreeItem.cpp b/Engine/Source/Editor/SceneOutliner/Private/ActorFolderTreeItem.cpp
--- a/Engine/Source/Editor/SceneOutliner/Private/ActorFolderTreeItem.cpp
+++ b/Engine/Source/Editor/SceneOutliner/Private/ActorFolderTreeItem.cpp
@@ -1,3 +1,5 @@
+// Copyright Epic Games, Inc. All Rights Reserved.
+
 #include "ActorFolderTreeItem.h"
 #include "ActorTreeItem.h"
 #include "SSceneOutliner.h"
