Fix Linux build

diff --git a/Engine/Source/Editor/SceneOutliner/Private/ActorHierarchy.cpp b/Engine/Source/Editor/SceneOutliner/Private/ActorHierarchy.cpp
--- a/Engine/Source/Editor/SceneOutliner/Private/ActorHierarchy.cpp
+++ b/Engine/Source/Editor/SceneOutliner/Private/ActorHierarchy.cpp
@@ -155,22 +155,21 @@ FSceneOutlinerTreeItemPtr FActorHierarchy::FindParent(const ISceneOutlinerTreeIt
 	{
 		if (ActorTreeItem->Actor.IsValid())
 		{
-			const FSceneOutlinerTreeItemPtr* ParentItem;
 			if (const AActor* ParentActor = ActorTreeItem->Actor->GetAttachParentActor())
 			{
-				if (ParentItem = Items.Find(ParentActor))
+				if (const FSceneOutlinerTreeItemPtr* ParentItem = Items.Find(ParentActor))
 				{
 					return *ParentItem;
 				}
 			}
 			else if (Mode->ShouldShowFolders() && !ActorTreeItem->Actor->GetFolderPath().IsNone())
 			{
-				if (ParentItem = Items.Find(ActorTreeItem->Actor->GetFolderPath()))
+				if (const FSceneOutlinerTreeItemPtr* ParentItem = Items.Find(ActorTreeItem->Actor->GetFolderPath()))
 				{
 					return *ParentItem;
 				}
 			}
-			else if (ParentItem = Items.Find(ActorTreeItem->Actor->GetWorld()))
+			else if (const FSceneOutlinerTreeItemPtr* ParentItem = Items.Find(ActorTreeItem->Actor->GetWorld()))
 			{
 				return *ParentItem;
 			}
diff --git a/Engine/Source/Editor/SceneOutliner/Private/ActorModeInteractive.cpp b/Engine/Source/Editor/SceneOutliner/Private/ActorModeInteractive.cpp
--- a/Engine/Source/Editor/SceneOutliner/Private/ActorModeInteractive.cpp
+++ b/Engine/Source/Editor/SceneOutliner/Private/ActorModeInteractive.cpp
@@ -6,7 +6,7 @@
 #include "ActorTreeItem.h"
 
 FActorModeInteractive::FActorModeInteractive(SSceneOutliner* InSceneOutliner, bool bInHideComponents, TWeakObjectPtr<UWorld> InSpecifiedWorldToDisplay)
-	: FActorMode(InSceneOutliner, bHideComponents, SpecifiedWorldToDisplay)
+	: FActorMode(InSceneOutliner, bInHideComponents, SpecifiedWorldToDisplay)
 {
 	USelection::SelectionChangedEvent.AddRaw(this, &FActorModeInteractive::OnLevelSelectionChanged);
 	USelection::SelectObjectEvent.AddRaw(this, &FActorModeInteractive::OnLevelSelectionChanged);
