Fix macOS build.

#rnx

diff --git a/Engine/Source/Editor/SceneOutliner/Private/ActorModeInteractive.cpp b/Engine/Source/Editor/SceneOutliner/Private/ActorModeInteractive.cpp
--- a/Engine/Source/Editor/SceneOutliner/Private/ActorModeInteractive.cpp
+++ b/Engine/Source/Editor/SceneOutliner/Private/ActorModeInteractive.cpp
@@ -6,7 +6,7 @@
 #include "ActorTreeItem.h"
 
 FActorModeInteractive::FActorModeInteractive(SSceneOutliner* InSceneOutliner, bool bInHideComponents, TWeakObjectPtr<UWorld> InSpecifiedWorldToDisplay)
-	: FActorMode(InSceneOutliner, bInHideComponents, SpecifiedWorldToDisplay)
+	: FActorMode(InSceneOutliner, bInHideComponents, InSpecifiedWorldToDisplay)
 {
 	USelection::SelectionChangedEvent.AddRaw(this, &FActorModeInteractive::OnLevelSelectionChanged);
 	USelection::SelectObjectEvent.AddRaw(this, &FActorModeInteractive::OnLevelSelectionChanged);
