Fixed crash in Scene Outliner when Actor::GetComponents returns a null value

diff --git a/Engine/Source/Editor/SceneOutliner/Private/ActorHierarchy.cpp b/Engine/Source/Editor/SceneOutliner/Private/ActorHierarchy.cpp
--- a/Engine/Source/Editor/SceneOutliner/Private/ActorHierarchy.cpp
+++ b/Engine/Source/Editor/SceneOutliner/Private/ActorHierarchy.cpp
@@ -214,14 +214,18 @@ FSceneOutlinerTreeItemPtr FActorHierarchy::FindParent(const ISceneOutlinerTreeIt
 
 void FActorHierarchy::CreateComponentItems(const AActor* Actor, TArray<FSceneOutlinerTreeItemPtr>& OutItems) const
 {
+	check(Actor);
 	// Add all this actors components if showing components and the owning actor was created
 	if (bShowingComponents)
 	{
 		for (UActorComponent* Component : Actor->GetComponents())
 		{
-			if (FSceneOutlinerTreeItemPtr ComponentItem = Mode->CreateItemFor<FComponentTreeItem>(Component))
+			if (Component != nullptr)
 			{
-				OutItems.Add(ComponentItem);
+				if (FSceneOutlinerTreeItemPtr ComponentItem = Mode->CreateItemFor<FComponentTreeItem>(Component))
+				{
+					OutItems.Add(ComponentItem);
+				}
 			}
 		}
 	}
