Foundations in the Scene Outliner

diff --git a/Engine/Config/BaseEditorPerProjectUserSettings.ini b/Engine/Config/BaseEditorPerProjectUserSettings.ini
--- a/Engine/Config/BaseEditorPerProjectUserSettings.ini
+++ b/Engine/Config/BaseEditorPerProjectUserSettings.ini
@@ -123,6 +123,8 @@ bShowOnlyActorsInCurrentLevel=False
 bShowOnlySelectedActors=False
 ; Hide actor components
 bHideActorComponents=True
+; Hide foundation hierarchy
+bHideFoundationHierarchy=False
 
 
 [/Script/UnrealEd.ClassViewerSettings]
diff --git a/Engine/Source/Editor/SceneOutliner/Private/ActorBrowsingMode.cpp b/Engine/Source/Editor/SceneOutliner/Private/ActorBrowsingMode.cpp
--- a/Engine/Source/Editor/SceneOutliner/Private/ActorBrowsingMode.cpp
+++ b/Engine/Source/Editor/SceneOutliner/Private/ActorBrowsingMode.cpp
@@ -26,6 +26,9 @@
 #include "ComponentTreeItem.h"
 #include "ActorBrowsingModeSettings.h"
 #include "ScopedTransaction.h"
+#include "Foundation/FoundationActor.h"
+#include "Foundation/FoundationSubsystem.h"
+#include "Foundation/FoundationEditorInstanceActor.h"
 
 DEFINE_LOG_CATEGORY_STATIC(LogActorBrowser, Log, All);
 
@@ -38,7 +41,7 @@ UActorBrowsingModeSettings::UActorBrowsingModeSettings(const FObjectInitializer&
 using FActorFilter = TSceneOutlinerPredicateFilter<FActorTreeItem>;
 
 FActorBrowsingMode::FActorBrowsingMode(SSceneOutliner* InSceneOutliner, TWeakObjectPtr<UWorld> InSpecifiedWorldToDisplay)
-	: FActorModeInteractive(InSceneOutliner, true, InSpecifiedWorldToDisplay)
+	: FActorModeInteractive(FActorModeParams(InSceneOutliner, InSpecifiedWorldToDisplay, true, false))
 	, FilteredActorCount(0)
 {
 	// Capture selection changes of bones from mesh selection in fracture tools
@@ -87,7 +90,7 @@ FActorBrowsingMode::FActorBrowsingMode(SSceneOutliner* InSceneOutliner, TWeakObj
 	FilterInfoMap.Add(TEXT("ShowOnlyCurrentLevel"), OnlyCurrentLevelInfo);
 
 	bHideComponents = SharedSettings->bHideActorComponents;
-	FSceneOutlinerFilterInfo HideComponentsInfo(LOCTEXT("ToggleHideActorComponents", "Hide Actor Components"), LOCTEXT("ToggleHideActorComponentsToolTip", "When Enabled, hides components belonging to actors."), SharedSettings->bHideActorComponents, FCreateSceneOutlinerFilter::CreateStatic(&FActorBrowsingMode::CreateHideComponentsFilter));
+	FSceneOutlinerFilterInfo HideComponentsInfo(LOCTEXT("ToggleHideActorComponents", "Hide Actor Components"), LOCTEXT("ToggleHideActorComponentsToolTip", "When enabled, hides components belonging to actors."), SharedSettings->bHideActorComponents, FCreateSceneOutlinerFilter::CreateStatic(&FActorBrowsingMode::CreateHideComponentsFilter));
 	HideComponentsInfo.OnToggle().AddLambda([this](bool bIsActive)
 		{
 			UActorBrowsingModeSettings* Settings = GetMutableDefault<UActorBrowsingModeSettings>();
@@ -102,6 +105,41 @@ FActorBrowsingMode::FActorBrowsingMode(SSceneOutliner* InSceneOutliner, TWeakObj
 
 	FilterInfoMap.Add(TEXT("HideComponentsFilter"), HideComponentsInfo);
 
+	FSceneOutlinerFilterInfo HideFoundationsInfo(LOCTEXT("ToggleHideFoundations", "Hide Foundations"), LOCTEXT("ToggleHideFoundationsToolTip", "When enabled, hides all foundation content."), SharedSettings->bHideFoundationHierarchy, FCreateSceneOutlinerFilter::CreateStatic(&FActorBrowsingMode::CreateHideFoundationsFilter));
+	HideFoundationsInfo.OnToggle().AddLambda([this](bool bIsActive)
+		{
+			UActorBrowsingModeSettings* Settings = GetMutableDefault<UActorBrowsingModeSettings>();
+			Settings->bHideFoundationHierarchy = bHideFoundationHierarchy = bIsActive;
+			Settings->PostEditChange();
+
+			if (auto ActorHierarchy = StaticCast<FActorHierarchy*>(Hierarchy.Get()))
+			{
+				ActorHierarchy->SetShowingFoundations(!bIsActive);
+			}
+		});
+	FilterInfoMap.Add(TEXT("HideFoundationsFilter"), HideFoundationsInfo);
+
+	// Add a filter which sets the interactive mode of foundation items and their children
+	SceneOutliner->AddFilter(MakeShared<FActorFilter>(FActorTreeItem::FFilterPredicate::CreateStatic([](const AActor* Actor) {return true; }), FSceneOutlinerFilter::EDefaultBehaviour::Pass, FActorTreeItem::FFilterPredicate::CreateLambda([this](const AActor* Actor)
+		{
+			if (!bHideFoundationHierarchy)
+			{
+				if (const UFoundationSubsystem* FoundationSubsystem = RepresentingWorld->GetSubsystem<UFoundationSubsystem>())
+				{
+					// if actor has a valid parent and the parent is not being edited,
+					// then the actor should not be selectable.
+					if (const AFoundationActor* ParentFoundation = FoundationSubsystem->GetParentFoundation(Actor))
+					{
+						if (!FoundationSubsystem->IsEditingFoundation(ParentFoundation))
+						{
+							return false;
+						}
+					}
+				}
+			}
+			return true;
+		})));
+
 	Rebuild();
 }
 
@@ -227,10 +265,28 @@ TSharedRef<FSceneOutlinerFilter> FActorBrowsingMode::CreateIsInCurrentLevelFilte
 
 TSharedRef<FSceneOutlinerFilter> FActorBrowsingMode::CreateHideComponentsFilter()
 {
-	return MakeShared<TSceneOutlinerPredicateFilter<FComponentTreeItem>>(TSceneOutlinerPredicateFilter<FComponentTreeItem>(FComponentTreeItem::FFilterPredicate::CreateStatic([](const UActorComponent* Component)
+	return MakeShared<TSceneOutlinerPredicateFilter<FComponentTreeItem>>(TSceneOutlinerPredicateFilter<FComponentTreeItem>(
+		FComponentTreeItem::FFilterPredicate::CreateStatic([](const UActorComponent*) { return false; }), 
+		FSceneOutlinerFilter::EDefaultBehaviour::Pass));
+}
+
+TSharedRef<FSceneOutlinerFilter> FActorBrowsingMode::CreateHideFoundationsFilter()
+{
+	return MakeShareable(new FActorFilter(FActorTreeItem::FFilterPredicate::CreateStatic([](const AActor* Actor)
 		{
-			// Filter out all components
-			return false;
+			// Check if actor belongs to a foundation
+			if (const UFoundationSubsystem* FoundationSubsystem = Actor->GetWorld()->GetSubsystem<UFoundationSubsystem>())
+			{
+				if (const AFoundationActor* ParentFoundation = FoundationSubsystem->GetParentFoundation(Actor))
+				{
+					if (!FoundationSubsystem->IsEditingFoundation(ParentFoundation))
+					{
+						return false;
+					}
+				}
+			}
+			// Or if the actor itself is a foundation editor instance
+			return Cast<AFoundationEditorInstanceActor>(Actor) == nullptr;
 		}), FSceneOutlinerFilter::EDefaultBehaviour::Pass));
 }
 
@@ -493,33 +549,42 @@ void FActorBrowsingMode::OnItemSelectionChanged(FSceneOutlinerTreeItemPtr TreeIt
 
 void FActorBrowsingMode::OnItemDoubleClick(FSceneOutlinerTreeItemPtr Item)
 {
-	if (Item->CanInteract())
+	if (const FActorTreeItem* ActorItem = Item->CastTo<FActorTreeItem>())
 	{
-		if (Item->IsA<FActorTreeItem>())
+		AActor* Actor = ActorItem->Actor.Get();
+		check(Actor);
+
+		AFoundationActor* FoundationActor = Cast<AFoundationActor>(Actor);
+		if (FoundationActor && FSlateApplication::Get().GetModifierKeys().IsAltDown())
 		{
-			FSceneOutlinerItemSelection Selection(SceneOutliner->GetSelection());
-			if (Selection.Has<FActorTreeItem>())
+			if (FoundationActor->CanEdit())
 			{
-				const bool bActiveViewportOnly = false;
-				GEditor->MoveViewportCamerasToActor(Selection.GetData<AActor*>(SceneOutliner::FActorSelector()), bActiveViewportOnly);
+				FoundationActor->Edit();
+			}
+			else if (FoundationActor->CanCommit())
+			{
+				FoundationActor->Commit();
 			}
 		}
-	}
-	else
-	{
-		if (const FActorTreeItem* ActorItem = Item->CastTo<FActorTreeItem>())
+		else if (Item->CanInteract())
 		{
-			if (AActor* Actor = ActorItem->Actor.Get())
+			FSceneOutlinerItemSelection Selection(SceneOutliner->GetSelection());
+			if (Selection.Has<FActorTreeItem>())
 			{
 				const bool bActiveViewportOnly = false;
-				GEditor->MoveViewportCamerasToActor(*Actor, bActiveViewportOnly);
+				GEditor->MoveViewportCamerasToActor(Selection.GetData<AActor*>(SceneOutliner::FActorSelector()), bActiveViewportOnly);
 			}
 		}
-		else if (Item->IsA<FFolderTreeItem>())
+		else
 		{
-			SceneOutliner->SetItemExpansion(Item, !SceneOutliner->IsItemExpanded(Item));
+			const bool bActiveViewportOnly = false;
+			GEditor->MoveViewportCamerasToActor(*Actor, bActiveViewportOnly);
 		}
 	}
+	else if (Item->IsA<FFolderTreeItem>())
+	{
+		SceneOutliner->SetItemExpansion(Item, !SceneOutliner->IsItemExpanded(Item));
+	}
 }
 
 void FActorBrowsingMode::OnFilterTextCommited(FSceneOutlinerItemSelection& Selection, ETextCommit::Type CommitType)
@@ -639,7 +704,6 @@ bool FActorBrowsingMode::CanPaste() const
 	return CanPasteFoldersOnlyFromClipboard();
 }
 
-
 bool FActorBrowsingMode::CanPasteFoldersOnlyFromClipboard() const
 {
 	// Intentionally not checking if the level is locked/hidden here, as it's better feedback for the user if they attempt to paste
@@ -721,10 +785,21 @@ FSceneOutlinerDragValidationInfo FActorBrowsingMode::ValidateDrop(const ISceneOu
 			return FSceneOutlinerDragValidationInfo(ESceneOutlinerDropCompatibility::IncompatibleGeneric, FText());
 		}
 
+		const AFoundationActor* FoundationTarget = Cast<AFoundationActor>(ActorTarget);
+		const UFoundationSubsystem* FoundationSubsystem = RepresentingWorld->GetSubsystem<UFoundationSubsystem>();
+
+		if (FoundationTarget)
+		{
+			check(FoundationSubsystem);
+			if (!FoundationTarget->IsEditing())
+			{
+				return FSceneOutlinerDragValidationInfo(ESceneOutlinerDropCompatibility::IncompatibleGeneric, LOCTEXT("Error_AttachToClosedFoundation", "Cannot attach to foundation which is not being edited"));
+			}
+		}
+
 		FText AttachErrorMsg;
 		bool bCanAttach = true;
 		bool bDraggedOntoAttachmentParent = true;
-
 		const auto& DragActors = Payload.GetData<TWeakObjectPtr<AActor>>(SceneOutliner::FWeakActorSelector());
 		for (const auto& DragActorPtr : DragActors)
 		{
@@ -733,13 +808,32 @@ FSceneOutlinerDragValidationInfo FActorBrowsingMode::ValidateDrop(const ISceneOu
 			{
 				if (bCanAttach)
 				{
+					if (FoundationSubsystem)
+					{
+						// Either all actors must be in a foundation or none of them
+						if (const AFoundationActor* ParentFoundation = FoundationSubsystem->GetParentFoundation(DragActor))
+						{
+							if (!ParentFoundation->IsEditing())
+							{
+								return FSceneOutlinerDragValidationInfo(ESceneOutlinerDropCompatibility::IncompatibleGeneric, LOCTEXT("Error_RemoveEditingFoundation", "Cannot detach from a foundation which is not being edited"));
+							}
+						}
+
+						if (!FoundationSubsystem->CanMoveActorToLevel(DragActor))
+						{
+							AttachErrorMsg = LOCTEXT("Error_MoveActorToFoundation", "Cannot move foundation while it or its children are being edited");
+							bCanAttach = bDraggedOntoAttachmentParent = false;
+							break;
+						}
+					}
+
 					if (DragActor->IsChildActor())
 					{
 						AttachErrorMsg = FText::Format(LOCTEXT("Error_AttachChildActor", "Cannot move {0} as it is a child actor."), FText::FromString(DragActor->GetActorLabel()));
 						bCanAttach = bDraggedOntoAttachmentParent = false;
 						break;
 					}
-					if (!GEditor->CanParentActors(ActorTarget, DragActor, &AttachErrorMsg))
+					if (!FoundationTarget && !GEditor->CanParentActors(ActorTarget, DragActor, &AttachErrorMsg))
 					{
 						bCanAttach = false;
 					}
@@ -847,15 +941,38 @@ FSceneOutlinerDragValidationInfo FActorBrowsingMode::ValidateDrop(const ISceneOu
 
 		if (Payload.Has<FActorTreeItem>())
 		{
-			// Iterate over all the folders that have been dragged
+			const UFoundationSubsystem* FoundationSubsystem = RepresentingWorld->GetSubsystem<UFoundationSubsystem>();
+			// Iterate over all the actors that have been dragged
 			for (const TWeakObjectPtr<AActor>& WeakActor : Payload.GetData<TWeakObjectPtr<AActor>>(SceneOutliner::FWeakActorSelector()))
 			{
 				const AActor* Actor = WeakActor.Get();
+
+				bool bActorContainedInFoundation = false;
+				if (FoundationSubsystem)
+				{
+					if (const AFoundationActor* ParentFoundation = FoundationSubsystem->GetParentFoundation(Actor))
+					{
+						if (!ParentFoundation->IsEditing())
+						{
+							return FSceneOutlinerDragValidationInfo(ESceneOutlinerDropCompatibility::IncompatibleGeneric, LOCTEXT("Error_RemoveEditingFoundation", "Cannot detach from a foundation which is not being edited"));
+						}
+						bActorContainedInFoundation = true;
+					}
+
+					if (const AFoundationActor* FoundationActor = Cast<AFoundationActor>(Actor))
+					{
+						if (!FoundationSubsystem->CanMoveActorToLevel(FoundationActor))
+						{
+							return FSceneOutlinerDragValidationInfo(ESceneOutlinerDropCompatibility::IncompatibleGeneric, LOCTEXT("Error_RemoveEditingFoundation", "Cannot detach a foundation which is currently being edited"));
+						}
+					}
+				}
+				
 				if (Actor->IsChildActor())
 				{
 					return FSceneOutlinerDragValidationInfo(ESceneOutlinerDropCompatibility::IncompatibleGeneric, FText::Format(LOCTEXT("Error_AttachChildActor", "Cannot move {0} as it is a child actor."), FText::FromString(Actor->GetActorLabel())));
 				}
-				else if (Actor->GetFolderPath() == DestinationPath && !Actor->GetAttachParentActor())
+				else if (Actor->GetFolderPath() == DestinationPath && !Actor->GetAttachParentActor() && !bActorContainedInFoundation)
 				{
 					FFormatNamedArguments Args;
 					Args.Add(TEXT("SourceName"), FText::FromString(Actor->GetActorLabel()));
@@ -906,8 +1023,6 @@ void FActorBrowsingMode::OnDrop(ISceneOutlinerTreeItem& DropTarget, const FScene
 			return;
 		}
 
-		TArray<TWeakObjectPtr<AActor>> DraggedActors = Payload.GetData<TWeakObjectPtr<AActor>>(SceneOutliner::FWeakActorSelector());
-
 		FMessageLog EditorErrors("EditorErrors");
 		EditorErrors.NewPage(LOCTEXT("ActorAttachmentsPageLabel", "Actor attachment"));
 
@@ -915,6 +1030,7 @@ void FActorBrowsingMode::OnDrop(ISceneOutlinerTreeItem& DropTarget, const FScene
 		{
 			const FScopedTransaction Transaction(LOCTEXT("UndoAction_DetachActors", "Detach actors"));
 
+			TArray<TWeakObjectPtr<AActor>> DraggedActors = Payload.GetData<TWeakObjectPtr<AActor>>(SceneOutliner::FWeakActorSelector());
 			for (const auto& WeakActor : DraggedActors)
 			{
 				if (auto* DragActor = WeakActor.Get())
@@ -936,51 +1052,72 @@ void FActorBrowsingMode::OnDrop(ISceneOutlinerTreeItem& DropTarget, const FScene
 		{
 			// Show socket chooser if we have sockets to select
 
-			auto PerformAttachment = [](FName SocketName, TWeakObjectPtr<AActor> Parent, const TArray<TWeakObjectPtr<AActor>> NewAttachments)
+			if (AFoundationActor* TargetFoundation = Cast<AFoundationActor>(DropActor))
 			{
-				AActor* ParentActor = Parent.Get();
-				if (ParentActor)
+				// Actors inside foundations cannot have folder paths
+				TArray<AActor*> DraggedActors = Payload.GetData<AActor*>(SceneOutliner::FActorSelector());
+				for (auto& Actor : DraggedActors)
 				{
-					// modify parent and child
-					const FScopedTransaction Transaction(LOCTEXT("UndoAction_PerformAttachment", "Attach actors"));
+					Actor->SetFolderPath_Recursively(FName());
+				}
+
+				UFoundationSubsystem* FoundationSubsystem = RepresentingWorld->GetSubsystem<UFoundationSubsystem>();
+				check(FoundationSubsystem);
 
-					// Attach each child
-					bool bAttached = false;
-					for (auto& Child : NewAttachments)
+				check(TargetFoundation->IsEditing());
+				const FScopedTransaction Transaction(LOCTEXT("UndoAction_MoveActorsToFoundation", "Move actors to foundation"));
+
+				FoundationSubsystem->MoveActorsTo(TargetFoundation, DraggedActors);
+			}
+			else
+			{
+				auto PerformAttachment = [](FName SocketName, TWeakObjectPtr<AActor> Parent, const TArray<TWeakObjectPtr<AActor>> NewAttachments)
+				{
+					AActor* ParentActor = Parent.Get();
+					if (ParentActor)
 					{
-						AActor* ChildActor = Child.Get();
-						if (GEditor->CanParentActors(ParentActor, ChildActor))
+						// modify parent and child
+						const FScopedTransaction Transaction(LOCTEXT("UndoAction_PerformAttachment", "Attach actors"));
+
+						// Attach each child
+						bool bAttached = false;
+						for (auto& Child : NewAttachments)
 						{
-							GEditor->ParentActors(ParentActor, ChildActor, SocketName);
+							AActor* ChildActor = Child.Get();
+							if (GEditor->CanParentActors(ParentActor, ChildActor))
+							{
+								GEditor->ParentActors(ParentActor, ChildActor, SocketName);
 
-							ChildActor->SetFolderPath_Recursively(ParentActor->GetFolderPath());
+								ChildActor->SetFolderPath_Recursively(ParentActor->GetFolderPath());
+							}
 						}
 					}
-				}
-			};
+				};
 
-			//@TODO: Should create a menu for each component that contains sockets, or have some form of disambiguation within the menu (like a fully qualified path)
-			// Instead, we currently only display the sockets on the root component
-			USceneComponent* Component = DropActor->GetRootComponent();
-			if ((Component != NULL) && (Component->HasAnySockets()))
-			{
-				// Create the popup
-				FSlateApplication::Get().PushMenu(
-					SceneOutliner->AsShared(),
-					FWidgetPath(),
-					SNew(SSocketChooserPopup)
-					.SceneComponent(Component)
-					.OnSocketChosen_Lambda(PerformAttachment, DropActor, MoveTemp(DraggedActors)),
-					FSlateApplication::Get().GetCursorPos(),
-					FPopupTransitionEffect(FPopupTransitionEffect::TypeInPopup)
-				);
-			}
-			else
-			{
-				PerformAttachment(NAME_None, DropActor, MoveTemp(DraggedActors));
+				TArray<TWeakObjectPtr<AActor>> DraggedActors = Payload.GetData<TWeakObjectPtr<AActor>>(SceneOutliner::FWeakActorSelector());
+				//@TODO: Should create a menu for each component that contains sockets, or have some form of disambiguation within the menu (like a fully qualified path)
+				// Instead, we currently only display the sockets on the root component
+				USceneComponent* Component = DropActor->GetRootComponent();
+				if ((Component != NULL) && (Component->HasAnySockets()))
+				{
+					// Create the popup
+					FSlateApplication::Get().PushMenu(
+						SceneOutliner->AsShared(),
+						FWidgetPath(),
+						SNew(SSocketChooserPopup)
+						.SceneComponent(Component)
+						.OnSocketChosen_Lambda(PerformAttachment, DropActor, MoveTemp(DraggedActors)),
+						FSlateApplication::Get().GetCursorPos(),
+						FPopupTransitionEffect(FPopupTransitionEffect::TypeInPopup)
+					);
+				}
+				else
+				{
+					PerformAttachment(NAME_None, DropActor, MoveTemp(DraggedActors));
+				}
 			}
+			
 		}
-
 		// Report errors
 		EditorErrors.Notify(NSLOCTEXT("ActorAttachmentError", "AttachmentsFailed", "Attachments Failed!"));
 	}
@@ -991,7 +1128,7 @@ void FActorBrowsingMode::OnDrop(ISceneOutlinerTreeItem& DropTarget, const FScene
 		const FName& DestinationPath = FolderItem ? FolderItem->Path : FName(NAME_None);
 
 		const FScopedTransaction Transaction(LOCTEXT("MoveOutlinerItems", "Move World Outliner Items"));
-			
+
 		auto MoveToDestination = [&DestinationPath](FFolderTreeItem& Item)
 		{
 			Item.MoveTo(DestinationPath);
@@ -1004,22 +1141,22 @@ void FActorBrowsingMode::OnDrop(ISceneOutlinerTreeItem& DropTarget, const FScene
 			TSet<const AActor*> ParentActors;
 			TSet<const AActor*> ChildActors;
 
-			for (const auto& WeakActor : Payload.GetData<TWeakObjectPtr<AActor>>(SceneOutliner::FWeakActorSelector()))
-			{
-				AActor* Actor = WeakActor.Get();
-				if (Actor)
+			Payload.ForEachItem<FActorTreeItem>([&DestinationPath, &ParentActors, &ChildActors](const FActorTreeItem& ActorItem)
 				{
-					// First mark this object as a parent, then set its children's path
-					ParentActors.Add(Actor);
-					Actor->SetFolderPath(DestinationPath);
-
-					FActorEditorUtils::TraverseActorTree_ParentFirst(Actor, [&](AActor* InActor) {
-						ChildActors.Add(InActor);
-						InActor->SetFolderPath(DestinationPath);
-						return true;
-						}, false);
-				}
-			}
+					AActor* Actor = ActorItem.Actor.Get();
+					if (Actor)
+					{
+						// First mark this object as a parent, then set its children's path
+						ParentActors.Add(Actor);
+						Actor->SetFolderPath(DestinationPath);
+
+						FActorEditorUtils::TraverseActorTree_ParentFirst(Actor, [&](AActor* InActor) {
+							ChildActors.Add(InActor);
+							InActor->SetFolderPath(DestinationPath);
+							return true;
+							}, false);
+					}
+				});
 
 			// Detach parent actors
 			for (const AActor* Parent : ParentActors)
@@ -1036,6 +1173,26 @@ void FActorBrowsingMode::OnDrop(ISceneOutlinerTreeItem& DropTarget, const FScene
 					RootComp->DetachFromComponent(FDetachmentTransformRules::KeepWorldTransform);
 				}
 			}
+
+			const UFoundationSubsystem* FoundationSubsystem = RepresentingWorld->GetSubsystem<UFoundationSubsystem>();
+			check(FoundationSubsystem);
+			// Since we are moving to a folder (or root), we must be moving into the persistent level.
+			ULevel* DestinationLevel = RepresentingWorld->PersistentLevel;
+			check(DestinationLevel);
+
+			TArray<AActor*> ActorsToMove;
+			Payload.ForEachItem<FActorTreeItem>([FoundationSubsystem, &ActorsToMove](const FActorTreeItem& ActorItem)
+				{
+					AActor* Actor = ActorItem.Actor.Get();
+					if (const AFoundationActor* ParentFoundation = FoundationSubsystem->GetParentFoundation(Actor))
+					{
+						check(ParentFoundation->IsEditing());
+						ActorsToMove.Add(Actor);
+					}
+				});
+
+			TArray<AActor*> DraggedActors = Payload.GetData<AActor*>(SceneOutliner::FActorSelector());
+			FoundationSubsystem->MoveActorsToLevel(ActorsToMove, DestinationLevel);
 		}
 	}
 }
diff --git a/Engine/Source/Editor/SceneOutliner/Private/ActorBrowsingModeSettings.h b/Engine/Source/Editor/SceneOutliner/Private/ActorBrowsingModeSettings.h
--- a/Engine/Source/Editor/SceneOutliner/Private/ActorBrowsingModeSettings.h
+++ b/Engine/Source/Editor/SceneOutliner/Private/ActorBrowsingModeSettings.h
@@ -29,7 +29,11 @@ class UActorBrowsingModeSettings
 	UPROPERTY(config)
 	uint32 bShowOnlySelectedActors : 1;
 
-	/** True when the Scene Outliner is only displaying Actor Components*/
+	/** True when the Scene Outliner is not displaying Actor Components*/
 	UPROPERTY(config)
 	uint32 bHideActorComponents : 1;
+
+	/** True when the Scene Outliner is not displaying foundations */
+	UPROPERTY(config)
+	uint32 bHideFoundationHierarchy : 1;
 };
diff --git a/Engine/Source/Editor/SceneOutliner/Private/ActorFolderPickingMode.cpp b/Engine/Source/Editor/SceneOutliner/Private/ActorFolderPickingMode.cpp
--- a/Engine/Source/Editor/SceneOutliner/Private/ActorFolderPickingMode.cpp
+++ b/Engine/Source/Editor/SceneOutliner/Private/ActorFolderPickingMode.cpp
@@ -7,7 +7,7 @@
 #define LOCTEXT_NAMESPACE "SceneOutliner_ActoFolderPickingMode"
 
 FActorFolderPickingMode::FActorFolderPickingMode(SSceneOutliner* InSceneOutliner, FOnSceneOutlinerItemPicked InOnItemPicked, TWeakObjectPtr<UWorld> InSpecifiedWorldToDisplay)
-	: FActorMode(InSceneOutliner, true, InSpecifiedWorldToDisplay)
+	: FActorMode(FActorModeParams(InSceneOutliner, InSpecifiedWorldToDisplay))
 	, OnItemPicked(InOnItemPicked)
 {
 	Rebuild();
diff --git a/Engine/Source/Editor/SceneOutliner/Private/ActorHierarchy.cpp b/Engine/Source/Editor/SceneOutliner/Private/ActorHierarchy.cpp
--- a/Engine/Source/Editor/SceneOutliner/Private/ActorHierarchy.cpp
+++ b/Engine/Source/Editor/SceneOutliner/Private/ActorHierarchy.cpp
@@ -13,10 +13,12 @@
 #include "LevelUtils.h"
 #include "GameFramework/WorldSettings.h"
 #include "EditorActorFolders.h"
+#include "Foundation/FoundationActor.h"
+#include "Foundation/FoundationSubsystem.h"
 
-TUniquePtr<FActorHierarchy> FActorHierarchy::Create(ISceneOutlinerMode* Mode, const TWeakObjectPtr<UWorld>& World, TFunction<bool(const AActor*)> InIsActorDisplayable)
+TUniquePtr<FActorHierarchy> FActorHierarchy::Create(ISceneOutlinerMode* Mode, const TWeakObjectPtr<UWorld>& World)
 {
-	FActorHierarchy* Hierarchy = new FActorHierarchy(Mode, World, InIsActorDisplayable);
+	FActorHierarchy* Hierarchy = new FActorHierarchy(Mode, World);
 		
 	GEngine->OnLevelActorAdded().AddRaw(Hierarchy, &FActorHierarchy::OnLevelActorAdded);
 	GEngine->OnLevelActorDeleted().AddRaw(Hierarchy, &FActorHierarchy::OnLevelActorDeleted);
@@ -36,10 +38,9 @@ TUniquePtr<FActorHierarchy> FActorHierarchy::Create(ISceneOutlinerMode* Mode, co
 	return TUniquePtr<FActorHierarchy>(Hierarchy);
 }
 
-FActorHierarchy::FActorHierarchy(ISceneOutlinerMode* Mode, const TWeakObjectPtr<UWorld>& World, TFunction<bool(const AActor*)> InIsActorDisplayable)
+FActorHierarchy::FActorHierarchy(ISceneOutlinerMode* Mode, const TWeakObjectPtr<UWorld>& World)
 	: ISceneOutlinerHierarchy(Mode)
 	, RepresentingWorld(World)
-	, IsActorDisplayableCallback(InIsActorDisplayable)
 {
 }
 
@@ -70,6 +71,10 @@ FActorHierarchy::~FActorHierarchy()
 
 void FActorHierarchy::FindChildren(const ISceneOutlinerTreeItem& Item, const TMap<FSceneOutlinerTreeItemID, FSceneOutlinerTreeItemPtr>& Items, TArray<FSceneOutlinerTreeItemPtr>& OutChildren) const
 {
+	UFoundationSubsystem* FoundationSubsystem = RepresentingWorld->GetSubsystem<UFoundationSubsystem>();
+
+	check(!bShowingFoundations || FoundationSubsystem);
+
 	if (const FWorldTreeItem* WorldTreeItem = Item.CastTo<FWorldTreeItem>())
 	{
 		// All actors in the world are children of the world
@@ -105,23 +110,37 @@ void FActorHierarchy::FindChildren(const ISceneOutlinerTreeItem& Item, const TMa
 	}
 	else if (const FActorTreeItem* ActorTreeItem = Item.CastTo<FActorTreeItem>())
 	{
-		// All attached actors are children of this actor
-		TArray<AActor*> AttachedActors;
-		ActorTreeItem->Actor->GetAttachedActors(AttachedActors);
-		for (const AActor* AttachedActor : AttachedActors)
+		if (AFoundationActor* FoundationActor = Cast<AFoundationActor>(ActorTreeItem->Actor))
 		{
-			if (const FSceneOutlinerTreeItemPtr* AttachedActorItem = Items.Find(AttachedActor))
-			{
-				OutChildren.Add(*AttachedActorItem);
-			}
+			FoundationSubsystem->ForEachActorInFoundation(FoundationActor, [&OutChildren, &Items](AActor* Actor)
+				{
+					if (const FSceneOutlinerTreeItemPtr* ContainedActorItem = Items.Find(Actor))
+					{
+						OutChildren.Add(*ContainedActorItem);
+					}
+					return true;
+				});
 		}
-		if (bShowingComponents)
+		else
 		{
-			for (const UActorComponent* Component : ActorTreeItem->Actor->GetComponents())
+			// All attached actors are children of this actor
+			TArray<AActor*> AttachedActors;
+			ActorTreeItem->Actor->GetAttachedActors(AttachedActors);
+			for (const AActor* AttachedActor : AttachedActors)
 			{
-				if (const FSceneOutlinerTreeItemPtr* ComponentItem = Items.Find(Component))
+				if (const FSceneOutlinerTreeItemPtr* AttachedActorItem = Items.Find(AttachedActor))
 				{
-					OutChildren.Add(*ComponentItem);
+					OutChildren.Add(*AttachedActorItem);
+				}
+			}
+			if (bShowingComponents)
+			{
+				for (const UActorComponent* Component : ActorTreeItem->Actor->GetComponents())
+				{
+					if (const FSceneOutlinerTreeItemPtr* ComponentItem = Items.Find(Component))
+					{
+						OutChildren.Add(*ComponentItem);
+					}
 				}
 			}
 		}
@@ -153,23 +172,50 @@ FSceneOutlinerTreeItemPtr FActorHierarchy::FindParent(const ISceneOutlinerTreeIt
 	}
 	else if (const FActorTreeItem* ActorTreeItem = Item.CastTo<FActorTreeItem>())
 	{
-		if (ActorTreeItem->Actor.IsValid())
+		if (AActor* Actor = ActorTreeItem->Actor.Get())
 		{
-			if (const AActor* ParentActor = ActorTreeItem->Actor->GetAttachParentActor())
+			if (const AActor* ParentActor = Actor->GetAttachParentActor())
 			{
 				if (const FSceneOutlinerTreeItemPtr* ParentItem = Items.Find(ParentActor))
 				{
 					return *ParentItem;
 				}
 			}
-			else if (Mode->ShouldShowFolders() && !ActorTreeItem->Actor->GetFolderPath().IsNone())
+			else if (Mode->ShouldShowFolders() && !Actor->GetFolderPath().IsNone())
 			{
-				if (const FSceneOutlinerTreeItemPtr* ParentItem = Items.Find(ActorTreeItem->Actor->GetFolderPath()))
+				if (const FSceneOutlinerTreeItemPtr* ParentItem = Items.Find(Actor->GetFolderPath()))
 				{
 					return *ParentItem;
 				}
+				else
+				{
+					return nullptr;
+				}
 			}
-			else if (const FSceneOutlinerTreeItemPtr* ParentItem = Items.Find(ActorTreeItem->Actor->GetWorld()))
+
+			if (const UFoundationSubsystem* FoundationSubsystem = RepresentingWorld->GetSubsystem<UFoundationSubsystem>())
+			{
+				if (const AFoundationActor* OwningFoundation = FoundationSubsystem->GetParentFoundation(Actor))
+				{
+					const AFoundationActor* FoundationActor = Cast<AFoundationActor>(Actor);
+					const bool bIsAnEditingFoundation = FoundationActor ? FoundationActor->IsEditing() : false;
+					// Parent this to a foundation if the parent foundation is being edited or if this is a sub foundation which is being edited
+					if (bShowingFoundations || (OwningFoundation->IsEditing() || bIsAnEditingFoundation))
+					{
+						if (const FSceneOutlinerTreeItemPtr* OwningFoundationItem = Items.Find(OwningFoundation))
+						{
+							return *OwningFoundationItem;
+						}
+						else
+						{
+							return nullptr;
+						}
+					}
+				}
+			}
+
+			// Default to the world
+			if (const FSceneOutlinerTreeItemPtr* ParentItem = Items.Find(ActorTreeItem->Actor->GetWorld()))
 			{
 				return *ParentItem;
 			}
@@ -233,19 +279,32 @@ void FActorHierarchy::CreateComponentItems(const AActor* Actor, TArray<FSceneOut
 
 void FActorHierarchy::CreateWorldChildren(UWorld* World, TArray<FSceneOutlinerTreeItemPtr>& OutItems) const
 {
+	check(World);
+
+	const UFoundationSubsystem* FoundationSubsystem = World->GetSubsystem<UFoundationSubsystem>();
 	// Create all actor items
 	for (FActorIterator ActorIt(World); ActorIt; ++ActorIt)
 	{
-		if (IsActorDisplayableCallback(*ActorIt))
+		AActor* Actor = *ActorIt;
+		// If we are not showing foundations, foundation sub actor items should not be created unless they belong to a foundation which is being edited
+		if (FoundationSubsystem)
 		{
-			if (FSceneOutlinerTreeItemPtr ActorItem = Mode->CreateItemFor<FActorTreeItem>(*ActorIt))
-			{
-				OutItems.Add(ActorItem);
-
-				// Create all component items
-				CreateComponentItems(*ActorIt, OutItems);
+			if (const AFoundationActor* ParentFoundation = FoundationSubsystem->GetParentFoundation(Actor))
+		{
+				if (!bShowingFoundations && !ParentFoundation->IsEditing())
+				{
+					continue;
+				}
 			}
 		}
+
+		if (FSceneOutlinerTreeItemPtr ActorItem = Mode->CreateItemFor<FActorTreeItem>(Actor))
+		{
+			OutItems.Add(ActorItem);
+
+			// Create all component items
+			CreateComponentItems(Actor, OutItems);
+		}
 	}
 
 	// Create all folder items
@@ -283,13 +342,31 @@ void FActorHierarchy::CreateChildren(const FSceneOutlinerTreeItemPtr& Item, TArr
 	}
 	else if (const FActorTreeItem* ParentActorItem = Item->CastTo<FActorTreeItem>())
 	{
-		const AActor* ParentActor = ParentActorItem->Actor.Get();
+		AActor* ParentActor = ParentActorItem->Actor.Get();
+		check(ParentActor->GetWorld() == RepresentingWorld);
+
+		CreateComponentItems(ParentActor, OutChildren);
+
+		TArray<AActor*> ChildActors;
 
-		if (IsActorDisplayableCallback(ParentActor))
+		if (const AFoundationActor* FoundationParentActor = Cast<AFoundationActor>(ParentActor))
 		{
-			CreateComponentItems(ParentActor, OutChildren);
+			const UFoundationSubsystem* FoundationSubsystem = RepresentingWorld->GetSubsystem<UFoundationSubsystem>();
+			check(FoundationSubsystem);
 
-			TArray<AActor*> ChildActors;
+			FoundationSubsystem->ForEachActorInFoundation(FoundationParentActor, [this, FoundationParentActor, FoundationSubsystem, &ChildActors](AActor* SubActor)
+				{
+					const AFoundationActor* FoundationActor = Cast<AFoundationActor>(SubActor);
+					const bool bIsAnEditingFoundation = FoundationActor ? FoundationSubsystem->IsEditingFoundation(FoundationActor) : false;
+					if (bShowingFoundations || (FoundationSubsystem->IsEditingFoundation(FoundationParentActor) || bIsAnEditingFoundation))
+					{
+						ChildActors.Add(SubActor);
+					}
+					return true;
+				});
+		}
+		else
+		{
 			TFunction<bool(AActor*)> GetAttachedActors = [&ChildActors, &GetAttachedActors](AActor* Child)
 			{
 				ChildActors.Add(Child);
@@ -301,18 +378,18 @@ void FActorHierarchy::CreateChildren(const FSceneOutlinerTreeItemPtr& Item, TArr
 
 			// Grab all direct/indirect children of an actor
 			ParentActor->ForEachAttachedActors(GetAttachedActors);
+		}
 
-			for (auto ChildActor : ChildActors)
+		for (auto ChildActor : ChildActors)
+		{
+			if (FSceneOutlinerTreeItemPtr ChildActorItem = Mode->CreateItemFor<FActorTreeItem>(ChildActor))
 			{
-				if (FSceneOutlinerTreeItemPtr ChildActorItem = Mode->CreateItemFor<FActorTreeItem>(ChildActor))
-				{
-					OutChildren.Add(ChildActorItem);
+				OutChildren.Add(ChildActorItem);
 
-					CreateComponentItems(ChildActor, OutChildren);
-				}
+				CreateComponentItems(ChildActor, OutChildren);
 			}
 		}
-	}
+		}
 	else if (FActorFolderTreeItem* FolderItem = Item->CastTo<FActorFolderTreeItem>())
 	{
 		check(Mode->ShouldShowFolders());
@@ -338,24 +415,37 @@ FSceneOutlinerTreeItemPtr FActorHierarchy::CreateParentItem(const FSceneOutliner
 	}
 	else if (const FActorTreeItem* ActorTreeItem = Item->CastTo<FActorTreeItem>())
 	{
-		if (ActorTreeItem->Actor.IsValid())
+		if (const AActor* Actor = ActorTreeItem->Actor.Get())
 		{
-			AActor* ParentActor = ActorTreeItem->Actor->GetAttachParentActor();
-			if (ParentActor)
+			if (AActor* ParentActor = Actor->GetAttachParentActor())
 			{
 				return Mode->CreateItemFor<FActorTreeItem>(ParentActor, true);
 			}
-			// if this item does not belong to a folder (or folders are disabled)
-			else if (!Mode->ShouldShowFolders() || ActorTreeItem->Actor->GetFolderPath().IsNone())
+			
+			// if this item belongs in a folder
+			if (Mode->ShouldShowFolders() && !ActorTreeItem->Actor->GetFolderPath().IsNone())
 			{
-				UWorld* OwningWorld = ActorTreeItem->Actor->GetWorld();
-				check(OwningWorld);
-				return Mode->CreateItemFor<FWorldTreeItem>(OwningWorld, true);
+				return Mode->CreateItemFor<FActorFolderTreeItem>(FActorFolderTreeItem(ActorTreeItem->Actor->GetFolderPath(), ActorTreeItem->Actor->GetWorld()), true);
 			}
-			else
+
+			// If item belongs to a foundation
+			if (const UFoundationSubsystem* FoundationSubsystem = RepresentingWorld->GetSubsystem<UFoundationSubsystem>())
 			{
-				return Mode->CreateItemFor<FActorFolderTreeItem>(FActorFolderTreeItem(ActorTreeItem->Actor->GetFolderPath(), ActorTreeItem->Actor->GetWorld()), true);
+				if (AFoundationActor* ParentFoundation = FoundationSubsystem->GetParentFoundation(Actor))
+				{
+					const AFoundationActor* FoundationActor = Cast<AFoundationActor>(Actor);
+					const bool bIsAnEditingFoundation = FoundationActor ? FoundationActor->IsEditing() : false;
+					if (bShowingFoundations || (ParentFoundation->IsEditing() || bIsAnEditingFoundation))
+					{
+						return Mode->CreateItemFor<FActorTreeItem>(ParentFoundation, true);
+					}
+				}
 			}
+
+			// Default to the world
+			UWorld* OwningWorld = ActorTreeItem->Actor->GetWorld();
+			check(OwningWorld);
+			return Mode->CreateItemFor<FWorldTreeItem>(OwningWorld, true);
 		}
 	}
 	else if (const FComponentTreeItem* ComponentTreeItem = Item->CastTo<FComponentTreeItem>())
@@ -394,7 +484,7 @@ void FActorHierarchy::FullRefreshEvent()
 
 void FActorHierarchy::OnLevelActorAdded(AActor* InActor)
 {
-	if (InActor && RepresentingWorld.Get() == InActor->GetWorld() && IsActorDisplayableCallback(InActor))
+	if (InActor && RepresentingWorld.Get() == InActor->GetWorld())
 	{
 		FSceneOutlinerHierarchyChangedData EventData;
 		EventData.Type = FSceneOutlinerHierarchyChangedData::Added;
diff --git a/Engine/Source/Editor/SceneOutliner/Private/ActorMode.cpp b/Engine/Source/Editor/SceneOutliner/Private/ActorMode.cpp
--- a/Engine/Source/Editor/SceneOutliner/Private/ActorMode.cpp
+++ b/Engine/Source/Editor/SceneOutliner/Private/ActorMode.cpp
@@ -14,6 +14,7 @@
 #include "FolderTreeItem.h"
 #include "ComponentTreeItem.h"
 #include "WorldTreeItem.h"
+#include "Foundation/FoundationEditorInstanceActor.h"
 
 #define LOCTEXT_NAMESPACE "SceneOutliner_ActorMode"
 
@@ -66,11 +67,16 @@ namespace SceneOutliner
 	}
 }
 
-FActorMode::FActorMode(SSceneOutliner* InSceneOutliner, bool bInHideComponents, TWeakObjectPtr<UWorld> InSpecifiedWorldToDisplay)
-	: ISceneOutlinerMode(InSceneOutliner)
-	, bHideComponents(bInHideComponents)
-	, SpecifiedWorldToDisplay(InSpecifiedWorldToDisplay)
+FActorMode::FActorMode(const FActorModeParams& Params)
+	: ISceneOutlinerMode(Params.SceneOutliner)
+	, bHideComponents(Params.bHideComponents)
+	, bHideFoundationHierarchy(Params.bHideFoundationHierarchy)
+	, SpecifiedWorldToDisplay(Params.SpecifiedWorldToDisplay)
 {
+	SceneOutliner->AddFilter(MakeShared<FActorFilter>(FActorTreeItem::FFilterPredicate::CreateLambda([this](const AActor* Actor)
+		{
+			return IsActorDisplayable(Actor);
+		}), FSceneOutlinerFilter::EDefaultBehaviour::Pass));
 }
 
 FActorMode::~FActorMode()
@@ -83,8 +89,9 @@ FActorMode::~FActorMode()
 
 TUniquePtr<ISceneOutlinerHierarchy> FActorMode::CreateHierarchy()
 {
-	TUniquePtr<FActorHierarchy> ActorHierarchy = FActorHierarchy::Create(this, RepresentingWorld, [this](const AActor* InActor) {return IsActorDisplayable(InActor); });
+	TUniquePtr<FActorHierarchy> ActorHierarchy = FActorHierarchy::Create(this, RepresentingWorld);
 	ActorHierarchy->SetShowingComponents(!bHideComponents);
+	ActorHierarchy->SetShowingFoundations(!bHideFoundationHierarchy);
 
 	return ActorHierarchy;
 }
@@ -267,18 +274,18 @@ bool FActorMode::IsActorDisplayable(const AActor* Actor) const
 {
 	static const FName SequencerActorTag(TEXT("SequencerActor"));
 
-	return Actor && 
+	return Actor &&
 		!SceneOutliner->GetSharedData().bOnlyShowFolders && 												// Don't show actors if we're only showing folders
 		Actor->IsEditable() &&															// Only show actors that are allowed to be selected and drawn in editor
 		Actor->IsListedInSceneOutliner() &&
 		((Actor->GetWorld()->IsPlayInEditor() || !Actor->HasAnyFlags(RF_Transient)) ||
-			(SceneOutliner->GetSharedData().bShowTransient && Actor->HasAnyFlags(RF_Transient)) ||			// Don't show transient actors in non-play worlds
+		(SceneOutliner->GetSharedData().bShowTransient && Actor->HasAnyFlags(RF_Transient)) ||			// Don't show transient actors in non-play worlds
 			(Actor->ActorHasTag(SequencerActorTag))) &&
 		!Actor->IsTemplate() &&															// Should never happen, but we never want CDOs displayed
 		!FActorEditorUtils::IsABuilderBrush(Actor) &&									// Don't show the builder brush
 		!Actor->IsA(AWorldSettings::StaticClass()) &&									// Don't show the WorldSettings actor, even though it is technically editable
 		!Actor->IsPendingKill() &&														// We don't want to show actors that are about to go away
-		FLevelUtils::IsLevelVisible(Actor->GetLevel());								// Only show Actors whose level is visible
+		FLevelUtils::IsLevelVisible(Actor->GetLevel());									// Only show Actors whose level is visible
 }
 
 void FActorMode::OnFilterTextChanged(const FText& InFilterText)
diff --git a/Engine/Source/Editor/SceneOutliner/Private/ActorModeInteractive.cpp b/Engine/Source/Editor/SceneOutliner/Private/ActorModeInteractive.cpp
--- a/Engine/Source/Editor/SceneOutliner/Private/ActorModeInteractive.cpp
+++ b/Engine/Source/Editor/SceneOutliner/Private/ActorModeInteractive.cpp
@@ -5,8 +5,8 @@
 #include "Editor.h"
 #include "ActorTreeItem.h"
 
-FActorModeInteractive::FActorModeInteractive(SSceneOutliner* InSceneOutliner, bool bInHideComponents, TWeakObjectPtr<UWorld> InSpecifiedWorldToDisplay)
-	: FActorMode(InSceneOutliner, bInHideComponents, InSpecifiedWorldToDisplay)
+FActorModeInteractive::FActorModeInteractive(const FActorModeParams& Params)
+	: FActorMode(Params)
 {
 	USelection::SelectionChangedEvent.AddRaw(this, &FActorModeInteractive::OnLevelSelectionChanged);
 	USelection::SelectObjectEvent.AddRaw(this, &FActorModeInteractive::OnLevelSelectionChanged);
diff --git a/Engine/Source/Editor/SceneOutliner/Private/ActorPickingMode.cpp b/Engine/Source/Editor/SceneOutliner/Private/ActorPickingMode.cpp
--- a/Engine/Source/Editor/SceneOutliner/Private/ActorPickingMode.cpp
+++ b/Engine/Source/Editor/SceneOutliner/Private/ActorPickingMode.cpp
@@ -4,8 +4,8 @@
 #include "SSceneOutliner.h"
 #include "ActorTreeItem.h"
 
-FActorPickingMode::FActorPickingMode(SSceneOutliner* InSceneOutliner, bool bInHideComponents, FOnSceneOutlinerItemPicked OnItemPickedDelegate, TWeakObjectPtr<UWorld> InSpecifiedWorldToDisplay)
-	: FActorMode(InSceneOutliner, bInHideComponents, InSpecifiedWorldToDisplay)
+FActorPickingMode::FActorPickingMode(const FActorModeParams& Params, FOnSceneOutlinerItemPicked OnItemPickedDelegate)
+	: FActorMode(Params)
 	, OnItemPicked(OnItemPickedDelegate)
 {
 }
diff --git a/Engine/Source/Editor/SceneOutliner/Private/ActorTreeItem.cpp b/Engine/Source/Editor/SceneOutliner/Private/ActorTreeItem.cpp
--- a/Engine/Source/Editor/SceneOutliner/Private/ActorTreeItem.cpp
+++ b/Engine/Source/Editor/SceneOutliner/Private/ActorTreeItem.cpp
@@ -17,6 +17,7 @@
 #include "ISceneOutlinerMode.h"
 #include "Logging/MessageLog.h"
 #include "SSocketChooser.h"
+#include "Foundation/FoundationActor.h"
 
 #define LOCTEXT_NAMESPACE "SceneOutliner_ActorTreeItem"
 
@@ -111,6 +112,17 @@ private:
 	FText GetDisplayText() const
 	{
 		const AActor* Actor = ActorPtr.Get();
+		if (const AFoundationActor* FoundationActor = Cast<AFoundationActor>(Actor))
+		{
+			if (FoundationActor->IsDirty())
+			{
+				FFormatNamedArguments Args;
+				Args.Add(TEXT("ActorLabel"), FText::FromString(Actor->GetActorLabel()));
+				Args.Add(TEXT("EditTag"), LOCTEXT("EditingFoundationLabel", "*"));
+				return FText::Format(LOCTEXT("FoundationDisplay", "{ActorLabel}{EditTag}"), Args);
+			}
+		}
+
 		return Actor ? FText::FromString(Actor->GetActorLabel()) : LOCTEXT("ActorLabelForMissingActor", "(Deleted Actor)");
 	}
 
@@ -233,13 +245,28 @@ private:
 
 	FSlateColor GetForegroundColor() const
 	{
+		AActor* Actor = ActorPtr.Get();
+
+		// Color foundations differently if they are being edited
+		if (const AFoundationActor* FoundationActor = Cast<AFoundationActor>(Actor))
+		{
+			if (FoundationActor->IsEditing() && !FoundationActor->IsSelected())
+			{
+				return FAppStyle::Get().GetSlateColor("Colors.AccentGreen");
+			}
+			else
+			{
+				return FSlateColor::UseForeground();
+			}
+		}
+
+
 		auto TreeItem = TreeItemPtr.Pin();
 		if (auto BaseColor = FSceneOutlinerCommonLabelData::GetForegroundColor(*TreeItem))
 		{
 			return BaseColor.GetValue();
 		}
 
-		AActor* Actor = ActorPtr.Get();
 		if (!Actor)
 		{
 			// Deleted actor!
diff --git a/Engine/Source/Editor/SceneOutliner/Private/SceneOutlinerModule.cpp b/Engine/Source/Editor/SceneOutliner/Private/SceneOutlinerModule.cpp
--- a/Engine/Source/Editor/SceneOutliner/Private/SceneOutlinerModule.cpp
+++ b/Engine/Source/Editor/SceneOutliner/Private/SceneOutlinerModule.cpp
@@ -58,7 +58,12 @@ TSharedRef<ISceneOutliner> FSceneOutlinerModule::CreateActorPicker(const FSceneO
 
 	FCreateSceneOutlinerMode ModeFactory = FCreateSceneOutlinerMode::CreateLambda([&OnItemPicked, &SpecifiedWorld](SSceneOutliner* Outliner)
 		{
-			return new FActorPickingMode(Outliner, true, OnItemPicked, SpecifiedWorld);
+			FActorModeParams Params;
+			Params.SceneOutliner = Outliner;
+			Params.SpecifiedWorldToDisplay = SpecifiedWorld;
+			Params.bHideComponents = true;
+			Params.bHideFoundationHierarchy = true;
+			return new FActorPickingMode(Params, OnItemPicked);
 		});
 
 	
@@ -88,7 +93,12 @@ TSharedRef<ISceneOutliner> FSceneOutlinerModule::CreateComponentPicker(const FSc
 
 	FCreateSceneOutlinerMode ModeFactory = FCreateSceneOutlinerMode::CreateLambda([&OnItemPicked, &SpecifiedWorld](SSceneOutliner* Outliner)
 		{
-			return new FActorPickingMode(Outliner, false, OnItemPicked, SpecifiedWorld);
+			FActorModeParams Params;
+			Params.SceneOutliner = Outliner;
+			Params.SpecifiedWorldToDisplay = SpecifiedWorld;
+			Params.bHideComponents = false;
+			Params.bHideFoundationHierarchy = true;
+			return new FActorPickingMode(Params, OnItemPicked);
 		});
 
 	FSceneOutlinerInitializationOptions InitOptions(InInitOptions);
diff --git a/Engine/Source/Editor/SceneOutliner/Private/SceneOutlinerStandaloneTypes.cpp b/Engine/Source/Editor/SceneOutliner/Private/SceneOutlinerStandaloneTypes.cpp
--- a/Engine/Source/Editor/SceneOutliner/Private/SceneOutlinerStandaloneTypes.cpp
+++ b/Engine/Source/Editor/SceneOutliner/Private/SceneOutlinerStandaloneTypes.cpp
@@ -13,7 +13,7 @@
 uint32 FSceneOutlinerTreeItemType::NextUniqueID = 0;
 const FSceneOutlinerTreeItemType ISceneOutlinerTreeItem::Type;
 
-const FLinearColor FSceneOutlinerCommonLabelData::DarkColor(0.3f, 0.3f, 0.3f);
+const FLinearColor FSceneOutlinerCommonLabelData::DarkColor(0.15f, 0.15f, 0.15f);
 
 TOptional<FLinearColor> FSceneOutlinerCommonLabelData::GetForegroundColor(const ISceneOutlinerTreeItem& TreeItem) const
 {
diff --git a/Engine/Source/Editor/SceneOutliner/Public/ActorBrowsingMode.h b/Engine/Source/Editor/SceneOutliner/Public/ActorBrowsingMode.h
--- a/Engine/Source/Editor/SceneOutliner/Public/ActorBrowsingMode.h
+++ b/Engine/Source/Editor/SceneOutliner/Public/ActorBrowsingMode.h
@@ -98,6 +98,7 @@ private:
 	static TSharedRef<FSceneOutlinerFilter> CreateHideTemporaryActorsFilter();
 	static TSharedRef<FSceneOutlinerFilter> CreateIsInCurrentLevelFilter();
 	static TSharedRef<FSceneOutlinerFilter> CreateHideComponentsFilter();
+	static TSharedRef<FSceneOutlinerFilter> CreateHideFoundationsFilter();
 private:
 	/** Number of actors which have passed through the filters */
 	uint32 FilteredActorCount = 0;
diff --git a/Engine/Source/Editor/SceneOutliner/Public/ActorHierarchy.h b/Engine/Source/Editor/SceneOutliner/Public/ActorHierarchy.h
--- a/Engine/Source/Editor/SceneOutliner/Public/ActorHierarchy.h
+++ b/Engine/Source/Editor/SceneOutliner/Public/ActorHierarchy.h
@@ -10,7 +10,7 @@ class FActorHierarchy : public ISceneOutlinerHierarchy
 public:
 	virtual ~FActorHierarchy();
 
-	static TUniquePtr<FActorHierarchy> Create(ISceneOutlinerMode* Mode, const TWeakObjectPtr<UWorld>& World, TFunction<bool(const AActor*)> InIsActorDisplayableCallback);
+	static TUniquePtr<FActorHierarchy> Create(ISceneOutlinerMode* Mode, const TWeakObjectPtr<UWorld>& World);
 
 	/** Find all direct children of a tree item in an existing item map, if any exist. */
 	virtual void FindChildren(const ISceneOutlinerTreeItem& Item, const TMap<FSceneOutlinerTreeItemID, FSceneOutlinerTreeItemPtr>& Items, TArray<FSceneOutlinerTreeItemPtr>& OutChildItems) const override;
@@ -25,6 +25,7 @@ public:
 	virtual FSceneOutlinerTreeItemPtr CreateParentItem(const FSceneOutlinerTreeItemPtr& Item) const override;
 
 	void SetShowingComponents(bool bInShowingComponents) { bShowingComponents = bInShowingComponents; }
+	void SetShowingFoundations(bool bInShowingFoundations) { bShowingFoundations = bInShowingFoundations; }
 private:
 	/** Adds all the direct and indirect children of a world to OutItems */
 	void CreateWorldChildren(UWorld* World, TArray<FSceneOutlinerTreeItemPtr>& OutItems) const;
@@ -62,14 +63,11 @@ private:
 	void FullRefreshEvent();
 
 	bool bShowingComponents = false;
+	bool bShowingFoundations = false;
 
 	TWeakObjectPtr<UWorld> RepresentingWorld;
-
-	/** Used to determine if an actor can be displayed in the outliner */
-	TFunction<bool(const AActor*)> IsActorDisplayableCallback;
-
 private:
-	FActorHierarchy(ISceneOutlinerMode* Mode, const TWeakObjectPtr<UWorld>& Worlds, TFunction<bool(const AActor*)> InIsActorDisplayableCallback);
+	FActorHierarchy(ISceneOutlinerMode* Mode, const TWeakObjectPtr<UWorld>& Worlds);
 
 	FActorHierarchy(const FActorHierarchy&) = delete;
 	FActorHierarchy& operator=(const FActorHierarchy&) = delete;
diff --git a/Engine/Source/Editor/SceneOutliner/Public/ActorMode.h b/Engine/Source/Editor/SceneOutliner/Public/ActorMode.h
--- a/Engine/Source/Editor/SceneOutliner/Public/ActorMode.h
+++ b/Engine/Source/Editor/SceneOutliner/Public/ActorMode.h
@@ -19,6 +19,23 @@ namespace SceneOutliner
 	};
 }
 
+struct SCENEOUTLINER_API FActorModeParams
+{
+	FActorModeParams() {}
+
+	FActorModeParams(SSceneOutliner* InSceneOutliner, const TWeakObjectPtr<UWorld>& InSpecifiedWorldToDisplay = nullptr, bool bInHideComponents = true, bool bInHideFoundationHierarchy = true)
+		: SpecifiedWorldToDisplay(InSpecifiedWorldToDisplay)
+		, SceneOutliner(InSceneOutliner)
+		, bHideComponents(bInHideComponents)
+		, bHideFoundationHierarchy(bInHideFoundationHierarchy)
+	{}
+
+	TWeakObjectPtr<UWorld> SpecifiedWorldToDisplay = nullptr;
+	SSceneOutliner* SceneOutliner;
+	bool bHideComponents;
+	bool bHideFoundationHierarchy;
+};
+
 class SCENEOUTLINER_API FActorMode : public ISceneOutlinerMode
 {
 public:
@@ -27,7 +44,7 @@ public:
 		enum Type { World = 0, Folder = 10, Actor = 20 };
 	};
 	
-	FActorMode(SSceneOutliner* InSceneOutliner, bool bHideComponents, TWeakObjectPtr<UWorld> InSpecifiedWorldToDisplay = nullptr);
+	FActorMode(const FActorModeParams& Params);
 	virtual ~FActorMode();
 
 	virtual void Rebuild() override;
@@ -55,6 +72,8 @@ protected:
 protected:
 	// Should the hide components filter be enabled
 	bool bHideComponents;
+	// Are foundations being hidden
+	bool bHideFoundationHierarchy;
 	/** The world which we are currently representing */
 	TWeakObjectPtr<UWorld> RepresentingWorld;
 	/** The world which the user manually selected */
diff --git a/Engine/Source/Editor/SceneOutliner/Public/ActorModeInteractive.h b/Engine/Source/Editor/SceneOutliner/Public/ActorModeInteractive.h
--- a/Engine/Source/Editor/SceneOutliner/Public/ActorModeInteractive.h
+++ b/Engine/Source/Editor/SceneOutliner/Public/ActorModeInteractive.h
@@ -7,7 +7,7 @@
 class FActorModeInteractive : public FActorMode
 {
 public:
-	FActorModeInteractive(SSceneOutliner* InSceneOutliner, bool bHideComponents, TWeakObjectPtr<UWorld> InSpecifiedWorldToDisplay = nullptr);
+	FActorModeInteractive(const FActorModeParams& Params);
 	virtual ~FActorModeInteractive();
 
 	virtual bool IsInteractive() const override { return true; }
diff --git a/Engine/Source/Editor/SceneOutliner/Public/ActorPickingMode.h b/Engine/Source/Editor/SceneOutliner/Public/ActorPickingMode.h
--- a/Engine/Source/Editor/SceneOutliner/Public/ActorPickingMode.h
+++ b/Engine/Source/Editor/SceneOutliner/Public/ActorPickingMode.h
@@ -7,7 +7,7 @@
 class FActorPickingMode : public FActorMode
 {
 public:
-	FActorPickingMode(SSceneOutliner* InSceneOutliner, bool bInHideComponents, FOnSceneOutlinerItemPicked OnItemPickedDelegate, TWeakObjectPtr<UWorld> InSpecifiedWorldToDisplay = nullptr);
+	FActorPickingMode(const FActorModeParams& Params, FOnSceneOutlinerItemPicked OnItemPickedDelegate);
 
 	virtual ~FActorPickingMode() {};
 public:
diff --git a/Engine/Source/Editor/SceneOutliner/Public/SceneOutlinerDragDrop.h b/Engine/Source/Editor/SceneOutliner/Public/SceneOutlinerDragDrop.h
--- a/Engine/Source/Editor/SceneOutliner/Public/SceneOutlinerDragDrop.h
+++ b/Engine/Source/Editor/SceneOutliner/Public/SceneOutlinerDragDrop.h
@@ -21,10 +21,10 @@ enum class ESceneOutlinerDropCompatibility : uint8
 	CompatibleMultipleAttach,
 	IncompatibleMultipleAttach,
 	CompatibleDetach,
-	CompatibleMultipleDetach
+	CompatibleMultipleDetach,
 };
 
-/** Consilidated drag/drop with parsing functions for the scene outliner */
+/** Consolidated drag/drop with parsing functions for the scene outliner */
 struct FSceneOutlinerDragDropPayload
 {
 	/** Default constructor, resulting in unset contents */
diff --git a/Engine/Source/Runtime/Engine/Classes/GameFramework/Actor.h b/Engine/Source/Runtime/Engine/Classes/GameFramework/Actor.h
--- a/Engine/Source/Runtime/Engine/Classes/GameFramework/Actor.h
+++ b/Engine/Source/Runtime/Engine/Classes/GameFramework/Actor.h
@@ -3623,7 +3623,6 @@ private:
 
 	friend UWorld;
 	friend class FFoliageHelper;
-	friend class ULevelStreamingFoundationInstance;
 };
 
 struct FSetActorGuid
diff --git a/Engine/Source/Runtime/Engine/Private/Actor.cpp b/Engine/Source/Runtime/Engine/Private/Actor.cpp
--- a/Engine/Source/Runtime/Engine/Private/Actor.cpp
+++ b/Engine/Source/Runtime/Engine/Private/Actor.cpp
@@ -4590,9 +4590,15 @@ bool AActor::IsSelectionChild() const
 		return true;
 	}
 
-	if (AActor* AttachParent = GetAttachParentActor())
+	const AActor* AttachParent = GetAttachParentActor();
+	while (AttachParent != nullptr)
 	{
-		return AttachParent->IsSelectionParentOfAttachedActors();
+		if (AttachParent->IsSelectionParentOfAttachedActors())
+		{
+			return true;
+		}
+
+		AttachParent = AttachParent->GetAttachParentActor();
 	}
 
 	return false;
@@ -4605,12 +4611,15 @@ AActor* AActor::GetSelectionParent() const
 		return GetParentActor();
 	}
 
-	if (AActor* AttachParent = GetAttachParentActor())
+	AActor* AttachParent = GetAttachParentActor();
+	while (AttachParent != nullptr)
 	{
 		if (AttachParent->IsSelectionParentOfAttachedActors())
 		{
 			return AttachParent;
 		}
+
+		AttachParent = AttachParent->GetAttachParentActor();
 	}
 
 	return nullptr;
@@ -4659,11 +4668,17 @@ void AActor::PushSelectionToProxies()
 
 	if (IsSelectionParentOfAttachedActors())
 	{
-		ForEachAttachedActors([](AActor* AttachedActor)
+		TFunction<bool(AActor*)> PushAllChildrenToProxies = [&PushAllChildrenToProxies](AActor* AttachedActor)
 		{
 			AttachedActor->PushSelectionToProxies();
+			if (!AttachedActor->IsSelectionParentOfAttachedActors())
+			{
+				AttachedActor->ForEachAttachedActors(PushAllChildrenToProxies);
+			}
 			return true;
-		});
+		};
+
+		ForEachAttachedActors(PushAllChildrenToProxies);
 	}
 }
 
diff --git a/Engine/Source/Runtime/Engine/Private/Foundation/FoundationActor.cpp b/Engine/Source/Runtime/Engine/Private/Foundation/FoundationActor.cpp
--- a/Engine/Source/Runtime/Engine/Private/Foundation/FoundationActor.cpp
+++ b/Engine/Source/Runtime/Engine/Private/Foundation/FoundationActor.cpp
@@ -631,6 +631,14 @@ void AFoundationActor::Commit()
 	FoundationSubsystem->CommitFoundation(this);
 }
 
+void AFoundationActor::Discard()
+{
+	UFoundationSubsystem* FoundationSubsystem = GetFoundationSubsystem();
+	check(FoundationSubsystem);
+	const bool bDiscardEdits = true;
+	FoundationSubsystem->CommitFoundation(this, bDiscardEdits);
+}
+
 void AFoundationActor::SaveAs()
 {
 	UFoundationSubsystem* FoundationSubsystem = GetFoundationSubsystem();
diff --git a/Engine/Source/Runtime/Engine/Private/Foundation/FoundationEditorInstanceActor.cpp b/Engine/Source/Runtime/Engine/Private/Foundation/FoundationEditorInstanceActor.cpp
--- a/Engine/Source/Runtime/Engine/Private/Foundation/FoundationEditorInstanceActor.cpp
+++ b/Engine/Source/Runtime/Engine/Private/Foundation/FoundationEditorInstanceActor.cpp
@@ -40,7 +40,7 @@ AFoundationEditorInstanceActor* AFoundationEditorInstanceActor::Create(AFoundati
 	
 	for (AActor* LevelActor : LoadedLevel->Actors)
 	{
-		if (LevelActor && !LevelActor->IsChildActor() && LevelActor != InstanceActor)
+		if (LevelActor && LevelActor->GetAttachParentActor() == nullptr && !LevelActor->IsChildActor() && LevelActor != InstanceActor)
 		{
 			LevelActor->AttachToActor(InstanceActor, FAttachmentTransformRules::KeepWorldTransform);
 		}
diff --git a/Engine/Source/Runtime/Engine/Private/Foundation/FoundationInstanceLevelStreaming.cpp b/Engine/Source/Runtime/Engine/Private/Foundation/FoundationInstanceLevelStreaming.cpp
--- a/Engine/Source/Runtime/Engine/Private/Foundation/FoundationInstanceLevelStreaming.cpp
+++ b/Engine/Source/Runtime/Engine/Private/Foundation/FoundationInstanceLevelStreaming.cpp
@@ -21,12 +21,6 @@ ULevelStreamingFoundationInstance::ULevelStreamingFoundationInstance(const FObje
 {
 #if WITH_EDITOR
 	SetShouldBeVisibleInEditor(true);
-
-	if (!IsTemplate() && !GetWorld()->IsPlayInEditor())
-	{
-		FWorldDelegates::LevelAddedToWorld.AddUObject(this, &ULevelStreamingFoundationInstance::OnLevelAddedToWorld);
-		FWorldDelegates::LevelRemovedFromWorld.AddUObject(this, &ULevelStreamingFoundationInstance::OnLevelRemovedFromWorld);
-	}
 #endif
 }
 
@@ -41,56 +35,11 @@ AFoundationActor* ULevelStreamingFoundationInstance::GetFoundationActor() const
 }
 
 #if WITH_EDITOR
-void ULevelStreamingFoundationInstance::BeginDestroy()
-{
-	FWorldDelegates::LevelAddedToWorld.RemoveAll(this);
-	FWorldDelegates::LevelRemovedFromWorld.RemoveAll(this);
-	Super::BeginDestroy();
-}
-
-void ULevelStreamingFoundationInstance::SetLoadedLevel(ULevel* Level)
-{
-	Super::SetLoadedLevel(Level);
-
-	if (Level)
-	{
-		for (AActor* LevelActor : Level->Actors)
-		{
-			if (LevelActor)
-			{
-				FSetActorHiddenInSceneOutliner Hide(LevelActor);
-			}
-		}
-	}
-}
-
 FBox ULevelStreamingFoundationInstance::GetBounds() const
 {
 	check(GetLoadedLevel());
 	return ALevelBounds::CalculateLevelBounds(GetLoadedLevel());
 }
-
-void ULevelStreamingFoundationInstance::OnLoadedActorAddedToLevel(AActor& InActor)
-{
-	FSetActorHiddenInSceneOutliner Hide(&InActor);
-}
-
-void ULevelStreamingFoundationInstance::OnLevelRemovedFromWorld(ULevel* Level, UWorld* InWorld)
-{
-	if (Level && (Level == LoadedLevel))
-	{
-		Level->OnLoadedActorAddedToLevelEvent.RemoveAll(this);
-	}
-}
-
-void ULevelStreamingFoundationInstance::OnLevelAddedToWorld(ULevel* Level, UWorld* InWorld)
-{
-	if (Level == LoadedLevel)
-	{
-		check(Level->GetWorld() == InWorld);
-		Level->OnLoadedActorAddedToLevelEvent.AddUObject(this, &ULevelStreamingFoundationInstance::OnLoadedActorAddedToLevel);
-	}
-}
 #endif
 
 ULevelStreamingFoundationInstance* ULevelStreamingFoundationInstance::LoadInstance(AFoundationActor* FoundationActor)
diff --git a/Engine/Source/Runtime/Engine/Private/Foundation/FoundationSubsystem.cpp b/Engine/Source/Runtime/Engine/Private/Foundation/FoundationSubsystem.cpp
--- a/Engine/Source/Runtime/Engine/Private/Foundation/FoundationSubsystem.cpp
+++ b/Engine/Source/Runtime/Engine/Private/Foundation/FoundationSubsystem.cpp
@@ -367,7 +367,7 @@ void UFoundationSubsystem::ForEachFoundationAncestors(AActor* Actor, TFunctionRe
 	} while (ParentFoundation != nullptr && Operation(ParentFoundation));
 }
 
-AFoundationActor* UFoundationSubsystem::GetOwningFoundation(ULevel* Level) const
+AFoundationActor* UFoundationSubsystem::GetOwningFoundation(const ULevel* Level) const
 {
 	if (ULevelStreaming* BaseLevelStreaming = FLevelUtils::FindStreamingLevel(Level))
 	{
@@ -538,25 +538,30 @@ bool UFoundationSubsystem::IsCurrent(const AFoundationActor* FoundationActor) co
 	return false;
 }
 
-bool UFoundationSubsystem::MoveActorsTo(AFoundationActor* FoundationActor, const TArray<AActor*>& ActorsToMove)
+bool UFoundationSubsystem::MoveActorsToLevel(const TArray<AActor*>& ActorsToRemove, ULevel* DestinationLevel) const
 {
-	check(IsEditingFoundation(FoundationActor));
-	ULevel* FoundationLevel = GetFoundationLevel(FoundationActor);
-	check(FoundationLevel);
+	check(DestinationLevel);
 
 	const bool bWarnAboutReferences = true;
 	const bool bWarnAboutRenaming = true;
 	const bool bMoveAllOrFail = true;
-
-	if (!EditorLevelUtils::MoveActorsToLevel(ActorsToMove, FoundationLevel, bWarnAboutReferences, bWarnAboutRenaming, bMoveAllOrFail))
+	if (!EditorLevelUtils::MoveActorsToLevel(ActorsToRemove, DestinationLevel, bWarnAboutReferences, bWarnAboutRenaming, bMoveAllOrFail))
 	{
-		UE_LOG(LogFoundation, Warning, TEXT("Failed to move actors to foundation because not all actors could be moved"));
+		UE_LOG(LogFoundation, Warning, TEXT("Failed to move actors out of foundation because not all actors could be moved"));
 		return false;
 	}
-
 	return true;
 }
 
+bool UFoundationSubsystem::MoveActorsTo(AFoundationActor* FoundationActor, const TArray<AActor*>& ActorsToMove)
+{
+	check(IsEditingFoundation(FoundationActor));
+	ULevel* FoundationLevel = GetFoundationLevel(FoundationActor);
+	check(FoundationLevel);
+
+	return MoveActorsToLevel(ActorsToMove, FoundationLevel);
+}
+
 AFoundationActor* UFoundationSubsystem::CreateFoundationFrom(const TArray<AActor*>& ActorsToMove, UWorld* TemplateWorld)
 {
 	ULevel* CurrentLevel = GetWorld()->GetCurrentLevel();
@@ -679,7 +684,7 @@ UFoundationSubsystem::FLevelsToRemoveScope::~FLevelsToRemoveScope()
 		double StartTime = FPlatformTime::Seconds();
 		const bool bClearSelection = false;
 		// No need to clear the whole editor selection since actor of this level will be removed from the selection by: UEditorEngine::OnLevelRemovedFromWorld
-		EditorLevelUtils::RemoveLevelsFromWorld(Levels, bClearSelection, bResetTrans);
+		EditorLevelUtils::RemoveLevelsFromWorld(Levels, bClearSelection, true);
 		double ElapsedTime = FPlatformTime::Seconds() - StartTime;
 		UE_LOG(LogFoundation, Log, TEXT("Unloaded %s levels in %s seconds"), *FText::AsNumber(Levels.Num()).ToString(), *FText::AsNumber(ElapsedTime).ToString());
 	}
@@ -1014,7 +1019,8 @@ void UFoundationSubsystem::EditFoundation(AFoundationActor* FoundationActor, TWe
 			{
 				check(!IsFoundationEditDirty(FoundationEdit));
 				check(FoundationToCommit == nullptr);
-				FoundationToCommit = Foundation;
+				FoundationToCommit = FoundationEdit->LevelStreaming->GetFoundationActor();
+				check(FoundationToCommit != nullptr);
 				return false;
 			}
 			return true;
@@ -1078,7 +1084,7 @@ void UFoundationSubsystem::CommitChildrenFoundations(AFoundationActor* Foundatio
 	});
 }
 
-void UFoundationSubsystem::CommitFoundation(AFoundationActor* FoundationActor)
+void UFoundationSubsystem::CommitFoundation(AFoundationActor* FoundationActor, bool bDiscardEdits)
 {
 	check(CanCommitFoundation(FoundationActor));
 
@@ -1088,7 +1094,7 @@ void UFoundationSubsystem::CommitFoundation(AFoundationActor* FoundationActor)
 	check(EditingWorld);
 			
 	bool bChangesCommitted = false;
-	if (IsFoundationEditDirty(FoundationEdit))
+	if (IsFoundationEditDirty(FoundationEdit) && !bDiscardEdits)
 	{
 		const bool bPromptUserToSave = true;
 		const bool bSaveMapPackages = true;
@@ -1192,6 +1198,13 @@ void UFoundationSubsystem::SaveFoundationAs(AFoundationActor* FoundationActor)
 	GEditor->SelectActor(FoundationActor, true, true);
 }
 
+AFoundationActor* UFoundationSubsystem::GetParentFoundation(const AActor* Actor) const
+{
+	check(Actor);
+	const ULevel* OwningLevel = Actor->GetLevel();
+	check(OwningLevel);
+	return GetOwningFoundation(OwningLevel);
+}
 #endif
 
 #undef LOCTEXT_NAMESPACE
diff --git a/Engine/Source/Runtime/Engine/Public/Foundation/FoundationActor.h b/Engine/Source/Runtime/Engine/Public/Foundation/FoundationActor.h
--- a/Engine/Source/Runtime/Engine/Public/Foundation/FoundationActor.h
+++ b/Engine/Source/Runtime/Engine/Public/Foundation/FoundationActor.h
@@ -60,6 +60,7 @@ public:
 	bool IsEditing() const;
 	void Edit(AActor* ContextActor = nullptr);
 	void Commit();
+	void Discard();
 	void SaveAs();
 	bool HasDirtyChildren() const;
 	bool IsDirty() const;
diff --git a/Engine/Source/Runtime/Engine/Public/Foundation/FoundationInstanceLevelStreaming.h b/Engine/Source/Runtime/Engine/Public/Foundation/FoundationInstanceLevelStreaming.h
--- a/Engine/Source/Runtime/Engine/Public/Foundation/FoundationInstanceLevelStreaming.h
+++ b/Engine/Source/Runtime/Engine/Public/Foundation/FoundationInstanceLevelStreaming.h
@@ -19,20 +19,11 @@ public:
 	AFoundationActor* GetFoundationActor() const;
 
 #if WITH_EDITOR
-	virtual void BeginDestroy() override;
-	virtual void SetLoadedLevel(ULevel* Level) override;
 	virtual bool ShowInLevelCollection() const override { return false; }
 	FBox GetBounds() const;
 #endif
 	
 protected:
-
-#if WITH_EDITOR
-	void OnLevelAddedToWorld(ULevel* Level, UWorld* InWorld);
-	void OnLevelRemovedFromWorld(ULevel* Level, UWorld* InWorld);
-	void OnLoadedActorAddedToLevel(AActor& InActor);
-#endif
-
 	static ULevelStreamingFoundationInstance* LoadInstance(AFoundationActor* FoundationActor);
 	static void UnloadInstance(ULevelStreamingFoundationInstance* LevelStreaming);
 
diff --git a/Engine/Source/Runtime/Engine/Public/Foundation/FoundationSubsystem.h b/Engine/Source/Runtime/Engine/Public/Foundation/FoundationSubsystem.h
--- a/Engine/Source/Runtime/Engine/Public/Foundation/FoundationSubsystem.h
+++ b/Engine/Source/Runtime/Engine/Public/Foundation/FoundationSubsystem.h
@@ -60,7 +60,7 @@ public:
 	bool CanEditFoundation(const AFoundationActor* FoundationActor, FText* OutReason = nullptr) const;
 	bool CanCommitFoundation(const AFoundationActor* FoundationActor, FText* OutReason = nullptr) const;
 	void EditFoundation(AFoundationActor* FoundationActor, TWeakObjectPtr<AActor> ContextActorPtr = nullptr);
-	void CommitFoundation(AFoundationActor* FoundationActor);
+	void CommitFoundation(AFoundationActor* FoundationActor, bool bDiscardEdits = false);
 	void SaveFoundationAs(AFoundationActor* FoundationActor);
 	bool IsEditingFoundationDirty(const AFoundationActor* FoundationActor) const;
 	bool IsEditingFoundation(const AFoundationActor* FoundationActor) const { return GetFoundationEdit(FoundationActor) != nullptr; }
@@ -81,12 +81,15 @@ public:
 	bool SetCurrent(AFoundationActor* FoundationActor) const;
 	bool IsCurrent(const AFoundationActor* FoundationActor) const;
 	AFoundationActor* CreateFoundationFrom(const TArray<AActor*>& ActorsToMove, UWorld* TemplateWorld = nullptr);
+	bool MoveActorsToLevel(const TArray<AActor*>& ActorsToRemove, ULevel* DestinationLevel) const;
 	bool MoveActorsTo(AFoundationActor* FoundationActor, const TArray<AActor*>& ActorsToMove);
 
 	bool CanMoveActorToLevel(const AActor* Actor) const;
 	void DiscardEdits();
 	void OnActorDeleted(AActor* Actor);
 	ULevel* GetFoundationLevel(const AFoundationActor* FoundationActor) const;
+
+	AFoundationActor* GetParentFoundation(const AActor* Actor) const;
 #endif
 
 private:
@@ -95,7 +98,7 @@ private:
 	void UnloadFoundation(const FFoundationID& FoundationID);
 	void ForEachActorInLevel(ULevel* Level, TFunctionRef<bool(AActor * LevelActor)> Operation) const;
 	void ForEachFoundationAncestors(AActor* Actor, TFunctionRef<bool(AFoundationActor*)> Operation) const;
-	AFoundationActor* GetOwningFoundation(ULevel* Level) const;
+	AFoundationActor* GetOwningFoundation(const ULevel* Level) const;
 	FFoundationID ComputeFoundationID(AFoundationActor* FoundationActor) const;
 #if WITH_EDITOR
 	void CommitChildrenFoundations(AFoundationActor* FoundationActor);
