Fixed crash in SceneOutliner when actors are spawned and attached in the same frame.

diff --git a/Engine/Source/Editor/SceneOutliner/Private/SSceneOutliner.cpp b/Engine/Source/Editor/SceneOutliner/Private/SSceneOutliner.cpp
--- a/Engine/Source/Editor/SceneOutliner/Private/SSceneOutliner.cpp
+++ b/Engine/Source/Editor/SceneOutliner/Private/SSceneOutliner.cpp
@@ -1569,10 +1569,6 @@ void SSceneOutliner::OnHierarchyChangedEvent(FSceneOutlinerHierarchyChangedData
 	else if (Event.Type == FSceneOutlinerHierarchyChangedData::Moved)
 	{
 		FSceneOutlinerTreeItemPtr* Item = Event.Item ? &Event.Item : TreeItemMap.Find(Event.ItemID);
-		if (!Item)
-		{
-			Item = PendingTreeItemMap.Find(Event.ItemID);
-		}
 
 		if (Item)
 		{
