Fixed static analysis warning in actortreeitem.cpp

diff --git a/Engine/Source/Editor/SceneOutliner/Private/ActorTreeItem.cpp b/Engine/Source/Editor/SceneOutliner/Private/ActorTreeItem.cpp
--- a/Engine/Source/Editor/SceneOutliner/Private/ActorTreeItem.cpp
+++ b/Engine/Source/Editor/SceneOutliner/Private/ActorTreeItem.cpp
@@ -117,7 +117,7 @@ private:
 			if (FoundationActor->IsDirty())
 			{
 				FFormatNamedArguments Args;
-				Args.Add(TEXT("ActorLabel"), FText::FromString(Actor->GetActorLabel()));
+				Args.Add(TEXT("ActorLabel"), FText::FromString(FoundationActor->GetActorLabel()));
 				Args.Add(TEXT("EditTag"), LOCTEXT("EditingFoundationLabel", "*"));
 				return FText::Format(LOCTEXT("FoundationDisplay", "{ActorLabel}{EditTag}"), Args);
 			}
