Moved FolderDragDropOp out of SceneOutliner module

diff --git a/Engine/Plugins/VirtualProduction/Takes/Source/TakeRecorderSources/Private/TakeRecorderSourcesModule.cpp b/Engine/Plugins/VirtualProduction/Takes/Source/TakeRecorderSources/Private/TakeRecorderSourcesModule.cpp
--- a/Engine/Plugins/VirtualProduction/Takes/Source/TakeRecorderSources/Private/TakeRecorderSourcesModule.cpp
+++ b/Engine/Plugins/VirtualProduction/Takes/Source/TakeRecorderSources/Private/TakeRecorderSourcesModule.cpp
@@ -14,8 +14,9 @@
 #include "TakeRecorderSources.h"
 #include "TakeRecorderSettings.h"
 #include "Features/IModularFeatures.h"
+#include "DragAndDrop/CompositeDragDropOp.h"
 #include "DragAndDrop/ActorDragDropOp.h"
-#include "SceneOutlinerDragDrop.h"
+#include "DragAndDrop/FolderDragDropOp.h"
 #include "EngineUtils.h"
 #include "Algo/Sort.h"
 #include "ScopedTransaction.h"
@@ -181,7 +182,7 @@ struct FActorTakeRecorderDropHandler : ITakeRecorderDropHandler
 			{
 				return StaticCastSharedPtr<FActorDragDropOp>(InOperation)->Actors.Num() > 0;
 			}
-			else if (InOperation->IsOfType<FSceneOutlinerDragDropOp>())
+			else if (InOperation->IsOfType<FCompositeDragDropOp>())
 			{
 				return true;
 			}
@@ -200,9 +201,9 @@ struct FActorTakeRecorderDropHandler : ITakeRecorderDropHandler
 		{
 			return TArray<AActor*>();
 		}
-		else if (OperationPtr->IsOfType<FSceneOutlinerDragDropOp>())
+		else if (OperationPtr->IsOfType<FCompositeDragDropOp>())
 		{
-			FSceneOutlinerDragDropOp* OutlinerOp = static_cast<FSceneOutlinerDragDropOp*>(OperationPtr);
+			FCompositeDragDropOp* OutlinerOp = static_cast<FCompositeDragDropOp*>(OperationPtr);
 			FolderDrag = OutlinerOp->GetSubOp<FFolderDragDropOp>().Get();
 			ActorDrag  = OutlinerOp->GetSubOp<FActorDragDropOp>().Get();
 
diff --git a/Engine/Source/Editor/SceneOutliner/Private/ActorBrowsingMode.cpp b/Engine/Source/Editor/SceneOutliner/Private/ActorBrowsingMode.cpp
--- a/Engine/Source/Editor/SceneOutliner/Private/ActorBrowsingMode.cpp
+++ b/Engine/Source/Editor/SceneOutliner/Private/ActorBrowsingMode.cpp
@@ -15,8 +15,10 @@
 #include "HAL/PlatformApplicationMisc.h"
 #include "SceneOutlinerDragDrop.h"
 #include "EditorActorFolders.h"
+#include "EditorFolderUtils.h"
 #include "ActorEditorUtils.h"
 #include "DragAndDrop/ActorDragDropOp.h"
+#include "DragAndDrop/FolderDragDropOp.h"
 #include "Logging/MessageLog.h"
 #include "SSocketChooser.h"
 #include "ActorFolderPickingMode.h"
@@ -892,8 +894,8 @@ FSceneOutlinerDragValidationInfo FActorBrowsingMode::ValidateDrop(const ISceneOu
 			// Iterate over all the folders that have been dragged
 			for (FName DraggedFolder : Payload.GetData<FName>(SceneOutliner::FFolderPathSelector()))
 			{
-				const FName Leaf = SceneOutliner::GetFolderLeafName(DraggedFolder);
-				const FName Parent = SceneOutliner::GetParentPath(DraggedFolder);
+				const FName Leaf = FEditorFolderUtils::GetLeafName(DraggedFolder);
+				const FName Parent = FEditorFolderUtils::GetParentPath(DraggedFolder);
 
 				if (Parent == DestinationPath)
 				{
diff --git a/Engine/Source/Editor/SceneOutliner/Private/ActorFolderHierarchy.cpp b/Engine/Source/Editor/SceneOutliner/Private/ActorFolderHierarchy.cpp
--- a/Engine/Source/Editor/SceneOutliner/Private/ActorFolderHierarchy.cpp
+++ b/Engine/Source/Editor/SceneOutliner/Private/ActorFolderHierarchy.cpp
@@ -5,6 +5,7 @@
 #include "WorldTreeItem.h"
 #include "ActorFolderTreeItem.h"
 #include "EditorActorFolders.h"
+#include "EditorFolderUtils.h"
 
 FActorFolderHierarchy::FActorFolderHierarchy(ISceneOutlinerMode* InMode, const TWeakObjectPtr<UWorld>& World)
 	: ISceneOutlinerHierarchy(InMode)
@@ -22,7 +23,7 @@ FSceneOutlinerTreeItemPtr FActorFolderHierarchy::FindParent(const ISceneOutliner
 	}
 	else if (const FActorFolderTreeItem* ActorFolderItem = Item.CastTo<FActorFolderTreeItem>())
 	{
-		const FName ParentPath = SceneOutliner::GetParentPath(ActorFolderItem->Path);
+		const FName ParentPath = FEditorFolderUtils::GetParentPath(ActorFolderItem->Path);
 
 		const FSceneOutlinerTreeItemPtr* ParentItem = nullptr;
 		// If the folder has no parent path, it must be parented to the root world
@@ -32,7 +33,7 @@ FSceneOutlinerTreeItemPtr FActorFolderHierarchy::FindParent(const ISceneOutliner
 		}
 		else
 		{
-			ParentItem = Items.Find(SceneOutliner::GetParentPath(ActorFolderItem->Path));
+			ParentItem = Items.Find(FEditorFolderUtils::GetParentPath(ActorFolderItem->Path));
 		}
 
 		if (ParentItem)
@@ -54,7 +55,7 @@ void FActorFolderHierarchy::FindChildren(const ISceneOutlinerTreeItem& Item, con
 			{
 				if (const FFolderTreeItem* FolderItem = (*PotentialChild)->CastTo<FFolderTreeItem>())
 				{
-					if (SceneOutliner::GetParentPath(FolderItem->Path).IsNone())
+					if (FEditorFolderUtils::GetParentPath(FolderItem->Path).IsNone())
 					{
 						OutChildren.Add(*PotentialChild);
 					}
@@ -69,7 +70,7 @@ void FActorFolderHierarchy::FindChildren(const ISceneOutlinerTreeItem& Item, con
 		{
 			if (const FFolderTreeItem* PotentialChild = Pair.Value->CastTo<FFolderTreeItem>())
 			{
-				if (SceneOutliner::PathIsChildOf(PotentialChild->Path, FolderItem->Path))
+				if (FEditorFolderUtils::PathIsChildOf(PotentialChild->Path, FolderItem->Path))
 				{
 					OutChildren.Add(Pair.Value);
 				}
@@ -114,7 +115,7 @@ void FActorFolderHierarchy::CreateChildren(const FSceneOutlinerTreeItemPtr& Item
 		// and manually check the path to know if it a child
 		for (const auto& Pair : FActorFolders::Get().GetFolderPropertiesForWorld(*(FolderItem->World)))
 		{
-			if (SceneOutliner::PathIsChildOf(Pair.Key, FolderItem->Path))
+			if (FEditorFolderUtils::PathIsChildOf(Pair.Key, FolderItem->Path))
 			{
 				if (FSceneOutlinerTreeItemPtr NewFolderItem = Mode->CreateItemFor<FActorFolderTreeItem>(FActorFolderTreeItem(Pair.Key, FolderItem->World), true))
 				{
@@ -133,7 +134,7 @@ FSceneOutlinerTreeItemPtr FActorFolderHierarchy::CreateParentItem(const FSceneOu
 	}
 	else if (const FActorFolderTreeItem* FolderTreeItem = Item->CastTo<FActorFolderTreeItem>())
 	{
-		const FName ParentPath = SceneOutliner::GetParentPath(FolderTreeItem->Path);
+		const FName ParentPath = FEditorFolderUtils::GetParentPath(FolderTreeItem->Path);
 		if (ParentPath.IsNone())
 		{
 			UWorld* OwningWorld = FolderTreeItem->World.Get();
diff --git a/Engine/Source/Editor/SceneOutliner/Private/ActorFolderTreeItem.cpp b/Engine/Source/Editor/SceneOutliner/Private/ActorFolderTreeItem.cpp
--- a/Engine/Source/Editor/SceneOutliner/Private/ActorFolderTreeItem.cpp
+++ b/Engine/Source/Editor/SceneOutliner/Private/ActorFolderTreeItem.cpp
@@ -5,6 +5,7 @@
 #include "SSceneOutliner.h"
 #include "ActorEditorUtils.h"
 #include "EditorActorFolders.h"
+#include "EditorFolderUtils.h"
 #include "Widgets/Text/SInlineEditableTextBlock.h"
 #include "ISceneOutliner.h"
 #include "ISceneOutlinerMode.h"
@@ -144,7 +145,7 @@ private:
 		}
 
 		// Validate that this folder doesn't exist already
-		FName NewPath = SceneOutliner::GetParentPath(TreeItem->Path);
+		FName NewPath = FEditorFolderUtils::GetParentPath(TreeItem->Path);
 		if (NewPath.IsNone())
 		{
 			NewPath = FName(*LabelString);
@@ -169,7 +170,7 @@ private:
 		if (TreeItem.IsValid() && !InLabel.ToString().Equals(TreeItem->LeafName.ToString(), ESearchCase::CaseSensitive))
 		{
 			// Rename the item
-			FName NewPath = SceneOutliner::GetParentPath(TreeItem->Path);
+			FName NewPath = FEditorFolderUtils::GetParentPath(TreeItem->Path);
 			if (NewPath.IsNone())
 			{
 				NewPath = FName(*InLabel.ToString());
diff --git a/Engine/Source/Editor/SceneOutliner/Private/ActorHierarchy.cpp b/Engine/Source/Editor/SceneOutliner/Private/ActorHierarchy.cpp
--- a/Engine/Source/Editor/SceneOutliner/Private/ActorHierarchy.cpp
+++ b/Engine/Source/Editor/SceneOutliner/Private/ActorHierarchy.cpp
@@ -13,6 +13,7 @@
 #include "LevelUtils.h"
 #include "GameFramework/WorldSettings.h"
 #include "EditorActorFolders.h"
+#include "EditorFolderUtils.h"
 #include "LevelInstance/LevelInstanceActor.h"
 #include "LevelInstance/LevelInstanceSubsystem.h"
 
@@ -99,7 +100,7 @@ void FActorHierarchy::FindChildren(const ISceneOutlinerTreeItem& Item, const TMa
 				{
 					if (const FFolderTreeItem* FolderItem = (*PotentialChild)->CastTo<FFolderTreeItem>())
 					{
-						if (SceneOutliner::GetParentPath(FolderItem->Path).IsNone())
+						if (FEditorFolderUtils::GetParentPath(FolderItem->Path).IsNone())
 						{
 							OutChildren.Add(*PotentialChild);
 						}
@@ -153,7 +154,7 @@ void FActorHierarchy::FindChildren(const ISceneOutlinerTreeItem& Item, const TMa
 		// Search through all folders and find any folders with paths which are children of this item
 		for (const auto& Pair : FActorFolders::Get().GetFolderPropertiesForWorld(*RepresentingWorld))
 		{
-			if (SceneOutliner::PathIsChildOf(Pair.Key, FolderItem->Path))
+			if (FEditorFolderUtils::PathIsChildOf(Pair.Key, FolderItem->Path))
 			{
 				if (const FSceneOutlinerTreeItemPtr* Child = Items.Find(Pair.Key))
 				{
@@ -226,7 +227,7 @@ FSceneOutlinerTreeItemPtr FActorHierarchy::FindParent(const ISceneOutlinerTreeIt
 		// We should never call FindParents on a folder item if folders are not being shown
 		check(Mode->ShouldShowFolders());
 
-		const FName ParentPath = SceneOutliner::GetParentPath(FolderItem->Path);
+		const FName ParentPath = FEditorFolderUtils::GetParentPath(FolderItem->Path);
 
 		const FSceneOutlinerTreeItemPtr* ParentItem = nullptr;
 		// If the folder has no parent path, it must be parented to the root world
@@ -236,7 +237,7 @@ FSceneOutlinerTreeItemPtr FActorHierarchy::FindParent(const ISceneOutlinerTreeIt
 		}
 		else
 		{
-			ParentItem = Items.Find(SceneOutliner::GetParentPath(FolderItem->Path));
+			ParentItem = Items.Find(FEditorFolderUtils::GetParentPath(FolderItem->Path));
 		}
 			
 		if (ParentItem)
@@ -396,7 +397,7 @@ void FActorHierarchy::CreateChildren(const FSceneOutlinerTreeItemPtr& Item, TArr
 
 		for (const auto& Pair : FActorFolders::Get().GetFolderPropertiesForWorld(*FolderItem->World))
 		{
-			if (SceneOutliner::PathIsChildOf(Pair.Key, FolderItem->Path))
+			if (FEditorFolderUtils::PathIsChildOf(Pair.Key, FolderItem->Path))
 			{
 				if (FSceneOutlinerTreeItemPtr NewFolderItem = Mode->CreateItemFor<FActorFolderTreeItem>(FActorFolderTreeItem(Pair.Key, FolderItem->World)))
 				{
@@ -459,7 +460,7 @@ FSceneOutlinerTreeItemPtr FActorHierarchy::CreateParentItem(const FSceneOutliner
 	{
 		check(Mode->ShouldShowFolders());
 			
-		const FName ParentPath = SceneOutliner::GetParentPath(FolderTreeItem->Path);
+		const FName ParentPath = FEditorFolderUtils::GetParentPath(FolderTreeItem->Path);
 		if (ParentPath.IsNone())
 		{
 			UWorld* OwningWorld = FolderTreeItem->World.Get();
diff --git a/Engine/Source/Editor/SceneOutliner/Private/FolderTreeItem.cpp b/Engine/Source/Editor/SceneOutliner/Private/FolderTreeItem.cpp
--- a/Engine/Source/Editor/SceneOutliner/Private/FolderTreeItem.cpp
+++ b/Engine/Source/Editor/SceneOutliner/Private/FolderTreeItem.cpp
@@ -11,6 +11,7 @@
 
 #include "ActorEditorUtils.h"
 #include "EditorActorFolders.h"
+#include "EditorFolderUtils.h"
 
 #define LOCTEXT_NAMESPACE "SceneOutliner_FolderTreeItem"
 
@@ -35,14 +36,14 @@ const FSceneOutlinerTreeItemType FFolderTreeItem::Type(&ISceneOutlinerTreeItem::
 FFolderTreeItem::FFolderTreeItem(FName InPath)
 	: ISceneOutlinerTreeItem(Type)
 	, Path(InPath)
-	, LeafName(SceneOutliner::GetFolderLeafName(InPath))
+	, LeafName(FEditorFolderUtils::GetLeafName(InPath))
 {
 }
 
 FFolderTreeItem::FFolderTreeItem(FName InPath, FSceneOutlinerTreeItemType InType)
 	: ISceneOutlinerTreeItem(InType)
 	, Path(InPath)
-	, LeafName(SceneOutliner::GetFolderLeafName(InPath))
+	, LeafName(FEditorFolderUtils::GetLeafName(InPath))
 {
 }
 
diff --git a/Engine/Source/Editor/SceneOutliner/Private/SSceneOutliner.cpp b/Engine/Source/Editor/SceneOutliner/Private/SSceneOutliner.cpp
--- a/Engine/Source/Editor/SceneOutliner/Private/SSceneOutliner.cpp
+++ b/Engine/Source/Editor/SceneOutliner/Private/SSceneOutliner.cpp
@@ -33,11 +33,12 @@
 #include "Widgets/Layout/SSeparator.h"
 #include "Widgets/SOverlay.h"
 #include "SceneOutlinerMenuContext.h"
+#include "ISceneOutlinerMode.h"
+#include "FolderTreeItem.h"
+#include "EditorFolderUtils.h"
 
 #include "Framework/Notifications/NotificationManager.h"
 #include "Widgets/Notifications/SNotificationList.h"
-#include "ISceneOutlinerMode.h"
-#include "FolderTreeItem.h"
 
 DEFINE_LOG_CATEGORY_STATIC(LogSceneOutliner, Log, All);
 
@@ -1006,7 +1007,7 @@ void SSceneOutliner::AddMoveToFolderOutliner(UToolMenu* Menu) const
 		auto FilterOutChildFolders = [](FName Path, TSharedRef<TSet<FName>> ExcludedParents){
 			for (const auto& Parent : *ExcludedParents)
 			{
-				if (Path == Parent || SceneOutliner::PathIsChildOf(Path, Parent))
+				if (Path == Parent || FEditorFolderUtils::PathIsChildOf(Path, Parent))
 				{
 					return false;
 				}
@@ -1246,8 +1247,8 @@ void SSceneOutliner::PasteFoldersEnd()
 	TMap<FName, FName> FolderMap;
 	for (FName Folder : CacheFoldersEdit)
 	{
-		FName ParentPath = SceneOutliner::GetParentPath(Folder);
-		FName LeafName = SceneOutliner::GetFolderLeafName(Folder);
+		FName ParentPath = FEditorFolderUtils::GetParentPath(Folder);
+		FName LeafName = FEditorFolderUtils::GetLeafName(Folder);
 		if (LeafName != TEXT(""))
 		{
 			if (FName* NewParentPath = FolderMap.Find(ParentPath))
@@ -1360,10 +1361,10 @@ void SSceneOutliner::DeleteFoldersEnd()
 			if (Folder)
 			{
 				// Find lowest parent not being deleted, for reparenting children of current folder
-				FName NewParentPath = SceneOutliner::GetParentPath(Folder->Path);
+				FName NewParentPath = FEditorFolderUtils::GetParentPath(Folder->Path);
 				while (!NewParentPath.IsNone() && CacheFoldersDelete.FindByPredicate(FMatchName(NewParentPath)))
 				{
-					NewParentPath = SceneOutliner::GetParentPath(NewParentPath);
+					NewParentPath = FEditorFolderUtils::GetParentPath(NewParentPath);
 				}
 
 				Folder->Delete(NewParentPath);
@@ -1587,7 +1588,7 @@ void SSceneOutliner::OnHierarchyChangedEvent(FSceneOutlinerHierarchyChangedData
 			// Now change the path and put it back in the map with its new ID
 			auto Folder = StaticCastSharedPtr<FFolderTreeItem>(Item);
 			Folder->Path = Event.NewPath;
-			Folder->LeafName = SceneOutliner::GetFolderLeafName(Event.NewPath);
+			Folder->LeafName = FEditorFolderUtils::GetLeafName(Event.NewPath);
 
 			TreeItemMap.Add(Item->GetID(), Item);
 
diff --git a/Engine/Source/Editor/SceneOutliner/Private/SceneOutlinerDragDrop.cpp b/Engine/Source/Editor/SceneOutliner/Private/SceneOutlinerDragDrop.cpp
--- a/Engine/Source/Editor/SceneOutliner/Private/SceneOutlinerDragDrop.cpp
+++ b/Engine/Source/Editor/SceneOutliner/Private/SceneOutlinerDragDrop.cpp
@@ -6,21 +6,6 @@
 
 #define LOCTEXT_NAMESPACE "SSceneOutliner"
 
-void FFolderDragDropOp::Init(TArray<FName> InFolders)
-{
-	Folders = MoveTemp(InFolders);
-
-	CurrentIconBrush = FEditorStyle::Get().GetBrush(TEXT("SceneOutliner.FolderClosed"));
-	if (Folders.Num() == 1)
-	{
-		CurrentHoverText = FText::FromName(SceneOutliner::GetFolderLeafName(Folders[0]));
-	}
-	else
-	{
-		CurrentHoverText = FText::Format(NSLOCTEXT("FFolderDragDropOp", "FormatFolders", "{0} Folders"), FText::AsNumber(Folders.Num()));
-	}
-}
-
 FSceneOutlinerDragDropOp::FSceneOutlinerDragDropOp()
 	: OverrideText()
 	, OverrideIcon(nullptr)
diff --git a/Engine/Source/Editor/SceneOutliner/Private/SceneOutlinerStandaloneTypes.cpp b/Engine/Source/Editor/SceneOutliner/Private/SceneOutlinerStandaloneTypes.cpp
--- a/Engine/Source/Editor/SceneOutliner/Private/SceneOutlinerStandaloneTypes.cpp
+++ b/Engine/Source/Editor/SceneOutliner/Private/SceneOutlinerStandaloneTypes.cpp
@@ -52,34 +52,4 @@ bool FSceneOutlinerCommonLabelData::CanExecuteRenameRequest(const ISceneOutliner
 	return false;
 }
 
-namespace SceneOutliner
-{
-	/** Parse a new path (including leaf-name) into this tree item. Does not do any notification */
-	FName GetFolderLeafName(FName InPath)
-	{
-		FString PathString = InPath.ToString();
-		int32 LeafIndex = 0;
-		if (PathString.FindLastChar('/', LeafIndex))
-		{
-			return FName(*PathString.RightChop(LeafIndex + 1));
-		}
-		else
-		{
-			return InPath;
-		}
-	}
-
-	bool PathIsChildOf(const FName& PotentialChild, const FName& Parent)
-	{
-		const FString ParentString = Parent.ToString();
-		const FString ChildString = PotentialChild.ToString();
-		const int32 ParentLen = ParentString.Len();
-		return
-			ChildString.Len() > ParentLen &&
-			ChildString[ParentLen] == '/' &&
-			ChildString.Left(ParentLen) == ParentString;
-	}
-
-}	// namespace SceneOutliner
-
 #undef LOCTEXT_NAMESPACE
diff --git a/Engine/Source/Editor/SceneOutliner/Public/SceneOutlinerDragDrop.h b/Engine/Source/Editor/SceneOutliner/Public/SceneOutlinerDragDrop.h
--- a/Engine/Source/Editor/SceneOutliner/Public/SceneOutlinerDragDrop.h
+++ b/Engine/Source/Editor/SceneOutliner/Public/SceneOutlinerDragDrop.h
@@ -150,17 +150,6 @@ struct FSceneOutlinerDragValidationInfo
 	}
 };
 
-/* A drag/drop operation when dragging folders in the scene outliner */
-struct FFolderDragDropOp: public FDecoratedDragDropOp
-{
-	DRAG_DROP_OPERATOR_TYPE(FFolderDragDropOp, FDecoratedDragDropOp)
-
-	/** Array of folders that we are dragging */
-	TArray<FName> Folders;
-
-	void Init(TArray<FName> InFolders);
-};
-
 /** A drag/drop operation that was started from the scene outliner */
 struct FSceneOutlinerDragDropOp : public FCompositeDragDropOp
 {
diff --git a/Engine/Source/Editor/SceneOutliner/Public/SceneOutlinerStandaloneTypes.h b/Engine/Source/Editor/SceneOutliner/Public/SceneOutlinerStandaloneTypes.h
--- a/Engine/Source/Editor/SceneOutliner/Public/SceneOutlinerStandaloneTypes.h
+++ b/Engine/Source/Editor/SceneOutliner/Public/SceneOutlinerStandaloneTypes.h
@@ -8,23 +8,6 @@
 #include "Templates/MaxSizeof.h"
 #include "SceneOutlinerFwd.h"
 
-namespace SceneOutliner
-{
-	/** Folder-specific helper functions */
-
-	/** Parse a new path (including leaf-name) into this tree item. Does not do any notification */
-	FName GetFolderLeafName(FName InPath);
-
-	/** Get the parent path for the specified folder path */
-	FORCEINLINE FName GetParentPath(FName Path)
-	{
-		return FName(*FPaths::GetPath(Path.ToString()));
-	}
-
-	bool PathIsChildOf(const FName& PotentialChild, const FName& Parent);
-}	// namespace SceneOutliner
-
-
 /** Variant type that defines an identifier for a tree item. Assumes 'trivial relocatability' as with many unreal containers. */
 struct FSceneOutlinerTreeItemID
 {
diff --git a/Engine/Source/Editor/UnrealEd/Private/EditorActorFolders.cpp b/Engine/Source/Editor/UnrealEd/Private/EditorActorFolders.cpp
--- a/Engine/Source/Editor/UnrealEd/Private/EditorActorFolders.cpp
+++ b/Engine/Source/Editor/UnrealEd/Private/EditorActorFolders.cpp
@@ -11,6 +11,7 @@
 #include "Engine/Selection.h"
 #include "EngineUtils.h"
 #include "Editor.h"
+#include "EditorFolderUtils.h"
 #include "ScopedTransaction.h"
 
 #define LOCTEXT_NAMESPACE "FActorFolders"
@@ -181,15 +182,6 @@ void FActorFolders::OnActorFolderChanged(const AActor* InActor, FName OldPath)
 	}
 }
 
-bool FActorFolders::PathIsChildOf(const FString& InPotentialChild, const FString& InParent)
-{
-	const int32 ParentLen = InParent.Len();
-	return
-		InPotentialChild.Len() > ParentLen &&
-		InPotentialChild[ParentLen] == '/' &&
-		InPotentialChild.Left(ParentLen) == InParent;
-}
-
 void FActorFolders::RebuildFolderListForWorld(UWorld& InWorld)
 {
 	if (TemporaryWorldFolders.Contains(&InWorld))
@@ -450,7 +442,7 @@ bool FActorFolders::RenameFolderInWorld(UWorld& World, FName OldPath, FName NewP
 	const FString OldPathString = OldPath.ToString();
 	const FString NewPathString = NewPath.ToString();
 
-	if (OldPath.IsNone() || OldPathString.Equals(NewPathString) || PathIsChildOf(NewPathString, OldPathString))
+	if (OldPath.IsNone() || OldPathString.Equals(NewPathString) || FEditorFolderUtils::PathIsChildOf(NewPathString, OldPathString))
 	{
 		return false;
 	}
@@ -470,7 +462,7 @@ bool FActorFolders::RenameFolderInWorld(UWorld& World, FName OldPath, FName NewP
 		auto Path = Pair.Key;
 
 		const FString FolderPath = Path.ToString();
-		if (OldPath == Path || PathIsChildOf(FolderPath, OldPathString))
+		if (OldPath == Path || FEditorFolderUtils::PathIsChildOf(FolderPath, OldPathString))
 		{
 			const FName NewFolder = OldPathToNewPath(OldPathString, NewPathString, FolderPath);
 			
@@ -501,7 +493,7 @@ bool FActorFolders::RenameFolderInWorld(UWorld& World, FName OldPath, FName NewP
 				OnFolderCreate.Broadcast(World, NewFolder);
 			}
 
-			// case insensive compare as we don't want to remove the folder if it has the same name
+			// case insensitive compare as we don't want to remove the folder if it has the same name
 			if (Path != NewFolder)
 			{
 				RenamedFolders.Add(Path);
@@ -521,12 +513,12 @@ bool FActorFolders::RenameFolderInWorld(UWorld& World, FName OldPath, FName NewP
 			continue;
 		}
 
-		if (OldActorPath == OldPath || PathIsChildOf(OldActorPath.ToString(), OldPathString))
+		if (OldActorPath == OldPath || FEditorFolderUtils::PathIsChildOf(OldActorPath.ToString(), OldPathString))
 		{
 			ActorIt->SetFolderPath_Recursively(OldPathToNewPath(OldPathString, NewPathString, OldActorPath.ToString()));
 			const FName& NewActorPath = ActorIt->GetFolderPath();
 
-			// case insensive compare as we don't want to remove the folder if it has the same name
+			// case insensitive compare as we don't want to remove the folder if it has the same name
 			if (OldActorPath != NewActorPath)
 			{
 				RenamedFolders.Add(OldActorPath);
diff --git a/Engine/Source/Editor/UnrealEd/Private/EditorFolderUtils.cpp b/Engine/Source/Editor/UnrealEd/Private/EditorFolderUtils.cpp
new file mode 100644
--- /dev/null
+++ b/Engine/Source/Editor/UnrealEd/Private/EditorFolderUtils.cpp
@@ -0,0 +1,31 @@
+// Copyright Epic Games, Inc. All Rights Reserved.
+
+#include "EditorFolderUtils.h"
+
+FName FEditorFolderUtils::GetLeafName(const FName& InPath)
+{
+	FString PathString = InPath.ToString();
+	int32 LeafIndex = 0;
+	if (PathString.FindLastChar('/', LeafIndex))
+	{
+		return FName(*PathString.RightChop(LeafIndex + 1));
+	}
+	else
+	{
+		return InPath;
+	}
+}
+
+bool FEditorFolderUtils::PathIsChildOf(const FString& PotentialChild, const FString& Parent)
+{
+	const int32 ParentLen = Parent.Len();
+	return
+		PotentialChild.Len() > ParentLen&&
+		PotentialChild[ParentLen] == '/' &&
+		PotentialChild.Left(ParentLen) == Parent;
+}
+
+bool FEditorFolderUtils::PathIsChildOf(const FName& PotentialChild, const FName& Parent)
+{
+	return PathIsChildOf(PotentialChild.ToString(), Parent.ToString());
+}
\ No newline at end of file
diff --git a/Engine/Source/Editor/UnrealEd/Public/DragAndDrop/FolderDragDropOp.h b/Engine/Source/Editor/UnrealEd/Public/DragAndDrop/FolderDragDropOp.h
new file mode 100644
--- /dev/null
+++ b/Engine/Source/Editor/UnrealEd/Public/DragAndDrop/FolderDragDropOp.h
@@ -0,0 +1,32 @@
+// Copyright Epic Games, Inc. All Rights Reserved.
+
+#pragma once
+
+#include "CoreMinimal.h"
+#include "Input/DragAndDrop.h"
+#include "DragAndDrop/DecoratedDragDropOp.h"
+#include "EditorFolderUtils.h"
+
+/* A drag/drop operation when dragging actor folders */
+struct FFolderDragDropOp : public FDecoratedDragDropOp
+{
+	DRAG_DROP_OPERATOR_TYPE(FFolderDragDropOp, FDecoratedDragDropOp)
+
+	/** Array of folders that we are dragging */
+	TArray<FName> Folders;
+
+	void Init(TArray<FName> InFolders)
+	{
+		Folders = MoveTemp(InFolders);
+
+		CurrentIconBrush = FEditorStyle::Get().GetBrush(TEXT("SceneOutliner.FolderClosed"));
+		if (Folders.Num() == 1)
+		{
+			CurrentHoverText = FText::FromName(FEditorFolderUtils::GetLeafName(Folders[0]));
+		}
+		else
+		{
+			CurrentHoverText = FText::Format(NSLOCTEXT("FFolderDragDropOp", "FormatFolders", "{0} Folders"), FText::AsNumber(Folders.Num()));
+		}
+	}
+};
\ No newline at end of file
diff --git a/Engine/Source/Editor/UnrealEd/Public/EditorActorFolders.h b/Engine/Source/Editor/UnrealEd/Public/EditorActorFolders.h
--- a/Engine/Source/Editor/UnrealEd/Public/EditorActorFolders.h
+++ b/Engine/Source/Editor/UnrealEd/Public/EditorActorFolders.h
@@ -75,9 +75,6 @@ struct UNREALED_API FActorFolders : public FGCObject
 	static FOnActorFolderMove 	OnFolderMove;
 	static FOnActorFolderDelete OnFolderDelete;
 
-	/** Check if the specified path is a child of the specified parent */
-	static bool PathIsChildOf(const FString& InPotentialChild, const FString& InParent);
-
 	/** Get a map of folder properties for the specified world (map of folder path -> properties) */
 	const TMap<FName, FActorFolderProps>& GetFolderPropertiesForWorld(UWorld& InWorld);
 
diff --git a/Engine/Source/Editor/UnrealEd/Public/EditorFolderUtils.h b/Engine/Source/Editor/UnrealEd/Public/EditorFolderUtils.h
new file mode 100644
--- /dev/null
+++ b/Engine/Source/Editor/UnrealEd/Public/EditorFolderUtils.h
@@ -0,0 +1,22 @@
+// Copyright Epic Games, Inc. All Rights Reserved.
+
+#pragma once
+
+#include "CoreMinimal.h"
+
+class UNREALED_API FEditorFolderUtils
+{
+public:
+	/** Get the leaf name of a specified folder path */
+	static FName GetLeafName(const FName& InPath);
+
+	/** Get the parent path for the specified folder path */
+	FORCEINLINE static FName GetParentPath(const FName& InPath)
+	{
+		return FName(*FPaths::GetPath(InPath.ToString()));
+	}
+
+	/** Check if the specified path is a child of the specified parent */
+	static bool PathIsChildOf(const FString& PotentialChild, const FString& Parent);
+	static bool PathIsChildOf(const FName& PotentialChild, const FName& Parent);
+};
\ No newline at end of file
