Fixed FSceneOutlinerDragDropOp type info

diff --git a/Engine/Source/Editor/SceneOutliner/Public/SceneOutlinerDragDrop.h b/Engine/Source/Editor/SceneOutliner/Public/SceneOutlinerDragDrop.h
--- a/Engine/Source/Editor/SceneOutliner/Public/SceneOutlinerDragDrop.h
+++ b/Engine/Source/Editor/SceneOutliner/Public/SceneOutlinerDragDrop.h
@@ -153,7 +153,7 @@ struct FSceneOutlinerDragValidationInfo
 /** A drag/drop operation that was started from the scene outliner */
 struct FSceneOutlinerDragDropOp : public FCompositeDragDropOp
 {
-	DRAG_DROP_OPERATOR_TYPE(FSceneOutlinerDragDropOp, FDecoratedDragDropOp);
+	DRAG_DROP_OPERATOR_TYPE(FSceneOutlinerDragDropOp, FCompositeDragDropOp);
 		
 	FSceneOutlinerDragDropOp();
 
