WorldOutliner: Scroll selected item into view when FullRefresh finishes processing pendingoperations

diff --git a/Engine/Source/Editor/SceneOutliner/Private/SSceneOutliner.cpp b/Engine/Source/Editor/SceneOutliner/Private/SSceneOutliner.cpp
--- a/Engine/Source/Editor/SceneOutliner/Private/SSceneOutliner.cpp
+++ b/Engine/Source/Editor/SceneOutliner/Private/SSceneOutliner.cpp
@@ -59,6 +59,7 @@ void SSceneOutliner::Construct(const FArguments& InArgs, const FSceneOutlinerIni
 	Mode = InInitOptions.ModeFactory.Execute(this);
 	check(Mode);
 
+	bProcessingFullRefresh = false;
 	bFullRefresh = true;
 	bNeedsRefresh = true;
 	bNeedsColumRefresh = true;
@@ -462,6 +463,22 @@ void SSceneOutliner::Populate()
 	bool bFinalSort = false;
 	if (PendingOperations.Num() == 0)
 	{
+		// When done processing a FullRefresh Scroll to First item in selection as it may have been
+		// scrolled out of view by the Refresh
+		if (bProcessingFullRefresh)
+		{
+			FSceneOutlinerItemSelection ItemSelection(*OutlinerTreeView);
+			if (ItemSelection.Num() > 0)
+			{
+				FSceneOutlinerTreeItemPtr ItemToScroll = ItemSelection.SelectedItems[0].Pin();
+				if (ItemToScroll)
+				{
+					ScrollItemIntoView(ItemToScroll);
+				}
+			}
+		}
+
+		bProcessingFullRefresh = false;
 		// We're fully refreshed now.
 		NewItemActions.Empty();
 		bNeedsRefresh = false;
@@ -535,6 +552,7 @@ void SSceneOutliner::RepopulateEntireTree()
 	{
 		AddPendingItem(Item);
 	}
+	bProcessingFullRefresh = PendingOperations.Num() > 0;
 
 	Refresh();
 }
diff --git a/Engine/Source/Editor/SceneOutliner/Public/SSceneOutliner.h b/Engine/Source/Editor/SceneOutliner/Public/SSceneOutliner.h
--- a/Engine/Source/Editor/SceneOutliner/Public/SSceneOutliner.h
+++ b/Engine/Source/Editor/SceneOutliner/Public/SSceneOutliner.h
@@ -715,6 +715,9 @@ private:
 		item set has changed in some way. */
 	uint8 bNeedsRefresh : 1;
 
+	/** Processing a full refresh until pending items are processed */
+	uint8 bProcessingFullRefresh : 1;
+
 	/** true if the Scene Outliner should do a full refresh. */
 	uint8 bFullRefresh : 1;
 
