Fix the world outliner "level" column for external actors.

diff --git a/Engine/Source/Editor/SceneOutliner/Private/SceneOutlinerActorInfoColumn.cpp b/Engine/Source/Editor/SceneOutliner/Private/SceneOutlinerActorInfoColumn.cpp
--- a/Engine/Source/Editor/SceneOutliner/Private/SceneOutlinerActorInfoColumn.cpp
+++ b/Engine/Source/Editor/SceneOutliner/Private/SceneOutlinerActorInfoColumn.cpp
@@ -40,7 +40,7 @@ struct FGetInfo
 				return Actor->GetClass()->GetName();
 
 			case SceneOutliner::ECustomColumnMode::Level:
-				return FPackageName::GetShortName(Actor->GetOutermost()->GetName());
+				return FPackageName::GetShortName(Actor->GetOutermostObject()->GetName());
 
 			case SceneOutliner::ECustomColumnMode::Socket:
 				return Actor->GetAttachParentSocketName().ToString();
