Fixed issue where an actor selected in the World Outliner can't be picked from an ActorPicker

diff --git a/Engine/Source/Editor/SceneOutliner/Public/ActorPickingMode.h b/Engine/Source/Editor/SceneOutliner/Public/ActorPickingMode.h
--- a/Engine/Source/Editor/SceneOutliner/Public/ActorPickingMode.h
+++ b/Engine/Source/Editor/SceneOutliner/Public/ActorPickingMode.h
@@ -17,6 +17,9 @@ public:
 	virtual void OnFilterTextCommited(FSceneOutlinerItemSelection& Selection, ETextCommit::Type CommitType) override;
 
 	virtual void CreateViewContent(FMenuBuilder& MenuBuilder) override;
+
+	// Don't synchronize world selection
+	virtual void SynchronizeSelection() override {};
 public:
 	virtual bool ShowViewButton() const override { return true; }
 private:
