Speculative fix for World Outliner: Context Menu > Move To - Crash

diff --git a/Engine/Source/Editor/SceneOutliner/Private/SSceneOutliner.cpp b/Engine/Source/Editor/SceneOutliner/Private/SSceneOutliner.cpp
--- a/Engine/Source/Editor/SceneOutliner/Private/SSceneOutliner.cpp
+++ b/Engine/Source/Editor/SceneOutliner/Private/SSceneOutliner.cpp
@@ -1039,31 +1039,22 @@ void SSceneOutliner::AddMoveToFolderOutliner(UToolMenu* Menu) const
 	{
 		struct FFilterRoot : public FSceneOutlinerFilter 
 		{
-			FFilterRoot(TFunctionRef<bool(const ISceneOutlinerTreeItem&)> InFindParent) 
+			FFilterRoot(const SSceneOutliner& InSceneOutliner) 
 				: FSceneOutlinerFilter(FSceneOutlinerFilter::EDefaultBehaviour::Pass) 
-				, FindParent(InFindParent)
+				, SceneOutliner(InSceneOutliner)
 			{}
 			virtual bool PassesFilter(const ISceneOutlinerTreeItem& Item) const override
 			{
+				FSceneOutlinerTreeItemPtr Parent = SceneOutliner.FindParent(Item);
 				// if item has no parent, it is a root item and should be filtered out
-				if (!FindParent(Item))
+				if (!Parent.IsValid())
 				{
 					return false;
 				}
 				return DefaultBehaviour == FSceneOutlinerFilter::EDefaultBehaviour::Pass;
 			}
 
-			TFunctionRef<bool(const ISceneOutlinerTreeItem&)> FindParent;
-		};
-
-		auto FindParent = [this](const ISceneOutlinerTreeItem& Item)
-		{
-			FSceneOutlinerTreeItemPtr Parent = Mode->GetHierarchy()->FindParent(Item, TreeItemMap);
-			if (!Parent)
-			{
-				Parent = Mode->GetHierarchy()->FindParent(Item, PendingTreeItemMap);
-			}
-			return Parent.IsValid();
+			const SSceneOutliner& SceneOutliner;
 		};
 
 		// Filter in/out root items according to whether it is valid to move to/from the root
@@ -1071,7 +1062,7 @@ void SSceneOutliner::AddMoveToFolderOutliner(UToolMenu* Menu) const
 		const bool bMoveToRootValid = Mode->ValidateDrop(FFolderTreeItem(NAME_None), DraggedObjects).IsValid();
 		if (!bMoveToRootValid)
 		{
-			MiniSceneOutlinerInitOptions.Filters->Add(MakeShared<FFilterRoot>(FindParent));
+			MiniSceneOutlinerInitOptions.Filters->Add(MakeShared<FFilterRoot>(*this));
 		}
 	}
 
@@ -1941,4 +1932,14 @@ FReply SSceneOutliner::OnDragOverItem(const FDragDropEvent& Event, const ISceneO
 	return Mode->OnDragOverItem(Event, Item); 
 }
 
+FSceneOutlinerTreeItemPtr SSceneOutliner::FindParent(const ISceneOutlinerTreeItem& InItem) const
+{
+	FSceneOutlinerTreeItemPtr Parent = Mode->GetHierarchy()->FindParent(InItem, TreeItemMap);
+	if (!Parent.IsValid())
+	{
+		Parent = Mode->GetHierarchy()->FindParent(InItem, PendingTreeItemMap);
+	}
+	return Parent;
+}
+
 #undef LOCTEXT_NAMESPACE
diff --git a/Engine/Source/Editor/SceneOutliner/Public/SSceneOutliner.h b/Engine/Source/Editor/SceneOutliner/Public/SSceneOutliner.h
--- a/Engine/Source/Editor/SceneOutliner/Public/SSceneOutliner.h
+++ b/Engine/Source/Editor/SceneOutliner/Public/SSceneOutliner.h
@@ -369,6 +369,11 @@ public:
 		*/
 	virtual FSceneOutlinerItemSelection GetSelection() const { return FSceneOutlinerItemSelection(*OutlinerTreeView); }
 
+	/**
+	 * Returns the parent tree item for a given item if it exists, nullptr otherwise.
+	 */
+	FSceneOutlinerTreeItemPtr FindParent(const ISceneOutlinerTreeItem& InItem) const;
+
 	/**
 		* Add a folder to the selection of the scene outliner
 		* @param FolderName The name of the folder to add to selection
