Fix SelectImmediateChildren crash on mac due to null functor

diff --git a/Engine/Source/Editor/SceneOutliner/Private/ActorBrowsingMode.cpp b/Engine/Source/Editor/SceneOutliner/Private/ActorBrowsingMode.cpp
--- a/Engine/Source/Editor/SceneOutliner/Private/ActorBrowsingMode.cpp
+++ b/Engine/Source/Editor/SceneOutliner/Private/ActorBrowsingMode.cpp
@@ -1229,7 +1229,8 @@ bool FActorBrowsingMode::ReparentItemToFolder(const FName& FolderPath, const FSc
 void FActorBrowsingMode::SelectFoldersDescendants(const TArray<FFolderTreeItem*>& FolderItems, bool bSelectImmediateChildrenOnly)
 {
 	// Ensure that all folder descendants are expanded
-	TFunctionRef<void(const FSceneOutlinerTreeItemPtr&)> RecursiveExpand = [&](const FSceneOutlinerTreeItemPtr&  Item)
+	TFunction<void(const FSceneOutlinerTreeItemPtr&)> RecursiveExpand;
+	RecursiveExpand = [&](const FSceneOutlinerTreeItemPtr&  Item)
 	{
 		SceneOutliner->SetItemExpansion(Item, true);
 			
@@ -1243,7 +1244,8 @@ void FActorBrowsingMode::SelectFoldersDescendants(const TArray<FFolderTreeItem*>
 	};
 
 	// Selects all actor descendants of a folder
-	TFunctionRef<void(const FSceneOutlinerTreeItemPtr&)> RecursiveActorSelect = [&](const FSceneOutlinerTreeItemPtr& Item)
+	TFunction<void(const FSceneOutlinerTreeItemPtr&)> RecursiveActorSelect;
+	RecursiveActorSelect = [&](const FSceneOutlinerTreeItemPtr& Item)
 	{
 		// If the current item is an actor, ensure to select it as well
 		if (FActorTreeItem* ActorItem = Item->CastTo<FActorTreeItem>())
