Cleaned up ActorBrowsingMode::SelectFoldersDescendants

diff --git a/Engine/Source/Editor/SceneOutliner/Private/ActorBrowsingMode.cpp b/Engine/Source/Editor/SceneOutliner/Private/ActorBrowsingMode.cpp
--- a/Engine/Source/Editor/SceneOutliner/Private/ActorBrowsingMode.cpp
+++ b/Engine/Source/Editor/SceneOutliner/Private/ActorBrowsingMode.cpp
@@ -1226,64 +1226,75 @@ bool FActorBrowsingMode::ReparentItemToFolder(const FName& FolderPath, const FSc
 	return false;
 }
 
-void FActorBrowsingMode::SelectFoldersDescendants(const TArray<FFolderTreeItem*>& FolderItems, bool bSelectImmediateChildrenOnly)
+namespace ActorBrowsingModeUtils
 {
-	// Ensure that all folder descendants are expanded
-	TFunction<void(const FSceneOutlinerTreeItemPtr&)> RecursiveExpand;
-	RecursiveExpand = [&](const FSceneOutlinerTreeItemPtr&  Item)
+	static void RecursiveFolderExpandChildren(SSceneOutliner* SceneOutliner, const FSceneOutlinerTreeItemPtr& Item)
 	{
-		SceneOutliner->SetItemExpansion(Item, true);
-			
-		if (!bSelectImmediateChildrenOnly)
+		if (Item.IsValid())
 		{
-			for (const TWeakPtr<ISceneOutlinerTreeItem> Child : Item->GetChildren())
+			for (const TWeakPtr<ISceneOutlinerTreeItem>& Child : Item->GetChildren())
 			{
-				RecursiveExpand(Child.Pin());
+				FSceneOutlinerTreeItemPtr ChildPtr = Child.Pin();
+				SceneOutliner->SetItemExpansion(ChildPtr, true);
+				RecursiveFolderExpandChildren(SceneOutliner, ChildPtr);
 			}
 		}
-	};
+	}
 
-	// Selects all actor descendants of a folder
-	TFunction<void(const FSceneOutlinerTreeItemPtr&)> RecursiveActorSelect;
-	RecursiveActorSelect = [&](const FSceneOutlinerTreeItemPtr& Item)
+	static void RecursiveActorSelect(SSceneOutliner* SceneOutliner, const FSceneOutlinerTreeItemPtr& Item, bool bSelectImmediateChildrenOnly)
 	{
-		// If the current item is an actor, ensure to select it as well
-		if (FActorTreeItem* ActorItem = Item->CastTo<FActorTreeItem>())
-		{
-			if (AActor* Actor = ActorItem->Actor.Get())
-			{
-				GEditor->SelectActor(Actor, true, false);
-			}
-		}
-		// Select all children
-		for (const TWeakPtr<ISceneOutlinerTreeItem> Child : Item->GetChildren())
+		if (Item.IsValid())
 		{
-			if (FActorTreeItem* ActorItem = Child.Pin()->CastTo<FActorTreeItem>())
+			// If the current item is an actor, ensure to select it as well
+			if (FActorTreeItem* ActorItem = Item->CastTo<FActorTreeItem>())
 			{
 				if (AActor* Actor = ActorItem->Actor.Get())
 				{
 					GEditor->SelectActor(Actor, true, false);
 				}
 			}
-			else if (FFolderTreeItem* FolderItem = Child.Pin()->CastTo<FFolderTreeItem>())
-			{
-				SceneOutliner->SetItemSelection(FolderItem->AsShared(), true);
-			}
-
-			if (!bSelectImmediateChildrenOnly)
+			// Select all children
+			for (const TWeakPtr<ISceneOutlinerTreeItem>& Child : Item->GetChildren())
 			{
-				for (const TWeakPtr<ISceneOutlinerTreeItem>& Grandchild : Child.Pin()->GetChildren())
+				FSceneOutlinerTreeItemPtr ChildPtr = Child.Pin();
+				if (ChildPtr.IsValid())
 				{
-					RecursiveActorSelect(Grandchild.Pin());
+					if (FActorTreeItem* ActorItem = ChildPtr->CastTo<FActorTreeItem>())
+					{
+						if (AActor* Actor = ActorItem->Actor.Get())
+						{
+							GEditor->SelectActor(Actor, true, false);
+						}
+					}
+					else if (FFolderTreeItem* FolderItem = ChildPtr->CastTo<FFolderTreeItem>())
+					{
+						SceneOutliner->SetItemSelection(FolderItem->AsShared(), true);
+					}
+
+					if (!bSelectImmediateChildrenOnly)
+					{
+						for (const TWeakPtr<ISceneOutlinerTreeItem>& Grandchild : ChildPtr->GetChildren())
+						{
+							RecursiveActorSelect(SceneOutliner, Grandchild.Pin(), bSelectImmediateChildrenOnly);
+						}
+					}
 				}
 			}
 		}
-	};
+	}
+}
 
+void FActorBrowsingMode::SelectFoldersDescendants(const TArray<FFolderTreeItem*>& FolderItems, bool bSelectImmediateChildrenOnly)
+{
 	// Expand everything before beginning selection
 	for (FFolderTreeItem* Folder : FolderItems)
 	{
-		RecursiveExpand(Folder->AsShared());
+		FSceneOutlinerTreeItemPtr FolderPtr = Folder->AsShared();
+		SceneOutliner->SetItemExpansion(FolderPtr, true);
+		if (!bSelectImmediateChildrenOnly)
+		{
+			ActorBrowsingModeUtils::RecursiveFolderExpandChildren(SceneOutliner, FolderPtr);
+		}
 	}
 
 	// batch selection
@@ -1291,7 +1302,7 @@ void FActorBrowsingMode::SelectFoldersDescendants(const TArray<FFolderTreeItem*>
 
 	for (FFolderTreeItem* Folder : FolderItems)
 	{
-		RecursiveActorSelect(Folder->AsShared());
+		ActorBrowsingModeUtils::RecursiveActorSelect(SceneOutliner, Folder->AsShared(), bSelectImmediateChildrenOnly);
 	}
 
 	GEditor->GetSelectedActors()->EndBatchSelectOperation(/*bNotify*/false);
