Make world outliner icons the proper size

diff --git a/Engine/Source/Editor/SceneOutliner/Private/ActorTreeItem.cpp b/Engine/Source/Editor/SceneOutliner/Private/ActorTreeItem.cpp
--- a/Engine/Source/Editor/SceneOutliner/Private/ActorTreeItem.cpp
+++ b/Engine/Source/Editor/SceneOutliner/Private/ActorTreeItem.cpp
@@ -91,6 +91,7 @@ struct SActorTreeLabel : FSceneOutlinerCommonLabelData, public SCompoundWidget
 						SNew(SImage)
 						.Image(this, &SActorTreeLabel::GetIcon)
 						.ToolTipText(this, &SActorTreeLabel::GetIconTooltip)
+						.ColorAndOpacity(FSlateColor::UseForeground())
 					]
 				]
 
diff --git a/Engine/Source/Editor/SceneOutliner/Public/SceneOutlinerPublicTypes.h b/Engine/Source/Editor/SceneOutliner/Public/SceneOutlinerPublicTypes.h
--- a/Engine/Source/Editor/SceneOutliner/Public/SceneOutlinerPublicTypes.h
+++ b/Engine/Source/Editor/SceneOutliner/Public/SceneOutlinerPublicTypes.h
@@ -145,8 +145,8 @@ public:
 /** Default metrics for outliner tree items */
 struct FSceneOutlinerDefaultTreeItemMetrics
 {
-	static int32	IconSize() { return 18; };
-	static FMargin	IconPadding() { return FMargin(0.f, 0.f, 6.f, 0.f); };
+	static int32	IconSize() { return 16; };
+	static FMargin	IconPadding() { return FMargin(0.f, 1.f, 6.f, 1.f); };
 };
 
 /** A struct which gets, and caches the visibility of a tree item */
