Removed unused function FindChildren in ISceneOutlinerHierarchy

diff --git a/Engine/Source/Editor/SceneOutliner/Private/ActorFolderHierarchy.cpp b/Engine/Source/Editor/SceneOutliner/Private/ActorFolderHierarchy.cpp
--- a/Engine/Source/Editor/SceneOutliner/Private/ActorFolderHierarchy.cpp
+++ b/Engine/Source/Editor/SceneOutliner/Private/ActorFolderHierarchy.cpp
@@ -44,41 +44,6 @@ FSceneOutlinerTreeItemPtr FActorFolderHierarchy::FindParent(const ISceneOutliner
 	return nullptr;
 }
 
-void FActorFolderHierarchy::FindChildren(const ISceneOutlinerTreeItem& Item, const TMap<FSceneOutlinerTreeItemID, FSceneOutlinerTreeItemPtr>& Items, TArray<FSceneOutlinerTreeItemPtr>& OutChildren) const
-{
-	if (const FWorldTreeItem* WorldTreeItem = Item.CastTo<FWorldTreeItem>())
-	{
-		// Find folders which are located at the root
-		for (const auto& Pair : FActorFolders::Get().GetFolderPropertiesForWorld(*WorldTreeItem->World))
-		{
-			if (const FSceneOutlinerTreeItemPtr* PotentialChild = Items.Find(Pair.Key))
-			{
-				if (const FFolderTreeItem* FolderItem = (*PotentialChild)->CastTo<FFolderTreeItem>())
-				{
-					if (FEditorFolderUtils::GetParentPath(FolderItem->Path).IsNone())
-					{
-						OutChildren.Add(*PotentialChild);
-					}
-				}
-			}
-		}
-	}
-	else if (const FFolderTreeItem* FolderItem = Item.CastTo<FFolderTreeItem>())
-	{
-		// Search through all items and see if there is an item with a path which is a child of this one
-		for (const auto& Pair : Items)
-		{
-			if (const FFolderTreeItem* PotentialChild = Pair.Value->CastTo<FFolderTreeItem>())
-			{
-				if (FEditorFolderUtils::PathIsChildOf(PotentialChild->Path, FolderItem->Path))
-				{
-					OutChildren.Add(Pair.Value);
-				}
-			}
-		}
-	}
-}
-
 void FActorFolderHierarchy::CreateWorldChildren(UWorld* World, TArray<FSceneOutlinerTreeItemPtr>& OutItems) const
 {
 	for (const auto& Pair : FActorFolders::Get().GetFolderPropertiesForWorld(*World))
diff --git a/Engine/Source/Editor/SceneOutliner/Private/ActorHierarchy.cpp b/Engine/Source/Editor/SceneOutliner/Private/ActorHierarchy.cpp
--- a/Engine/Source/Editor/SceneOutliner/Private/ActorHierarchy.cpp
+++ b/Engine/Source/Editor/SceneOutliner/Private/ActorHierarchy.cpp
@@ -70,101 +70,6 @@ FActorHierarchy::~FActorHierarchy()
 	}
 }
 
-void FActorHierarchy::FindChildren(const ISceneOutlinerTreeItem& Item, const TMap<FSceneOutlinerTreeItemID, FSceneOutlinerTreeItemPtr>& Items, TArray<FSceneOutlinerTreeItemPtr>& OutChildren) const
-{
-	ULevelInstanceSubsystem* LevelInstanceSubsystem = RepresentingWorld->GetSubsystem<ULevelInstanceSubsystem>();
-
-	check(!bShowingLevelInstances || LevelInstanceSubsystem);
-
-	if (const FWorldTreeItem* WorldTreeItem = Item.CastTo<FWorldTreeItem>())
-	{
-		// All actors in the world are children of the world
-		for (FActorIterator ActorIt(WorldTreeItem->World.Get()); ActorIt; ++ActorIt)
-		{
-			// Actors which have parents are not direct children of a world
-			if (ActorIt->GetAttachParentActor() == nullptr)
-			{
-				if (const FSceneOutlinerTreeItemPtr* ChildItem = Items.Find(*ActorIt))
-				{
-					OutChildren.Add(*ChildItem);
-				}
-			}
-		}
-
-		if (Mode->ShouldShowFolders())
-		{
-			// Find folders which are located at the root
-			for (const auto& Pair : FActorFolders::Get().GetFolderPropertiesForWorld(*WorldTreeItem->World))
-			{
-				if (const FSceneOutlinerTreeItemPtr* PotentialChild = Items.Find(Pair.Key))
-				{
-					if (const FFolderTreeItem* FolderItem = (*PotentialChild)->CastTo<FFolderTreeItem>())
-					{
-						if (FEditorFolderUtils::GetParentPath(FolderItem->Path).IsNone())
-						{
-							OutChildren.Add(*PotentialChild);
-						}
-					}
-				}
-			}
-		}
-	}
-	else if (const FActorTreeItem* ActorTreeItem = Item.CastTo<FActorTreeItem>())
-	{
-		if (ALevelInstance* LevelInstanceActor = Cast<ALevelInstance>(ActorTreeItem->Actor))
-		{
-			LevelInstanceSubsystem->ForEachActorInLevelInstance(LevelInstanceActor, [&OutChildren, &Items](AActor* Actor)
-				{
-					if (const FSceneOutlinerTreeItemPtr* ContainedActorItem = Items.Find(Actor))
-					{
-						OutChildren.Add(*ContainedActorItem);
-					}
-					return true;
-				});
-		}
-		else
-		{
-			// All attached actors are children of this actor
-			TArray<AActor*> AttachedActors;
-			ActorTreeItem->Actor->GetAttachedActors(AttachedActors);
-			for (const AActor* AttachedActor : AttachedActors)
-			{
-				if (const FSceneOutlinerTreeItemPtr* AttachedActorItem = Items.Find(AttachedActor))
-				{
-					OutChildren.Add(*AttachedActorItem);
-				}
-			}
-			if (bShowingComponents)
-			{
-				for (const UActorComponent* Component : ActorTreeItem->Actor->GetComponents())
-				{
-					if (const FSceneOutlinerTreeItemPtr* ComponentItem = Items.Find(Component))
-					{
-						OutChildren.Add(*ComponentItem);
-					}
-				}
-			}
-		}
-	}
-	else if (const FFolderTreeItem* FolderItem = Item.CastTo<FFolderTreeItem>())
-	{
-		// We should never call FindParents on a folder item if folders are not being shown
-		check(Mode->ShouldShowFolders());
-
-		// Search through all folders and find any folders with paths which are children of this item
-		for (const auto& Pair : FActorFolders::Get().GetFolderPropertiesForWorld(*RepresentingWorld))
-		{
-			if (FEditorFolderUtils::PathIsChildOf(Pair.Key, FolderItem->Path))
-			{
-				if (const FSceneOutlinerTreeItemPtr* Child = Items.Find(Pair.Key))
-				{
-					OutChildren.Add(*Child);
-				}
-			}
-		}
-	}
-}
-
 FSceneOutlinerTreeItemPtr FActorHierarchy::FindParent(const ISceneOutlinerTreeItem& Item, const TMap<FSceneOutlinerTreeItemID, FSceneOutlinerTreeItemPtr>& Items) const
 {
 	if (Item.IsA<FWorldTreeItem>())
diff --git a/Engine/Source/Editor/SceneOutliner/Public/ActorFolderHierarchy.h b/Engine/Source/Editor/SceneOutliner/Public/ActorFolderHierarchy.h
--- a/Engine/Source/Editor/SceneOutliner/Public/ActorFolderHierarchy.h
+++ b/Engine/Source/Editor/SceneOutliner/Public/ActorFolderHierarchy.h
@@ -13,7 +13,6 @@ public:
 
 
 	/* Begin ISceneOutlinerHierarchy Implementation */
-	virtual void FindChildren(const ISceneOutlinerTreeItem& Item, const TMap<FSceneOutlinerTreeItemID, FSceneOutlinerTreeItemPtr>& Items, TArray<FSceneOutlinerTreeItemPtr>& OutChildItems) const override;
 	virtual FSceneOutlinerTreeItemPtr FindParent(const ISceneOutlinerTreeItem& Item, const TMap<FSceneOutlinerTreeItemID, FSceneOutlinerTreeItemPtr>& Items) const override;
 	virtual void CreateItems(TArray<FSceneOutlinerTreeItemPtr>& OutItems) const override;
 	virtual void CreateChildren(const FSceneOutlinerTreeItemPtr& Item, TArray<FSceneOutlinerTreeItemPtr>& OutChildren) const override;
diff --git a/Engine/Source/Editor/SceneOutliner/Public/ActorHierarchy.h b/Engine/Source/Editor/SceneOutliner/Public/ActorHierarchy.h
--- a/Engine/Source/Editor/SceneOutliner/Public/ActorHierarchy.h
+++ b/Engine/Source/Editor/SceneOutliner/Public/ActorHierarchy.h
@@ -12,8 +12,6 @@ public:
 
 	static TUniquePtr<FActorHierarchy> Create(ISceneOutlinerMode* Mode, const TWeakObjectPtr<UWorld>& World);
 
-	/** Find all direct children of a tree item in an existing item map, if any exist. */
-	virtual void FindChildren(const ISceneOutlinerTreeItem& Item, const TMap<FSceneOutlinerTreeItemID, FSceneOutlinerTreeItemPtr>& Items, TArray<FSceneOutlinerTreeItemPtr>& OutChildItems) const override;
 	/** Find the parent of a tree item in an existing item map, it if exists. */
 	virtual FSceneOutlinerTreeItemPtr FindParent(const ISceneOutlinerTreeItem& Item, const TMap<FSceneOutlinerTreeItemID, FSceneOutlinerTreeItemPtr>& Items) const override;
 
diff --git a/Engine/Source/Editor/SceneOutliner/Public/ISceneOutlinerHierarchy.h b/Engine/Source/Editor/SceneOutliner/Public/ISceneOutlinerHierarchy.h
--- a/Engine/Source/Editor/SceneOutliner/Public/ISceneOutlinerHierarchy.h
+++ b/Engine/Source/Editor/SceneOutliner/Public/ISceneOutlinerHierarchy.h
@@ -14,8 +14,6 @@ public:
 	ISceneOutlinerHierarchy(ISceneOutlinerMode* InMode) : Mode(InMode) {}
 	virtual ~ISceneOutlinerHierarchy() {}
 
-	/** Find all direct children of a tree item in an existing item map, if any exist. */
-	virtual void FindChildren(const ISceneOutlinerTreeItem& Item, const TMap<FSceneOutlinerTreeItemID, FSceneOutlinerTreeItemPtr>& Items, TArray<FSceneOutlinerTreeItemPtr>& OutChildItems) const = 0;
 	/** Find the parent of a tree item in an existing item map, it if exists. */
 	virtual FSceneOutlinerTreeItemPtr FindParent(const ISceneOutlinerTreeItem& Item, const TMap<FSceneOutlinerTreeItemID, FSceneOutlinerTreeItemPtr>& Items) const = 0;
 		
