Source Control status column in the Outliner

diff --git a/Engine/Source/Editor/SceneOutliner/Private/ActorBrowsingMode.cpp b/Engine/Source/Editor/SceneOutliner/Private/ActorBrowsingMode.cpp
--- a/Engine/Source/Editor/SceneOutliner/Private/ActorBrowsingMode.cpp
+++ b/Engine/Source/Editor/SceneOutliner/Private/ActorBrowsingMode.cpp
@@ -224,6 +224,26 @@ FSlateColor FActorBrowsingMode::GetStatusTextColor() const
 
 void FActorBrowsingMode::CreateViewContent(FMenuBuilder& MenuBuilder)
 {
+	MenuBuilder.BeginSection("AssetThumbnails", LOCTEXT("ShowColumnHeading", "Columns"));
+	{
+		// For now hard code this column in.
+		// #todo_Outliner: refactor all info columns out of ActorInfoColumn into toggleable entries of this menu.
+		// could be done with a similar interface to FSceneOutlinerFilterInfo
+		MenuBuilder.AddMenuEntry(
+			LOCTEXT("SourceControlColumnName", "Source Control"),
+			LOCTEXT("SourceControlColumnTooltip", ""),
+			FSlateIcon(),
+			FUIAction(
+				FExecuteAction::CreateRaw(this, &FActorBrowsingMode::ToggleActorSCCStatusColumn),
+				FCanExecuteAction(),
+				FIsActionChecked::CreateRaw(this, &FActorBrowsingMode::isActorSCCStatusColumnActive)
+			),
+			NAME_None,
+			EUserInterfaceActionType::ToggleButton
+		);
+	}
+	MenuBuilder.EndSection();
+
 	MenuBuilder.BeginSection("AssetThumbnails", LOCTEXT("ShowWorldHeading", "World"));
 	{
 		MenuBuilder.AddSubMenu(
@@ -1396,4 +1416,17 @@ void FActorBrowsingMode::OnDeleteActorsEnd()
 	SceneOutliner->DeleteFoldersEnd();
 }
 
+void FActorBrowsingMode::ToggleActorSCCStatusColumn()
+{
+	if (bActorSCCStatusColumnActive)
+	{
+		SceneOutliner->RemoveColumn(FSceneOutlinerBuiltInColumnTypes::SourceControl());
+	}
+	else
+	{
+		SceneOutliner->AddColumn(FSceneOutlinerBuiltInColumnTypes::SourceControl(), FSceneOutlinerColumnInfo(ESceneOutlinerColumnVisibility::Visible, 30));
+	}
+	bActorSCCStatusColumnActive = !bActorSCCStatusColumnActive;
+}
+
 #undef LOCTEXT_NAMESPACE
\ No newline at end of file
diff --git a/Engine/Source/Editor/SceneOutliner/Private/SceneOutlinerActorSCCColumn.cpp b/Engine/Source/Editor/SceneOutliner/Private/SceneOutlinerActorSCCColumn.cpp
new file mode 100644
--- /dev/null
+++ b/Engine/Source/Editor/SceneOutliner/Private/SceneOutlinerActorSCCColumn.cpp
@@ -0,0 +1,170 @@
+// Copyright Epic Games, Inc. All Rights Reserved.
+
+#include "SceneOutlinerActorSCCColumn.h"
+#include "Widgets/Images/SImage.h"
+#include "Widgets/Views/STreeView.h"
+#include "ActorTreeItem.h"
+#include "ISourceControlProvider.h"
+#include "ISourceControlModule.h"
+#include "SourceControlHelpers.h"
+
+#define LOCTEXT_NAMESPACE "SceneOutlinerSourceControlColumn"
+
+class SSourceControlWidget : public SImage
+{
+public:
+	SLATE_BEGIN_ARGS(SSourceControlWidget) {}
+	SLATE_END_ARGS()
+
+	/** Construct this widget */
+	void Construct(const FArguments& InArgs, TWeakPtr<ISceneOutliner> InWeakOutliner, TWeakPtr<ISceneOutlinerTreeItem> InWeakTreeItem)
+	{
+		WeakTreeItem = InWeakTreeItem;
+		WeakOutliner = InWeakOutliner;
+
+		SImage::Construct(
+			SImage::FArguments()
+			.ColorAndOpacity(this, &SSourceControlWidget::GetForegroundColor)
+			.Image(this, &SSourceControlWidget::GetBrush));
+
+		SCCStateBrush = nullptr;
+
+		FSceneOutlinerTreeItemPtr TreeItemPtr = WeakTreeItem.Pin();
+		if (TreeItemPtr.IsValid())
+		{
+			if (FActorTreeItem* ActorItem = TreeItemPtr->CastTo<FActorTreeItem>())
+			{
+				AActor* Actor = ActorItem->Actor.Get();
+				if (Actor && Actor->GetLevel() && Actor->GetLevel()->bUseExternalActors)
+				{
+					ExternalPackageName = USourceControlHelpers::PackageFilename(Actor->GetExternalPackage());
+				}
+			}
+		}
+
+		if (!ExternalPackageName.IsEmpty())
+		{
+			ISourceControlModule& SCCModule = ISourceControlModule::Get();
+			SourceControlProviderChangedDelegateHandle = SCCModule.RegisterProviderChanged(FSourceControlProviderChanged::FDelegate::CreateSP(this, &SSourceControlWidget::HandleSourceControlProviderChanged));
+			SourceControlStateChangedDelegateHandle = SCCModule.GetProvider().RegisterSourceControlStateChanged_Handle(FSourceControlStateChanged::FDelegate::CreateSP(this, &SSourceControlWidget::HandleSourceControlStateChanged));
+
+			// Check if there is already a cached state for this item
+			FSourceControlStatePtr SourceControlState = ISourceControlModule::Get().GetProvider().GetState(ExternalPackageName, EStateCacheUsage::Use);
+			if (SourceControlState.IsValid() && SourceControlState->GetSmallIconName() != NAME_None)
+			{
+				SCCStateBrush = FEditorStyle::GetBrush(SourceControlState->GetSmallIconName());
+			}
+			else
+			{
+				SCCModule.QueueStatusUpdate(ExternalPackageName);
+			}
+		}
+	}
+
+	~SSourceControlWidget()
+	{
+		ISourceControlModule::Get().GetProvider().UnregisterSourceControlStateChanged_Handle(SourceControlStateChangedDelegateHandle);
+		ISourceControlModule::Get().UnregisterProviderChanged(SourceControlProviderChangedDelegateHandle);
+	}
+
+private:
+
+	virtual FReply OnMouseButtonDoubleClick(const FGeometry& InMyGeometry, const FPointerEvent& InMouseEvent) override
+	{
+		FSourceControlStatePtr SourceControlState = ISourceControlModule::Get().GetProvider().GetState(ExternalPackageName, EStateCacheUsage::ForceUpdate);
+		if (SourceControlState.IsValid())
+		{
+			SCCStateBrush = FEditorStyle::GetBrush(SourceControlState->GetSmallIconName());
+		}
+		return FReply::Handled();
+	}
+
+	void HandleSourceControlStateChanged()
+	{
+		FSourceControlStatePtr SourceControlState = ISourceControlModule::Get().GetProvider().GetState(ExternalPackageName, EStateCacheUsage::Use);
+		if (SourceControlState.IsValid())
+		{
+			SCCStateBrush = FEditorStyle::GetBrush(SourceControlState->GetSmallIconName());
+		}
+	}
+
+	void HandleSourceControlProviderChanged(ISourceControlProvider& OldProvider, ISourceControlProvider& NewProvider)
+	{
+		OldProvider.UnregisterSourceControlStateChanged_Handle(SourceControlStateChangedDelegateHandle);
+		SourceControlStateChangedDelegateHandle = NewProvider.RegisterSourceControlStateChanged_Handle(FSourceControlStateChanged::FDelegate::CreateSP(this, &SSourceControlWidget::HandleSourceControlStateChanged));
+		
+		SCCStateBrush = nullptr;
+		ISourceControlModule::Get().QueueStatusUpdate(ExternalPackageName);
+	}
+
+	/** Get the brush for this widget */
+	const FSlateBrush* GetBrush() const
+	{
+		return SCCStateBrush;
+	}
+
+	/** The tree item we relate to */
+	TWeakPtr<ISceneOutlinerTreeItem> WeakTreeItem;
+
+	/** Reference back to the outliner so we can set visibility of a whole selection */
+	TWeakPtr<ISceneOutliner> WeakOutliner;
+
+	/** Cached brush for the source control state */
+	const FSlateBrush* SCCStateBrush;
+
+	/** Cache the items external package name */
+	FString ExternalPackageName;
+
+	/** Source control state changed delegate handle */
+	FDelegateHandle SourceControlStateChangedDelegateHandle;
+
+	/** Source control provider changed delegate handle */
+	FDelegateHandle SourceControlProviderChangedDelegateHandle;
+};
+
+FName FSceneOutlinerActorSCCColumn::GetColumnID()
+{
+	return GetID();
+}
+
+SHeaderRow::FColumn::FArguments FSceneOutlinerActorSCCColumn::ConstructHeaderRowColumn()
+{
+	return SHeaderRow::Column(GetColumnID())
+		.FixedWidth(24.f)
+		.HAlignHeader(HAlign_Center)
+		.VAlignHeader(VAlign_Center)
+		.HAlignCell(HAlign_Center)
+		.VAlignCell(VAlign_Center)
+		[
+			SNew(SImage)
+			.ColorAndOpacity(FSlateColor::UseForeground())
+			.Image(this, &FSceneOutlinerActorSCCColumn::GetHeaderIcon)
+		];
+}
+
+const TSharedRef<SWidget> FSceneOutlinerActorSCCColumn::ConstructRowWidget(FSceneOutlinerTreeItemRef TreeItem, const STableRow<FSceneOutlinerTreeItemPtr>& Row)
+{
+	if (TreeItem->IsA<FActorTreeItem>())
+	{
+		return SNew(SHorizontalBox)
+			+ SHorizontalBox::Slot()
+			.AutoWidth()
+			.VAlign(VAlign_Center)
+			[
+				SNew(SSourceControlWidget, WeakSceneOutliner, TreeItem)
+			];
+	}
+	return SNullWidget::NullWidget;
+}
+
+const FSlateBrush* FSceneOutlinerActorSCCColumn::GetHeaderIcon() const
+{
+	if (ISourceControlModule::Get().IsEnabled())
+	{
+		return FEditorStyle::GetBrush("SourceControl.StatusIcon.On");
+	}
+	else
+	{
+		return FEditorStyle::GetBrush("SourceControl.StatusIcon.Off");
+	}
+}
\ No newline at end of file
diff --git a/Engine/Source/Editor/SceneOutliner/Private/SceneOutlinerActorSCCColumn.h b/Engine/Source/Editor/SceneOutliner/Private/SceneOutlinerActorSCCColumn.h
new file mode 100644
--- /dev/null
+++ b/Engine/Source/Editor/SceneOutliner/Private/SceneOutlinerActorSCCColumn.h
@@ -0,0 +1,42 @@
+// Copyright Epic Games, Inc. All Rights Reserved.
+
+#pragma once
+
+#include "CoreMinimal.h"
+#include "Widgets/SWidget.h"
+#include "Widgets/Views/SHeaderRow.h"
+#include "ISceneOutliner.h"
+#include "SceneOutlinerPublicTypes.h"
+#include "ISceneOutlinerColumn.h"
+
+template<typename ItemType> class STableRow;
+
+/** A column for the SceneOutliner that displays the item label */
+class FSceneOutlinerActorSCCColumn : public ISceneOutlinerColumn
+{
+
+public:
+	FSceneOutlinerActorSCCColumn(ISceneOutliner& SceneOutliner) : WeakSceneOutliner(StaticCastSharedRef<ISceneOutliner>(SceneOutliner.AsShared())) {}
+
+	virtual ~FSceneOutlinerActorSCCColumn() {}
+
+	static FName GetID() { return FSceneOutlinerBuiltInColumnTypes::SourceControl(); }
+
+	//////////////////////////////////////////////////////////////////////////
+	// Begin ISceneOutlinerColumn Implementation
+
+	virtual FName GetColumnID() override;
+
+	virtual SHeaderRow::FColumn::FArguments ConstructHeaderRowColumn() override;
+
+	virtual const TSharedRef<SWidget> ConstructRowWidget(FSceneOutlinerTreeItemRef TreeItem, const STableRow<FSceneOutlinerTreeItemPtr>& Row) override;
+
+	virtual bool SupportsSorting() const override { return false; }
+
+	// End ISceneOutlinerColumn Implementation
+	//////////////////////////////////////////////////////////////////////////
+private:
+	const FSlateBrush* GetHeaderIcon() const;
+
+	TWeakPtr<ISceneOutliner> WeakSceneOutliner;
+};
diff --git a/Engine/Source/Editor/SceneOutliner/Private/SceneOutlinerModule.cpp b/Engine/Source/Editor/SceneOutliner/Private/SceneOutlinerModule.cpp
--- a/Engine/Source/Editor/SceneOutliner/Private/SceneOutlinerModule.cpp
+++ b/Engine/Source/Editor/SceneOutliner/Private/SceneOutlinerModule.cpp
@@ -9,6 +9,7 @@
 #include "SceneOutlinerActorInfoColumn.h"
 #include "SceneOutlinerGutter.h"
 #include "SceneOutlinerItemLabelColumn.h"
+#include "SceneOutlinerActorSCCColumn.h"
 #include "SceneOutlinerPublicTypes.h"
 
 #include "ActorPickingMode.h"
@@ -26,6 +27,7 @@ void FSceneOutlinerModule::StartupModule()
 	// Register builtin column types which are not active by default
 	RegisterColumnType<FSceneOutlinerGutter>();
 	RegisterColumnType<FActorInfoColumn>();
+	RegisterColumnType<FSceneOutlinerActorSCCColumn>();
 }
 
 
diff --git a/Engine/Source/Editor/SceneOutliner/Public/ActorBrowsingMode.h b/Engine/Source/Editor/SceneOutliner/Public/ActorBrowsingMode.h
--- a/Engine/Source/Editor/SceneOutliner/Public/ActorBrowsingMode.h
+++ b/Engine/Source/Editor/SceneOutliner/Public/ActorBrowsingMode.h
@@ -99,9 +99,14 @@ private:
 	static TSharedRef<FSceneOutlinerFilter> CreateIsInCurrentLevelFilter();
 	static TSharedRef<FSceneOutlinerFilter> CreateHideComponentsFilter();
 	static TSharedRef<FSceneOutlinerFilter> CreateHideLevelInstancesFilter();
+	
+	void ToggleActorSCCStatusColumn();
+	bool isActorSCCStatusColumnActive() const { return bActorSCCStatusColumnActive; }
 private:
 	/** Number of actors which have passed through the filters */
 	uint32 FilteredActorCount = 0;
 	/** List of actors which passed the regular filters and may or may not have passed the search filter */
 	TSet<TWeakObjectPtr<AActor>> ApplicableActors;
+
+	bool bActorSCCStatusColumnActive = false;
 };
diff --git a/Engine/Source/Editor/SceneOutliner/Public/SceneOutlinerPublicTypes.h b/Engine/Source/Editor/SceneOutliner/Public/SceneOutlinerPublicTypes.h
--- a/Engine/Source/Editor/SceneOutliner/Public/SceneOutlinerPublicTypes.h
+++ b/Engine/Source/Editor/SceneOutliner/Public/SceneOutlinerPublicTypes.h
@@ -40,6 +40,12 @@ struct FSceneOutlinerBuiltInColumnTypes
 		static FName ActorInfo("ActorInfo");
 		return ActorInfo;
 	}
+
+	static FName& SourceControl()
+	{
+		static FName SourceControl("SourceControl");
+		return SourceControl;
+	}
 };
 
 /** Visibility enum for scene outliner columns */
diff --git a/Engine/Source/Editor/SceneOutliner/SceneOutliner.Build.cs b/Engine/Source/Editor/SceneOutliner/SceneOutliner.Build.cs
--- a/Engine/Source/Editor/SceneOutliner/SceneOutliner.Build.cs
+++ b/Engine/Source/Editor/SceneOutliner/SceneOutliner.Build.cs
@@ -24,6 +24,7 @@ public class SceneOutliner : ModuleRules
 				"EditorFramework",
 				"UnrealEd",
 				"ToolMenus",
+				"SourceControl"
             }
 		);
 	}
