Missing #undef LOCTEXT_NAMESPACE in SceneOutlinerActorSCCColumn.cpp

diff --git a/Engine/Source/Editor/SceneOutliner/Private/SceneOutlinerActorSCCColumn.cpp b/Engine/Source/Editor/SceneOutliner/Private/SceneOutlinerActorSCCColumn.cpp
--- a/Engine/Source/Editor/SceneOutliner/Private/SceneOutlinerActorSCCColumn.cpp
+++ b/Engine/Source/Editor/SceneOutliner/Private/SceneOutlinerActorSCCColumn.cpp
@@ -167,4 +167,6 @@ const FSlateBrush* FSceneOutlinerActorSCCColumn::GetHeaderIcon() const
 	{
 		return FEditorStyle::GetBrush("SourceControl.StatusIcon.Off");
 	}
-}
\ No newline at end of file
+}
+
+#undef LOCTEXT_NAMESPACE
\ No newline at end of file
