Fixed issue where actors in a LevelInstance with a folder path would not be correctly childed to LevelInstances.
Also prevents creating folders with actors in level instances.

diff --git a/Engine/Source/Editor/SceneOutliner/Private/ActorBrowsingMode.cpp b/Engine/Source/Editor/SceneOutliner/Private/ActorBrowsingMode.cpp
--- a/Engine/Source/Editor/SceneOutliner/Private/ActorBrowsingMode.cpp
+++ b/Engine/Source/Editor/SceneOutliner/Private/ActorBrowsingMode.cpp
@@ -1223,6 +1223,19 @@ FName FActorBrowsingMode::CreateNewFolder()
 {
 	const FScopedTransaction Transaction(LOCTEXT("UndoAction_CreateFolder", "Create Folder"));
 
+	if (ULevelInstanceSubsystem* LevelInstanceSubsystem = RepresentingWorld->GetSubsystem<ULevelInstanceSubsystem>())
+	{
+		for (FSelectionIterator It(*GEditor->GetSelectedActors()); It; ++It)
+		{
+			const AActor* Actor = CastChecked<AActor>(*It);
+			if (LevelInstanceSubsystem->GetParentLevelInstance(Actor) != nullptr)
+			{
+				UE_LOG(LogActorBrowser, Warning, TEXT("Cannot create a folder with actors who are children of a Level Instance"));
+				return NAME_None;
+			}
+		}
+	}
+
 	const FName NewFolderName = FActorFolders::Get().GetDefaultFolderNameForSelection(*RepresentingWorld);
 	FActorFolders::Get().CreateFolderContainingSelection(*RepresentingWorld, NewFolderName);
 
diff --git a/Engine/Source/Editor/SceneOutliner/Private/ActorHierarchy.cpp b/Engine/Source/Editor/SceneOutliner/Private/ActorHierarchy.cpp
--- a/Engine/Source/Editor/SceneOutliner/Private/ActorHierarchy.cpp
+++ b/Engine/Source/Editor/SceneOutliner/Private/ActorHierarchy.cpp
@@ -87,17 +87,6 @@ FSceneOutlinerTreeItemPtr FActorHierarchy::FindParent(const ISceneOutlinerTreeIt
 					return *ParentItem;
 				}
 			}
-			else if (Mode->ShouldShowFolders() && !Actor->GetFolderPath().IsNone())
-			{
-				if (const FSceneOutlinerTreeItemPtr* ParentItem = Items.Find(Actor->GetFolderPath()))
-				{
-					return *ParentItem;
-				}
-				else
-				{
-					return nullptr;
-				}
-			}
 
 			if (const ULevelInstanceSubsystem* LevelInstanceSubsystem = RepresentingWorld->GetSubsystem<ULevelInstanceSubsystem>())
 			{
@@ -120,6 +109,18 @@ FSceneOutlinerTreeItemPtr FActorHierarchy::FindParent(const ISceneOutlinerTreeIt
 				}
 			}
 
+			if (Mode->ShouldShowFolders() && !Actor->GetFolderPath().IsNone())
+			{
+				if (const FSceneOutlinerTreeItemPtr* ParentItem = Items.Find(Actor->GetFolderPath()))
+				{
+					return *ParentItem;
+				}
+				else
+				{
+					return nullptr;
+				}
+			}
+
 			// Default to the world
 			if (const FSceneOutlinerTreeItemPtr* ParentItem = Items.Find(ActorTreeItem->Actor->GetWorld()))
 			{
@@ -212,18 +213,6 @@ void FActorHierarchy::CreateWorldChildren(UWorld* World, TArray<FSceneOutlinerTr
 			CreateComponentItems(Actor, OutItems);
 		}
 	}
-
-	// Create all folder items
-	if (Mode->ShouldShowFolders())
-	{
-		for (const auto& Pair : FActorFolders::Get().GetFolderPropertiesForWorld(*RepresentingWorld))
-		{
-			if (FSceneOutlinerTreeItemPtr FolderItem = Mode->CreateItemFor<FActorFolderTreeItem>(FActorFolderTreeItem(Pair.Key, RepresentingWorld)))
-			{
-				OutItems.Add(FolderItem);
-			}
-		}
-	}
 }
 
 void FActorHierarchy::CreateItems(TArray<FSceneOutlinerTreeItemPtr>& OutItems) const
diff --git a/Engine/Source/Editor/SceneOutliner/Private/SSceneOutliner.cpp b/Engine/Source/Editor/SceneOutliner/Private/SSceneOutliner.cpp
--- a/Engine/Source/Editor/SceneOutliner/Private/SSceneOutliner.cpp
+++ b/Engine/Source/Editor/SceneOutliner/Private/SSceneOutliner.cpp
@@ -1152,19 +1152,22 @@ void SSceneOutliner::CreateFolder()
 {
 	const FName& NewFolderName = Mode->CreateNewFolder();
 
-	// Move any selected folders into the new folder
-	auto PreviouslySelectedItems = OutlinerTreeView->GetSelectedItems();
-	for (const auto& Item : PreviouslySelectedItems)
+	if (NewFolderName != NAME_None)
 	{
-		if (FFolderTreeItem* FolderItem = Item->CastTo<FFolderTreeItem>())
+		// Move any selected folders into the new folder
+		auto PreviouslySelectedItems = OutlinerTreeView->GetSelectedItems();
+		for (const auto& Item : PreviouslySelectedItems)
 		{
-			FolderItem->MoveTo(NewFolderName);
+			if (FFolderTreeItem* FolderItem = Item->CastTo<FFolderTreeItem>())
+			{
+				FolderItem->MoveTo(NewFolderName);
+			}
 		}
-	}
 
-	// At this point the new folder will be in our newly added list, so select it and open a rename when it gets refreshed
-	NewItemActions.Add(NewFolderName, SceneOutliner::ENewItemAction::Select | SceneOutliner::ENewItemAction::Rename);
-	RequestSort();
+		// At this point the new folder will be in our newly added list, so select it and open a rename when it gets refreshed
+		NewItemActions.Add(NewFolderName, SceneOutliner::ENewItemAction::Select | SceneOutliner::ENewItemAction::Rename);
+		RequestSort();
+	}
 }
 
 void SSceneOutliner::CopyFoldersBegin()
