- Added exports to SceneOutliner for upcoming changes.
- Added virtual ISceneOutlinerTreeItem::ShouldShowVisibilityState to allow never showing Visibility state.
- Added Source DragDropOperation in FSceneOutlinerDragDropPayload (for future use).

diff --git a/Engine/Source/Editor/SceneOutliner/Private/SOutlinerTreeView.cpp b/Engine/Source/Editor/SceneOutliner/Private/SOutlinerTreeView.cpp
--- a/Engine/Source/Editor/SceneOutliner/Private/SOutlinerTreeView.cpp
+++ b/Engine/Source/Editor/SceneOutliner/Private/SOutlinerTreeView.cpp
@@ -85,7 +85,7 @@ FReply HandleDrop(TSharedPtr<SSceneOutliner> SceneOutlinerPtr, const FDragDropEv
 		return FReply::Unhandled();
 	}
 
-	FSceneOutlinerDragDropPayload DraggedObjects;
+	FSceneOutlinerDragDropPayload DraggedObjects(*DragDropEvent.GetOperation());
 	// Validate now to make sure we don't doing anything we shouldn't
 	if (!SceneOutlinerPtr->ParseDragDrop(DraggedObjects, *DragDropEvent.GetOperation()))
 	{
diff --git a/Engine/Source/Editor/SceneOutliner/Private/SceneOutlinerGutter.cpp b/Engine/Source/Editor/SceneOutliner/Private/SceneOutlinerGutter.cpp
--- a/Engine/Source/Editor/SceneOutliner/Private/SceneOutlinerGutter.cpp
+++ b/Engine/Source/Editor/SceneOutliner/Private/SceneOutlinerGutter.cpp
@@ -360,13 +360,17 @@ SHeaderRow::FColumn::FArguments FSceneOutlinerGutter::ConstructHeaderRowColumn()
 
 const TSharedRef<SWidget> FSceneOutlinerGutter::ConstructRowWidget(FSceneOutlinerTreeItemRef TreeItem, const STableRow<FSceneOutlinerTreeItemPtr>& Row)
 {
-	return SNew(SHorizontalBox)
-		+SHorizontalBox::Slot()
-		.AutoWidth()
-		.VAlign(VAlign_Center)
-		[
-			SNew(SVisibilityWidget, SharedThis(this), WeakOutliner, TreeItem, &Row)
-		];
+	if (TreeItem->ShouldShowVisibilityState())
+	{
+		return SNew(SHorizontalBox)
+			+SHorizontalBox::Slot()
+			.AutoWidth()
+			.VAlign(VAlign_Center)
+			[
+				SNew(SVisibilityWidget, SharedThis(this), WeakOutliner, TreeItem, &Row)
+			];
+	}
+	return SNullWidget::NullWidget;
 }
 
 void FSceneOutlinerGutter::SortItems(TArray<FSceneOutlinerTreeItemPtr>& RootItems, const EColumnSortMode::Type SortMode) const
diff --git a/Engine/Source/Editor/SceneOutliner/Private/SceneOutlinerStandaloneTypes.cpp b/Engine/Source/Editor/SceneOutliner/Private/SceneOutlinerStandaloneTypes.cpp
--- a/Engine/Source/Editor/SceneOutliner/Private/SceneOutlinerStandaloneTypes.cpp
+++ b/Engine/Source/Editor/SceneOutliner/Private/SceneOutlinerStandaloneTypes.cpp
@@ -27,7 +27,7 @@ TOptional<FLinearColor> FSceneOutlinerCommonLabelData::GetForegroundColor(const
 	{
 		TSharedPtr<FDragDropOperation> DragDropOp = FSlateApplication::Get().GetDragDroppingContent();
 
-		FSceneOutlinerDragDropPayload DraggedObjects;
+		FSceneOutlinerDragDropPayload DraggedObjects(*DragDropOp);
 		const auto Outliner = WeakSceneOutliner.Pin();
 		if (Outliner->GetMode()->ParseDragDrop(DraggedObjects, *DragDropOp) && !Outliner->GetMode()->ValidateDrop(TreeItem, DraggedObjects).IsValid())
 		{
diff --git a/Engine/Source/Editor/SceneOutliner/Public/ISceneOutlinerTreeItem.h b/Engine/Source/Editor/SceneOutliner/Public/ISceneOutlinerTreeItem.h
--- a/Engine/Source/Editor/SceneOutliner/Public/ISceneOutlinerTreeItem.h
+++ b/Engine/Source/Editor/SceneOutliner/Public/ISceneOutlinerTreeItem.h
@@ -10,7 +10,7 @@
 class UToolMenu;
 
 /** Base tree item interface  */
-struct ISceneOutlinerTreeItem : TSharedFromThis<ISceneOutlinerTreeItem>
+struct SCENEOUTLINER_API ISceneOutlinerTreeItem : TSharedFromThis<ISceneOutlinerTreeItem>
 {
 	/** Friendship required for access to various internals */
 	friend SSceneOutliner;
@@ -138,8 +138,12 @@ public:
 	/** Called when this items visibility changed */
 	virtual void OnVisibilityChanged(const bool bNewVisibility) {}
 
+	/** Returns true if this item should show its visibility state */
+	virtual bool ShouldShowVisibilityState() const { return true; }
+
 	/** Returns true if this item can set its own visibility */
 	virtual bool HasVisibilityInfo() const { return false; }
+
 	/** Query this items visibility state. Only called if the item type has visibility info */
 	virtual bool GetVisibility() const { return false; }
 
diff --git a/Engine/Source/Editor/SceneOutliner/Public/SSceneOutliner.h b/Engine/Source/Editor/SceneOutliner/Public/SSceneOutliner.h
--- a/Engine/Source/Editor/SceneOutliner/Public/SSceneOutliner.h
+++ b/Engine/Source/Editor/SceneOutliner/Public/SSceneOutliner.h
@@ -193,7 +193,7 @@ struct FSceneOutlinerItemSelection
 /**
 	* Scene Outliner widget
 	*/
-class SSceneOutliner : public ISceneOutliner, public FEditorUndoClient, public FGCObject
+class SCENEOUTLINER_API SSceneOutliner : public ISceneOutliner, public FEditorUndoClient, public FGCObject
 {
 
 public:
@@ -214,7 +214,7 @@ public:
 	SSceneOutliner() : SharedData(MakeShareable(new FSharedSceneOutlinerData)) {}
 
 	/** SSceneOutliner destructor */
-	~SSceneOutliner();
+	virtual ~SSceneOutliner();
 
 	/** SWidget interface */
 	virtual void Tick(const FGeometry& AllottedGeometry, const double InCurrentTime, const float InDeltaTime) override;
diff --git a/Engine/Source/Editor/SceneOutliner/Public/SceneOutlinerDragDrop.h b/Engine/Source/Editor/SceneOutliner/Public/SceneOutlinerDragDrop.h
--- a/Engine/Source/Editor/SceneOutliner/Public/SceneOutlinerDragDrop.h
+++ b/Engine/Source/Editor/SceneOutliner/Public/SceneOutlinerDragDrop.h
@@ -28,11 +28,15 @@ enum class ESceneOutlinerDropCompatibility : uint8
 struct FSceneOutlinerDragDropPayload
 {
 	/** Default constructor, resulting in unset contents */
-	FSceneOutlinerDragDropPayload() {}
+	FSceneOutlinerDragDropPayload(const FDragDropOperation& InOperation = FDragDropOperation())
+	: SourceOperation(InOperation)
+	{
+	}
 
 	/** Populate this payload from an array of tree items */
 	template<typename TreeType>
-	FSceneOutlinerDragDropPayload(const TArray<TreeType>& InDraggedItems)
+	FSceneOutlinerDragDropPayload(const TArray<TreeType>& InDraggedItems, const FDragDropOperation& InOperation = FDragDropOperation())
+	: SourceOperation(InOperation)
 	{
 		for (const auto& Item : InDraggedItems)
 		{
@@ -109,6 +113,9 @@ struct FSceneOutlinerDragDropPayload
 
 	/** List of all dragged items */
 	mutable TArray<TWeakPtr<ISceneOutlinerTreeItem>> DraggedItems;
+
+	/** The source FDragDropOperation */
+	const FDragDropOperation& SourceOperation;
 };
 
 /** Struct used for validation of a drag/drop operation in the scene outliner */
@@ -151,7 +158,7 @@ struct FSceneOutlinerDragValidationInfo
 };
 
 /** A drag/drop operation that was started from the scene outliner */
-struct FSceneOutlinerDragDropOp : public FCompositeDragDropOp
+struct SCENEOUTLINER_API FSceneOutlinerDragDropOp : public FCompositeDragDropOp
 {
 	DRAG_DROP_OPERATOR_TYPE(FSceneOutlinerDragDropOp, FCompositeDragDropOp);
 		
diff --git a/Engine/Source/Editor/SceneOutliner/Public/SceneOutlinerMenuContext.h b/Engine/Source/Editor/SceneOutliner/Public/SceneOutlinerMenuContext.h
--- a/Engine/Source/Editor/SceneOutliner/Public/SceneOutlinerMenuContext.h
+++ b/Engine/Source/Editor/SceneOutliner/Public/SceneOutlinerMenuContext.h
@@ -10,7 +10,7 @@
 class SSceneOutliner;
 
 UCLASS()
-class USceneOutlinerMenuContext : public UObject
+class SCENEOUTLINER_API USceneOutlinerMenuContext : public UObject
 {
 	GENERATED_BODY()
 public:
diff --git a/Engine/Source/Editor/SceneOutliner/Public/SceneOutlinerPublicTypes.h b/Engine/Source/Editor/SceneOutliner/Public/SceneOutlinerPublicTypes.h
--- a/Engine/Source/Editor/SceneOutliner/Public/SceneOutlinerPublicTypes.h
+++ b/Engine/Source/Editor/SceneOutliner/Public/SceneOutlinerPublicTypes.h
@@ -78,7 +78,7 @@ struct FSceneOutlinerColumnInfo
 };
 
 /** Settings for the scene outliner which can be quieried publicly */
-struct FSharedSceneOutlinerData
+struct SCENEOUTLINER_API FSharedSceneOutlinerData
 {
 	/**	Invoked whenever the user attempts to delete an actor from within a Scene Outliner in the actor browsing mode */
 	FCustomSceneOutlinerDeleteDelegate CustomDelete;
diff --git a/Engine/Source/Editor/SceneOutliner/Public/SceneOutlinerStandaloneTypes.h b/Engine/Source/Editor/SceneOutliner/Public/SceneOutlinerStandaloneTypes.h
--- a/Engine/Source/Editor/SceneOutliner/Public/SceneOutlinerStandaloneTypes.h
+++ b/Engine/Source/Editor/SceneOutliner/Public/SceneOutlinerStandaloneTypes.h
@@ -142,7 +142,7 @@ private:
 	mutable uint8 Data[MaxSize];
 };
 	
-struct FSceneOutlinerTreeItemType
+struct SCENEOUTLINER_API FSceneOutlinerTreeItemType
 {
 public:
 	explicit FSceneOutlinerTreeItemType(const FSceneOutlinerTreeItemType* Parent = nullptr) : ID(++NextUniqueID), ParentType(Parent) {}
@@ -164,7 +164,7 @@ private:
 	const FSceneOutlinerTreeItemType* ParentType;
 };
 
-struct FSceneOutlinerCommonLabelData
+struct SCENEOUTLINER_API FSceneOutlinerCommonLabelData
 {
 	TWeakPtr<ISceneOutliner> WeakSceneOutliner;
 	static const FLinearColor DarkColor;
