Fix memory stomp when destroying FDataLayerHierarchy accessing dangling pointer FDataLayerMode.
Also some code cleanup for ISceneOutlinerHierarchy destruction.

diff --git a/Engine/Source/Editor/DataLayerEditor/Private/DataLayer/DataLayerMode.cpp b/Engine/Source/Editor/DataLayerEditor/Private/DataLayer/DataLayerMode.cpp
--- a/Engine/Source/Editor/DataLayerEditor/Private/DataLayer/DataLayerMode.cpp
+++ b/Engine/Source/Editor/DataLayerEditor/Private/DataLayer/DataLayerMode.cpp
@@ -39,10 +39,7 @@ FDataLayerMode::FDataLayerMode(const FDataLayerModeParams& Params)
 
 FDataLayerMode::~FDataLayerMode()
 {
-	if (Hierarchy)
-	{
-		Hierarchy->OnHierarchyChanged().RemoveAll(this);
-	}
+	Hierarchy.Release();
 }
 
 int32 FDataLayerMode::GetTypeSortPriority(const ISceneOutlinerTreeItem& Item) const
@@ -618,11 +615,6 @@ void FDataLayerMode::Rebuild()
 {
 	ChooseRepresentingWorld();
 	
-	if (Hierarchy)
-	{
-		Hierarchy->OnHierarchyChanged().Clear();
-	}
-
 	Hierarchy = CreateHierarchy();
 }
 
diff --git a/Engine/Source/Editor/SceneOutliner/Private/ActorMode.cpp b/Engine/Source/Editor/SceneOutliner/Private/ActorMode.cpp
--- a/Engine/Source/Editor/SceneOutliner/Private/ActorMode.cpp
+++ b/Engine/Source/Editor/SceneOutliner/Private/ActorMode.cpp
@@ -87,10 +87,6 @@ FActorMode::FActorMode(const FActorModeParams& Params)
 
 FActorMode::~FActorMode()
 {
-	if (Hierarchy)
-	{
-		Hierarchy->OnHierarchyChanged().RemoveAll(this);
-	}
 }
 
 TUniquePtr<ISceneOutlinerHierarchy> FActorMode::CreateHierarchy()
@@ -106,11 +102,6 @@ void FActorMode::Rebuild()
 {
 	ChooseRepresentingWorld();
 
-	if (Hierarchy)
-	{
-		Hierarchy->OnHierarchyChanged().Clear();
-	}
-
 	Hierarchy = CreateHierarchy();
 }
 
diff --git a/Engine/Source/Editor/SceneOutliner/Public/ISceneOutlinerHierarchy.h b/Engine/Source/Editor/SceneOutliner/Public/ISceneOutlinerHierarchy.h
--- a/Engine/Source/Editor/SceneOutliner/Public/ISceneOutlinerHierarchy.h
+++ b/Engine/Source/Editor/SceneOutliner/Public/ISceneOutlinerHierarchy.h
@@ -12,7 +12,7 @@ class ISceneOutlinerHierarchy : public TSharedFromThis<ISceneOutlinerHierarchy>
 {
 public:
 	ISceneOutlinerHierarchy(ISceneOutlinerMode* InMode) : Mode(InMode) {}
-	virtual ~ISceneOutlinerHierarchy() {}
+	virtual ~ISceneOutlinerHierarchy() { HierarchyChangedEvent.Clear(); }
 
 	/** Find the parent of a tree item in an existing item map, it if exists. */
 	virtual FSceneOutlinerTreeItemPtr FindParent(const ISceneOutlinerTreeItem& Item, const TMap<FSceneOutlinerTreeItemID, FSceneOutlinerTreeItemPtr>& Items) const = 0;
