Fix scrolling behavior in SceneOutliner when deselecting items.

diff --git a/Engine/Source/Editor/SceneOutliner/Private/ActorModeInteractive.cpp b/Engine/Source/Editor/SceneOutliner/Private/ActorModeInteractive.cpp
--- a/Engine/Source/Editor/SceneOutliner/Private/ActorModeInteractive.cpp
+++ b/Engine/Source/Editor/SceneOutliner/Private/ActorModeInteractive.cpp
@@ -51,7 +51,8 @@ void FActorModeInteractive::OnLevelSelectionChanged(UObject* Obj)
 	{
 		SceneOutliner->FullRefresh();
 	}
-	else
+	// If the SceneOutliner has user focus, we should not change the scrolling to override the user input.
+	else if (!SceneOutliner->HasAnyUserFocusOrFocusedDescendants())
 	{
 		SceneOutliner->ClearSelection();
 		SceneOutliner->RefreshSelection();
