WorldOutliner: Landscape: Streaming Proxies & Spline Actors now displayed as child of main Landscape

diff --git a/Engine/Source/Editor/SceneOutliner/Private/ActorBrowsingMode.cpp b/Engine/Source/Editor/SceneOutliner/Private/ActorBrowsingMode.cpp
--- a/Engine/Source/Editor/SceneOutliner/Private/ActorBrowsingMode.cpp
+++ b/Engine/Source/Editor/SceneOutliner/Private/ActorBrowsingMode.cpp
@@ -860,7 +860,7 @@ FSceneOutlinerDragValidationInfo FActorBrowsingMode::ValidateDrop(const ISceneOu
 					}
 				}
 
-				if (DragActor->GetAttachParentActor() != ActorTarget)
+				if (DragActor->GetSceneOutlinerParent() != ActorTarget)
 				{
 					bDraggedOntoAttachmentParent = false;
 				}
@@ -994,7 +994,7 @@ FSceneOutlinerDragValidationInfo FActorBrowsingMode::ValidateDrop(const ISceneOu
 				{
 					return FSceneOutlinerDragValidationInfo(ESceneOutlinerDropCompatibility::IncompatibleGeneric, FText::Format(LOCTEXT("Error_AttachChildActor", "Cannot move {0} as it is a child actor."), FText::FromString(Actor->GetActorLabel())));
 				}
-				else if (Actor->GetFolderPath() == DestinationPath && !Actor->GetAttachParentActor() && !bActorContainedInLevelInstance)
+				else if (Actor->GetFolderPath() == DestinationPath && !Actor->GetSceneOutlinerParent() && !bActorContainedInLevelInstance)
 				{
 					FFormatNamedArguments Args;
 					Args.Add(TEXT("SourceName"), FText::FromString(Actor->GetActorLabel()));
diff --git a/Engine/Source/Editor/SceneOutliner/Private/ActorHierarchy.cpp b/Engine/Source/Editor/SceneOutliner/Private/ActorHierarchy.cpp
--- a/Engine/Source/Editor/SceneOutliner/Private/ActorHierarchy.cpp
+++ b/Engine/Source/Editor/SceneOutliner/Private/ActorHierarchy.cpp
@@ -80,7 +80,7 @@ FSceneOutlinerTreeItemPtr FActorHierarchy::FindParent(const ISceneOutlinerTreeIt
 	{
 		if (AActor* Actor = ActorTreeItem->Actor.Get())
 		{
-			if (const AActor* ParentActor = Actor->GetAttachParentActor())
+			if (const AActor* ParentActor = Actor->GetSceneOutlinerParent())
 			{
 				if (const FSceneOutlinerTreeItemPtr* ParentItem = Items.Find(ParentActor))
 				{
@@ -312,7 +312,7 @@ FSceneOutlinerTreeItemPtr FActorHierarchy::CreateParentItem(const FSceneOutliner
 	{
 		if (const AActor* Actor = ActorTreeItem->Actor.Get())
 		{
-			if (AActor* ParentActor = Actor->GetAttachParentActor())
+			if (AActor* ParentActor = Actor->GetSceneOutlinerParent())
 			{
 				return Mode->CreateItemFor<FActorTreeItem>(ParentActor, true);
 			}
diff --git a/Engine/Source/Runtime/Engine/Classes/GameFramework/Actor.h b/Engine/Source/Runtime/Engine/Classes/GameFramework/Actor.h
--- a/Engine/Source/Runtime/Engine/Classes/GameFramework/Actor.h
+++ b/Engine/Source/Runtime/Engine/Classes/GameFramework/Actor.h
@@ -2058,6 +2058,9 @@ public:
 	/** Returns true if this actor is allowed to be attached to the given actor */
 	virtual bool EditorCanAttachTo(const AActor* InParent, FText& OutReason) const;
 
+	/** Returns the actor attachement parent that should be used in editor */
+	virtual AActor* GetSceneOutlinerParent() const;
+
 	/** Called before editor copy, true allow export */
 	virtual bool ShouldExport() { return true; }
 
diff --git a/Engine/Source/Runtime/Engine/Private/ActorEditor.cpp b/Engine/Source/Runtime/Engine/Private/ActorEditor.cpp
--- a/Engine/Source/Runtime/Engine/Private/ActorEditor.cpp
+++ b/Engine/Source/Runtime/Engine/Private/ActorEditor.cpp
@@ -830,6 +830,11 @@ bool AActor::EditorCanAttachTo(const AActor* InParent, FText& OutReason) const
 	return true;
 }
 
+AActor* AActor::GetSceneOutlinerParent() const
+{
+	return GetAttachParentActor();
+}
+
 class UHLODLayer* AActor::GetHLODLayer() const
 {
 	return HLODLayer;
diff --git a/Engine/Source/Runtime/Landscape/Classes/LandscapeProxy.h b/Engine/Source/Runtime/Landscape/Classes/LandscapeProxy.h
--- a/Engine/Source/Runtime/Landscape/Classes/LandscapeProxy.h
+++ b/Engine/Source/Runtime/Landscape/Classes/LandscapeProxy.h
@@ -865,6 +865,7 @@ public:
 	virtual uint32 GetDefaultGridSize(UWorld* InWorld) const override { return 1; }
 	virtual FGuid GetGridGuid() const override { return LandscapeGuid; }
 	virtual TUniquePtr<class FWorldPartitionActorDesc> CreateClassActorDesc() const override;
+	virtual bool EditorCanAttachTo(const AActor* InParent, FText& OutReason) const override { return false; }
 #endif	//WITH_EDITOR
 
 	virtual FGuid GetLandscapeGuid() const override { return LandscapeGuid; }
diff --git a/Engine/Source/Runtime/Landscape/Classes/LandscapeSplineActor.h b/Engine/Source/Runtime/Landscape/Classes/LandscapeSplineActor.h
--- a/Engine/Source/Runtime/Landscape/Classes/LandscapeSplineActor.h
+++ b/Engine/Source/Runtime/Landscape/Classes/LandscapeSplineActor.h
@@ -34,7 +34,8 @@ public:
 	virtual void PostRegisterAllComponents() override;
 	virtual void UnregisterAllComponents(bool bForReregister = false) override;
 	virtual void PostEditMove(bool bFinished) override;
-	
+	virtual bool EditorCanAttachTo(const AActor* InParent, FText& OutReason) const override { return false; }
+	virtual AActor* GetSceneOutlinerParent() const;
 	virtual bool SupportsForeignSplineMesh() const override { return false; }
 
 	// Interface existes for backward compatibility. Should already be created since its this actor's root component.
diff --git a/Engine/Source/Runtime/Landscape/Private/LandscapeEdit.cpp b/Engine/Source/Runtime/Landscape/Private/LandscapeEdit.cpp
--- a/Engine/Source/Runtime/Landscape/Private/LandscapeEdit.cpp
+++ b/Engine/Source/Runtime/Landscape/Private/LandscapeEdit.cpp
@@ -5046,6 +5046,16 @@ void ALandscapeStreamingProxy::PostRegisterAllComponents()
 	}
 }
 
+AActor* ALandscapeStreamingProxy::GetSceneOutlinerParent() const
+{
+	if (ULandscapeInfo* LandscapeInfo = GetLandscapeInfo())
+	{
+		return LandscapeInfo->LandscapeActor.Get();
+	}
+
+	return Super::GetSceneOutlinerParent();
+}
+
 void ALandscape::PreEditChange(FProperty* PropertyThatWillChange)
 {
 	PreEditLandscapeMaterial = LandscapeMaterial;
diff --git a/Engine/Source/Runtime/Landscape/Private/LandscapeSplineActor.cpp b/Engine/Source/Runtime/Landscape/Private/LandscapeSplineActor.cpp
--- a/Engine/Source/Runtime/Landscape/Private/LandscapeSplineActor.cpp
+++ b/Engine/Source/Runtime/Landscape/Private/LandscapeSplineActor.cpp
@@ -96,4 +96,14 @@ void ALandscapeSplineActor::PostEditMove(bool bFinished)
 	}
 }
 
+AActor* ALandscapeSplineActor::GetSceneOutlinerParent() const
+{
+	if (ULandscapeInfo* LandscapeInfo = GetLandscapeInfo())
+	{
+		return LandscapeInfo->LandscapeActor.Get();
+	}
+
+	return nullptr;
+}
+
 #endif
\ No newline at end of file
