Fix bug where ActorModeInteractive was overriding selection and scroll offset of SSceneOutliner after outliner had already updated itself correctly.

diff --git a/Engine/Source/Editor/SceneOutliner/Private/ActorModeInteractive.cpp b/Engine/Source/Editor/SceneOutliner/Private/ActorModeInteractive.cpp
--- a/Engine/Source/Editor/SceneOutliner/Private/ActorModeInteractive.cpp
+++ b/Engine/Source/Editor/SceneOutliner/Private/ActorModeInteractive.cpp
@@ -51,8 +51,8 @@ void FActorModeInteractive::OnLevelSelectionChanged(UObject* Obj)
 	{
 		SceneOutliner->FullRefresh();
 	}
-	// If the SceneOutliner has user focus, we should not change the scrolling to override the user input.
-	else if (!SceneOutliner->HasAnyUserFocusOrFocusedDescendants())
+	// If the SceneOutliner's reentrant flag is set, the selection change has already been handled in the outliner class
+	else if (!SceneOutliner->GetIsReentrant())
 	{
 		SceneOutliner->ClearSelection();
 		SceneOutliner->RefreshSelection();
diff --git a/Engine/Source/Editor/SceneOutliner/Public/SSceneOutliner.h b/Engine/Source/Editor/SceneOutliner/Public/SSceneOutliner.h
--- a/Engine/Source/Editor/SceneOutliner/Public/SSceneOutliner.h
+++ b/Engine/Source/Editor/SceneOutliner/Public/SSceneOutliner.h
@@ -412,6 +412,10 @@ public:
 
 	/** Called when a payload is dragged over an item */
 	FReply OnDragOverItem(const FDragDropEvent& Event, const ISceneOutlinerTreeItem& Item) const;
+
+	/** Used to test if Outliner related selection changes have already been handled */
+	bool GetIsReentrant() const { return bIsReentrant; }
+
 private:
 	/** Methods that implement structural modification logic for the tree */
 
