Fix empty folders dissapear in the scene outliner during undo / redo / PIE.

diff --git a/Engine/Source/Editor/SceneOutliner/Private/ActorHierarchy.cpp b/Engine/Source/Editor/SceneOutliner/Private/ActorHierarchy.cpp
--- a/Engine/Source/Editor/SceneOutliner/Private/ActorHierarchy.cpp
+++ b/Engine/Source/Editor/SceneOutliner/Private/ActorHierarchy.cpp
@@ -188,6 +188,15 @@ void FActorHierarchy::CreateWorldChildren(UWorld* World, TArray<FSceneOutlinerTr
 {
 	check(World);
 
+	// Add any folders which might match the current search terms
+	for (const auto& Pair : FActorFolders::Get().GetFolderPropertiesForWorld(*World))
+	{
+		if (FSceneOutlinerTreeItemPtr FolderItem = Mode->CreateItemFor<FActorFolderTreeItem>(FActorFolderTreeItem(Pair.Key, World)))
+		{
+			OutItems.Add(FolderItem);
+		}
+	}
+
 	const ULevelInstanceSubsystem* LevelInstanceSubsystem = World->GetSubsystem<ULevelInstanceSubsystem>();
 	// Create all actor items
 	for (FActorIterator ActorIt(World); ActorIt; ++ActorIt)
