Fixed Scene Outliner crash : Avoid creating FolderItems if Scene Outliner Mode doesn't show them

diff --git a/Engine/Source/Editor/SceneOutliner/Private/ActorHierarchy.cpp b/Engine/Source/Editor/SceneOutliner/Private/ActorHierarchy.cpp
--- a/Engine/Source/Editor/SceneOutliner/Private/ActorHierarchy.cpp
+++ b/Engine/Source/Editor/SceneOutliner/Private/ActorHierarchy.cpp
@@ -188,12 +188,15 @@ void FActorHierarchy::CreateWorldChildren(UWorld* World, TArray<FSceneOutlinerTr
 {
 	check(World);
 
-	// Add any folders which might match the current search terms
-	for (const auto& Pair : FActorFolders::Get().GetFolderPropertiesForWorld(*World))
+	if (Mode->ShouldShowFolders())
 	{
-		if (FSceneOutlinerTreeItemPtr FolderItem = Mode->CreateItemFor<FActorFolderTreeItem>(FActorFolderTreeItem(Pair.Key, World)))
+		// Add any folders which might match the current search terms
+		for (const auto& Pair : FActorFolders::Get().GetFolderPropertiesForWorld(*World))
 		{
-			OutItems.Add(FolderItem);
+			if (FSceneOutlinerTreeItemPtr FolderItem = Mode->CreateItemFor<FActorFolderTreeItem>(FActorFolderTreeItem(Pair.Key, World)))
+			{
+				OutItems.Add(FolderItem);
+			}
 		}
 	}
 
