SceneOutliner: Fixed drag drop tooltip not displaying correctly

diff --git a/Engine/Source/Editor/SceneOutliner/Private/SOutlinerTreeView.cpp b/Engine/Source/Editor/SceneOutliner/Private/SOutlinerTreeView.cpp
--- a/Engine/Source/Editor/SceneOutliner/Private/SOutlinerTreeView.cpp
+++ b/Engine/Source/Editor/SceneOutliner/Private/SOutlinerTreeView.cpp
@@ -18,18 +18,10 @@ static void UpdateOperationDecorator(const FDragDropEvent& Event, const FSceneOu
 	const FSlateBrush* Icon = ValidationInfo.IsValid() ? FEditorStyle::GetBrush(TEXT("Graph.ConnectorFeedback.OK")) : FEditorStyle::GetBrush(TEXT("Graph.ConnectorFeedback.Error"));
 
 	FDragDropOperation* Operation = Event.GetOperation().Get();
-	if (Operation)
+	if (Operation && Operation->IsOfType<FDecoratedDragDropOp>())
 	{
-		if (Operation->IsOfType<FSceneOutlinerDragDropOp>())
-		{
-			auto* OutlinerOp = static_cast<FSceneOutlinerDragDropOp*>(Operation);
-			OutlinerOp->SetTooltip(ValidationInfo.ValidationText, Icon);
-		}
-		else if (Operation->IsOfType<FDecoratedDragDropOp>())
-		{
-			auto* DecoratedOp = static_cast<FDecoratedDragDropOp*>(Operation);
-			DecoratedOp->SetToolTip(LOCTEXT("SceneOutlinerInvalidDragDropOp", "Invalid drag and drop operation, Drag into the Viewport."), Icon);
-		}
+		auto* DecoratedOp = static_cast<FDecoratedDragDropOp*>(Operation);
+		DecoratedOp->SetToolTip(ValidationInfo.ValidationText, Icon);
 	}
 }
 
