SceneOutliner: Fix for issue where editing the label of level instance while the level instance is in edit mode and dirty would append an asterisk to the label

diff --git a/Engine/Source/Editor/SceneOutliner/Private/ActorTreeItem.cpp b/Engine/Source/Editor/SceneOutliner/Private/ActorTreeItem.cpp
--- a/Engine/Source/Editor/SceneOutliner/Private/ActorTreeItem.cpp
+++ b/Engine/Source/Editor/SceneOutliner/Private/ActorTreeItem.cpp
@@ -52,6 +52,8 @@ struct SActorTreeLabel : FSceneOutlinerCommonLabelData, public SCompoundWidget
 				.ColorAndOpacity(this, &SActorTreeLabel::GetForegroundColor)
 				.OnTextCommitted(this, &SActorTreeLabel::OnLabelCommitted)
 				.OnVerifyTextChanged(this, &SActorTreeLabel::OnVerifyItemLabelChanged)
+				.OnEnterEditingMode(this, &SActorTreeLabel::OnEnterEditingMode)
+				.OnExitEditingMode(this, &SActorTreeLabel::OnExitEditingMode)
 				.IsSelected(FIsSelected::CreateSP(&InRow, &STableRow<FSceneOutlinerTreeItemPtr>::IsSelectedExclusively))
 				.IsReadOnly_Lambda([Item = ActorItem.AsShared(), this]()
 				{
@@ -115,7 +117,7 @@ private:
 		const AActor* Actor = ActorPtr.Get();
 		if (const ALevelInstance* LevelInstanceActor = Cast<ALevelInstance>(Actor))
 		{
-			if (LevelInstanceActor->IsDirty())
+			if (LevelInstanceActor->IsDirty() && !bInEditingMode)
 			{
 				FFormatNamedArguments Args;
 				Args.Add(TEXT("ActorLabel"), FText::FromString(LevelInstanceActor->GetActorLabel()));
@@ -319,6 +321,18 @@ private:
 			}
 		}
 	}
+
+	void OnEnterEditingMode()
+	{
+		bInEditingMode = true;
+	}
+
+	void OnExitEditingMode()
+	{
+		bInEditingMode = false;
+	}
+
+	bool bInEditingMode = false;
 };
 
 FActorTreeItem::FActorTreeItem(AActor* InActor)
