SceneOutliner: Fixed issue where converting a single actor to External did not show its Source Control status in the World Outliner

diff --git a/Engine/Source/Editor/SceneOutliner/Private/SceneOutlinerActorSCCColumn.cpp b/Engine/Source/Editor/SceneOutliner/Private/SceneOutlinerActorSCCColumn.cpp
--- a/Engine/Source/Editor/SceneOutliner/Private/SceneOutlinerActorSCCColumn.cpp
+++ b/Engine/Source/Editor/SceneOutliner/Private/SceneOutlinerActorSCCColumn.cpp
@@ -34,7 +34,7 @@ public:
 			if (FActorTreeItem* ActorItem = TreeItemPtr->CastTo<FActorTreeItem>())
 			{
 				AActor* Actor = ActorItem->Actor.Get();
-				if (Actor && Actor->GetLevel() && Actor->GetLevel()->bUseExternalActors)
+				if (Actor && Actor->IsPackageExternal())
 				{
 					ExternalPackageName = USourceControlHelpers::PackageFilename(Actor->GetExternalPackage());
 				}
