Added keyboard controls to socket shooser and bone reference selector

diff --git a/Engine/Source/Editor/Persona/Private/Shared/BoneSelectionWidget.cpp b/Engine/Source/Editor/Persona/Private/Shared/BoneSelectionWidget.cpp
--- a/Engine/Source/Editor/Persona/Private/Shared/BoneSelectionWidget.cpp
+++ b/Engine/Source/Editor/Persona/Private/Shared/BoneSelectionWidget.cpp
@@ -107,12 +107,32 @@ void SBoneTreeMenu::OnFilterTextChanged(const FText& InFilterText)
 void SBoneTreeMenu::OnSelectionChanged(TSharedPtr<SBoneTreeMenu::FBoneNameInfo> BoneInfo, ESelectInfo::Type SelectInfo)
 {
 	//Because we recreate all our items on tree refresh we will get a spurious null selection event initially.
-	if (BoneInfo.IsValid() && SelectInfo != ESelectInfo::Direct)
+	if (BoneInfo.IsValid() && SelectInfo == ESelectInfo::OnMouseClick)
 	{
-		OnSelectionChangedDelegate.ExecuteIfBound(BoneInfo->BoneName);
+		SelectBone(BoneInfo);
 	}
 }
 
+void SBoneTreeMenu::SelectBone(TSharedPtr<SBoneTreeMenu::FBoneNameInfo> BoneInfo)
+{
+	OnSelectionChangedDelegate.ExecuteIfBound(BoneInfo->BoneName);
+}
+
+FReply SBoneTreeMenu::OnKeyDown(const FGeometry& MyGeometry, const FKeyEvent& InKeyEvent)
+{
+	if(InKeyEvent.GetKey() == EKeys::Enter)
+	{
+		TArray<TSharedPtr<SBoneTreeMenu::FBoneNameInfo>> SelectedItems = TreeView->GetSelectedItems();
+		if(SelectedItems.Num() > 0)
+		{
+			SelectBone(SelectedItems[0]);
+			return FReply::Handled();
+		}
+	}
+
+	return FReply::Unhandled();
+}
+
 void SBoneTreeMenu::RebuildBoneList(const FName& SelectedBone)
 {
 	SkeletonTreeInfo.Empty();
diff --git a/Engine/Source/Editor/Persona/Public/BoneSelectionWidget.h b/Engine/Source/Editor/Persona/Public/BoneSelectionWidget.h
--- a/Engine/Source/Editor/Persona/Public/BoneSelectionWidget.h
+++ b/Engine/Source/Editor/Persona/Public/BoneSelectionWidget.h
@@ -60,6 +60,8 @@ public:
 	TSharedPtr<SWidget> GetFilterTextWidget();
 
 private:
+	// SWidget interface
+	virtual FReply OnKeyDown(const FGeometry& MyGeometry, const FKeyEvent& InKeyEvent) override;
 
 	// Using the current filter, repopulate the tree view
 	void RebuildBoneList(const FName& SelectedBone);
@@ -73,8 +75,12 @@ private:
 	// Called when the user changes the search filter
 	void OnFilterTextChanged(const FText& InFilterText);
 
+	// Handle the tree view selection changing
 	void OnSelectionChanged(TSharedPtr<SBoneTreeMenu::FBoneNameInfo> BoneInfo, ESelectInfo::Type SelectInfo);
 
+	// Select a specific bone, helper for UI handler functions
+	void SelectBone(TSharedPtr<SBoneTreeMenu::FBoneNameInfo> BoneInfo);
+	
 	// Tree info entries for bone picker
 	TArray<TSharedPtr<FBoneNameInfo>> SkeletonTreeInfo;
 	// Mirror of SkeletonTreeInfo but flattened for searching
diff --git a/Engine/Source/Editor/SceneOutliner/Private/SSocketChooser.h b/Engine/Source/Editor/SceneOutliner/Private/SSocketChooser.h
--- a/Engine/Source/Editor/SceneOutliner/Private/SSocketChooser.h
+++ b/Engine/Source/Editor/SceneOutliner/Private/SSocketChooser.h
@@ -133,7 +133,16 @@ public:
 	}
 
 	/** Called when item is selected */
-	void SelectedSocket(TSharedPtr<FSocketInfo> SocketInfo, ESelectInfo::Type /*SelectType*/)
+	void SelectedSocket(TSharedPtr<FSocketInfo> SocketInfo, ESelectInfo::Type SelectType)
+	{
+		if(SelectType == ESelectInfo::OnMouseClick)
+		{
+			ChooseSocket(SocketInfo);
+		}
+	}
+
+	/** Select an item */
+	void ChooseSocket(TSharedPtr<FSocketInfo> SocketInfo)
 	{
 		const FName SocketName = SocketInfo->Description.Name;
 
@@ -207,6 +216,7 @@ public:
 						.ListItemsSource( &FilteredSockets)
 						.OnGenerateRow( this, &SSocketChooserPopup::MakeItemWidget )
 						.OnSelectionChanged( this, &SSocketChooserPopup::SelectedSocket )
+						.SelectionMode(ESelectionMode::Single)
 					]
 				]
 			]
@@ -267,18 +277,33 @@ public:
 		SocketListView->RequestListRefresh();
 	}
 
-	bool SupportsKeyboardFocus() const
+	virtual FReply OnKeyDown(const FGeometry& MyGeometry, const FKeyEvent& InKeyEvent) override
+	{
+		if(InKeyEvent.GetKey() == EKeys::Enter)
+		{
+			TArray<TSharedPtr<FSocketInfo>> SelectedItems = SocketListView->GetSelectedItems();
+            if(SelectedItems.Num() > 0)
+            {
+            	ChooseSocket(SelectedItems[0]);
+            	return FReply::Handled();
+            }
+		}
+
+		return FReply::Unhandled();
+	}
+
+	virtual bool SupportsKeyboardFocus() const override
 	{
 		return SearchBox->SupportsKeyboardFocus();
 	}
 
-	bool HasKeyboardFocus() const
+	virtual bool HasKeyboardFocus() const override
 	{
 		// Since keyboard focus is forwarded to our editable text, we will test it instead
 		return SearchBox->HasKeyboardFocus();
 	}
 
-	FReply OnFocusReceived(const FGeometry& MyGeometry, const FFocusEvent& InFocusEvent)
+	virtual FReply OnFocusReceived(const FGeometry& MyGeometry, const FFocusEvent& InFocusEvent) override
 	{
 		// Forward keyboard focus to our editable text widget
 		return FReply::Handled().SetUserFocus(SearchBox.ToSharedRef(), InFocusEvent.GetCause());
