SceneOutliner: Fixed issue where actor SCC status would not update if actor packaging mode was changed.

diff --git a/Engine/Source/Editor/SceneOutliner/Private/SceneOutlinerActorSCCColumn.cpp b/Engine/Source/Editor/SceneOutliner/Private/SceneOutlinerActorSCCColumn.cpp
--- a/Engine/Source/Editor/SceneOutliner/Private/SceneOutlinerActorSCCColumn.cpp
+++ b/Engine/Source/Editor/SceneOutliner/Private/SceneOutlinerActorSCCColumn.cpp
@@ -33,37 +33,35 @@ public:
 		{
 			if (FActorTreeItem* ActorItem = TreeItemPtr->CastTo<FActorTreeItem>())
 			{
-				AActor* Actor = ActorItem->Actor.Get();
-				if (Actor && Actor->IsPackageExternal())
+				if (AActor* Actor = ActorItem->Actor.Get())
 				{
-					ExternalPackageName = USourceControlHelpers::PackageFilename(Actor->GetExternalPackage());
+					if (Actor->IsPackageExternal())
+					{
+						ExternalPackageName = USourceControlHelpers::PackageFilename(Actor->GetExternalPackage());
+						ConnectSourceControl();
+					}
+
+					ActorPackingModeChangedDelegateHandle = Actor->OnPackagingModeChanged.AddLambda([this](AActor* InActor, bool bExternal)
+						{
+							if (bExternal)
+							{
+								ExternalPackageName = USourceControlHelpers::PackageFilename(InActor->GetExternalPackage());
+								ConnectSourceControl();
+							}
+							else
+							{
+								ExternalPackageName = FString();
+								DisconnectSourceControl();
+							}
+						});
 				}
 			}
 		}
-
-		if (!ExternalPackageName.IsEmpty())
-		{
-			ISourceControlModule& SCCModule = ISourceControlModule::Get();
-			SourceControlProviderChangedDelegateHandle = SCCModule.RegisterProviderChanged(FSourceControlProviderChanged::FDelegate::CreateSP(this, &SSourceControlWidget::HandleSourceControlProviderChanged));
-			SourceControlStateChangedDelegateHandle = SCCModule.GetProvider().RegisterSourceControlStateChanged_Handle(FSourceControlStateChanged::FDelegate::CreateSP(this, &SSourceControlWidget::HandleSourceControlStateChanged));
-
-			// Check if there is already a cached state for this item
-			FSourceControlStatePtr SourceControlState = ISourceControlModule::Get().GetProvider().GetState(ExternalPackageName, EStateCacheUsage::Use);
-			if (SourceControlState.IsValid())
-			{
-				UpdateSourceControlStateIcon(SourceControlState);
-			}
-			else
-			{
-				SCCModule.QueueStatusUpdate(ExternalPackageName);
-			}
-		}
 	}
 
 	~SSourceControlWidget()
 	{
-		ISourceControlModule::Get().GetProvider().UnregisterSourceControlStateChanged_Handle(SourceControlStateChangedDelegateHandle);
-		ISourceControlModule::Get().UnregisterProviderChanged(SourceControlProviderChangedDelegateHandle);
+		DisconnectSourceControl();
 	}
 
 private:
@@ -79,19 +77,56 @@ private:
 		return FReply::Handled();
 	}
 
-	void HandleSourceControlStateChanged()
+	void ConnectSourceControl()
 	{
+		check(!ExternalPackageName.IsEmpty());
+
+		ISourceControlModule& SCCModule = ISourceControlModule::Get();
+		SourceControlProviderChangedDelegateHandle = SCCModule.RegisterProviderChanged(FSourceControlProviderChanged::FDelegate::CreateSP(this, &SSourceControlWidget::HandleSourceControlProviderChanged));
+		SourceControlStateChangedDelegateHandle = SCCModule.GetProvider().RegisterSourceControlStateChanged_Handle(FSourceControlStateChanged::FDelegate::CreateSP(this, &SSourceControlWidget::HandleSourceControlStateChanged, EStateCacheUsage::Use));
+
+		// Check if there is already a cached state for this item
 		FSourceControlStatePtr SourceControlState = ISourceControlModule::Get().GetProvider().GetState(ExternalPackageName, EStateCacheUsage::Use);
 		if (SourceControlState.IsValid())
 		{
 			UpdateSourceControlStateIcon(SourceControlState);
 		}
+		else
+		{
+			SCCModule.QueueStatusUpdate(ExternalPackageName);
+		}
+	}
+
+	void DisconnectSourceControl()
+	{
+		FSceneOutlinerTreeItemPtr TreeItemPtr = WeakTreeItem.Pin();
+		if (TreeItemPtr.IsValid())
+		{
+			if (FActorTreeItem* ActorItem = TreeItemPtr->CastTo<FActorTreeItem>())
+			{
+				if (AActor* Actor = ActorItem->Actor.Get())
+				{
+					Actor->OnPackagingModeChanged.Remove(ActorPackingModeChangedDelegateHandle);
+				}
+			}
+		}
+		ISourceControlModule::Get().GetProvider().UnregisterSourceControlStateChanged_Handle(SourceControlStateChangedDelegateHandle);
+		ISourceControlModule::Get().UnregisterProviderChanged(SourceControlProviderChangedDelegateHandle);
+	}
+
+	void HandleSourceControlStateChanged(EStateCacheUsage::Type CacheUsage)
+	{
+		FSourceControlStatePtr SourceControlState = ISourceControlModule::Get().GetProvider().GetState(ExternalPackageName, CacheUsage);
+		if (SourceControlState.IsValid())
+		{
+			UpdateSourceControlStateIcon(SourceControlState);
+		}
 	}
 
 	void HandleSourceControlProviderChanged(ISourceControlProvider& OldProvider, ISourceControlProvider& NewProvider)
 	{
 		OldProvider.UnregisterSourceControlStateChanged_Handle(SourceControlStateChangedDelegateHandle);
-		SourceControlStateChangedDelegateHandle = NewProvider.RegisterSourceControlStateChanged_Handle(FSourceControlStateChanged::FDelegate::CreateSP(this, &SSourceControlWidget::HandleSourceControlStateChanged));
+		SourceControlStateChangedDelegateHandle = NewProvider.RegisterSourceControlStateChanged_Handle(FSourceControlStateChanged::FDelegate::CreateSP(this, &SSourceControlWidget::HandleSourceControlStateChanged, EStateCacheUsage::Use));
 		
 		UpdateSourceControlStateIcon(nullptr);
 
@@ -127,6 +162,9 @@ private:
 
 	/** Source control provider changed delegate handle */
 	FDelegateHandle SourceControlProviderChangedDelegateHandle;
+
+	/** Actor packaging mode changed delegate handle */
+	FDelegateHandle ActorPackingModeChangedDelegateHandle;
 };
 
 FName FSceneOutlinerActorSCCColumn::GetColumnID()
