Scene Outliner - Rows have fixed height of 18px, Icons at 14px

diff --git a/Engine/Source/Editor/SceneOutliner/Private/ActorFolderTreeItem.cpp b/Engine/Source/Editor/SceneOutliner/Private/ActorFolderTreeItem.cpp
--- a/Engine/Source/Editor/SceneOutliner/Private/ActorFolderTreeItem.cpp
+++ b/Engine/Source/Editor/SceneOutliner/Private/ActorFolderTreeItem.cpp
@@ -53,6 +53,7 @@ struct SActorFolderTreeLabel : FSceneOutlinerCommonLabelData, public SCompoundWi
 			[
 				SNew(SBox)
 				.WidthOverride(FSceneOutlinerDefaultTreeItemMetrics::IconSize())
+				.HeightOverride(FSceneOutlinerDefaultTreeItemMetrics::IconSize())
 			[
 				SNew(SImage)
 				.Image(this, &SActorFolderTreeLabel::GetIcon)
diff --git a/Engine/Source/Editor/SceneOutliner/Private/SOutlinerTreeView.cpp b/Engine/Source/Editor/SceneOutliner/Private/SOutlinerTreeView.cpp
--- a/Engine/Source/Editor/SceneOutliner/Private/SOutlinerTreeView.cpp
+++ b/Engine/Source/Editor/SceneOutliner/Private/SOutlinerTreeView.cpp
@@ -260,20 +260,23 @@ TSharedRef<SWidget> SSceneOutlinerTreeRow::GenerateWidgetForColumn( const FName&
 	if( ColumnName == FSceneOutlinerBuiltInColumnTypes::Label() )
 	{
 		// The first column gets the tree expansion arrow for this row
-		return
-			SNew( SHorizontalBox )
-
-			+SHorizontalBox::Slot()
-			.AutoWidth()
-			.Padding(6, 0, 0, 0)
-			[
-				SNew( SExpanderArrow, SharedThis(this) ).IndentAmount(12)
-			]
-
-		+SHorizontalBox::Slot()
-			.FillWidth(1.0f)
+		return SNew(SBox)
+			.MinDesiredHeight(FSceneOutlinerDefaultTreeItemMetrics::RowHeight())
 			[
-				NewItemWidget
+				SNew( SHorizontalBox )
+
+				+SHorizontalBox::Slot()
+				.AutoWidth()
+				.Padding(6, 0, 0, 0)
+				[
+					SNew( SExpanderArrow, SharedThis(this) ).IndentAmount(12)
+				]
+
+				+SHorizontalBox::Slot()
+				.FillWidth(1.0f)
+				[
+					NewItemWidget
+				]
 			];
 	}
 	else
diff --git a/Engine/Source/Editor/SceneOutliner/Public/SceneOutlinerPublicTypes.h b/Engine/Source/Editor/SceneOutliner/Public/SceneOutlinerPublicTypes.h
--- a/Engine/Source/Editor/SceneOutliner/Public/SceneOutlinerPublicTypes.h
+++ b/Engine/Source/Editor/SceneOutliner/Public/SceneOutlinerPublicTypes.h
@@ -151,6 +151,7 @@ public:
 /** Default metrics for outliner tree items */
 struct FSceneOutlinerDefaultTreeItemMetrics
 {
+	static int32	RowHeight() { return 20; };
 	static int32	IconSize() { return 16; };
 	static FMargin	IconPadding() { return FMargin(0.f, 1.f, 6.f, 1.f); };
 };
