Scene Outliner: Fixes World Outliner Item Icon rendering

diff --git a/Engine/Source/Editor/SceneOutliner/Private/WorldTreeItem.cpp b/Engine/Source/Editor/SceneOutliner/Private/WorldTreeItem.cpp
--- a/Engine/Source/Editor/SceneOutliner/Private/WorldTreeItem.cpp
+++ b/Engine/Source/Editor/SceneOutliner/Private/WorldTreeItem.cpp
@@ -86,40 +86,42 @@ struct SWorldTreeLabel : FSceneOutlinerCommonLabelData, public SCompoundWidget
 		WeakSceneOutliner = StaticCastSharedRef<ISceneOutliner>(SceneOutliner.AsShared());
 
 		ChildSlot
-			[
-				SNew(SHorizontalBox)
+		[
+			SNew(SHorizontalBox)
 
-				+ SHorizontalBox::Slot()
+			+ SHorizontalBox::Slot()
 			.AutoWidth()
 			.VAlign(VAlign_Center)
 			.Padding(FSceneOutlinerDefaultTreeItemMetrics::IconPadding())
 			[
 				SNew(SBox)
 				.WidthOverride(FSceneOutlinerDefaultTreeItemMetrics::IconSize())
-			.HeightOverride(FSceneOutlinerDefaultTreeItemMetrics::IconSize())
-			[
-				SNew(SImage)
-				.Image(FSlateIconFinder::FindIconBrushForClass(UWorld::StaticClass()))
-			.ToolTipText(LOCTEXT("WorldIcon_Tooltip", "World"))
-			]
+				.HeightOverride(FSceneOutlinerDefaultTreeItemMetrics::IconSize())
+				[
+					SNew(SImage)
+					.Image(FSlateIconFinder::FindIconBrushForClass(UWorld::StaticClass()))
+					.ColorAndOpacity(FSlateColor::UseForeground())
+					.ToolTipText(LOCTEXT("WorldIcon_Tooltip", "World"))
+				]
 			]
 
-		+ SHorizontalBox::Slot()
+			+ SHorizontalBox::Slot()
 			.FillWidth(1.0f)
 			.VAlign(VAlign_Center)
 			.Padding(0.0f, 2.0f)
 			[
 				SNew(STextBlock)
 				.Text(this, &SWorldTreeLabel::GetDisplayText)
-			.HighlightText(SceneOutliner.GetFilterHighlightText())
-			.ColorAndOpacity(this, &SWorldTreeLabel::GetForegroundColor)
-			.ToolTip(IDocumentation::Get()->CreateToolTip(
-				TAttribute<FText>(this, &SWorldTreeLabel::GetTooltipText),
-				nullptr,
-				TEXT("Shared/LevelEditor/SceneOutliner"),
-				TEXT("WorldSettingsLabel")))
+				.HighlightText(SceneOutliner.GetFilterHighlightText())
+				.ColorAndOpacity(this, &SWorldTreeLabel::GetForegroundColor)
+				.ToolTip(IDocumentation::Get()->CreateToolTip(
+					TAttribute<FText>(this, &SWorldTreeLabel::GetTooltipText),
+					nullptr,
+					TEXT("Shared/LevelEditor/SceneOutliner"),
+					TEXT("WorldSettingsLabel")
+				))
 			]
-			];
+		];
 	}
 
 private:
