SceneOutliner: added support for batching hierarchy change events

diff --git a/Engine/Source/Editor/DataLayerEditor/Private/DataLayer/DataLayerHierarchy.cpp b/Engine/Source/Editor/DataLayerEditor/Private/DataLayer/DataLayerHierarchy.cpp
--- a/Engine/Source/Editor/DataLayerEditor/Private/DataLayer/DataLayerHierarchy.cpp
+++ b/Engine/Source/Editor/DataLayerEditor/Private/DataLayer/DataLayerHierarchy.cpp
@@ -140,11 +140,13 @@ void FDataLayerHierarchy::OnLevelActorAdded(AActor* InActor)
 			FSceneOutlinerHierarchyChangedData EventData;
 			EventData.Type = FSceneOutlinerHierarchyChangedData::Added;
 
-			for (const UDataLayer* DataLayer : InActor->GetDataLayerObjects())
+			TArray<const UDataLayer*> DataLayers = InActor->GetDataLayerObjects();
+			EventData.Items.Reserve(DataLayers.Num());
+			for (const UDataLayer* DataLayer : DataLayers)
 			{
-				EventData.Item = Mode->CreateItemFor<FDataLayerActorTreeItem>(FDataLayerActorTreeItemData(InActor, const_cast<UDataLayer*>(DataLayer)));
-				HierarchyChangedEvent.Broadcast(EventData);
+				EventData.Items.Add(Mode->CreateItemFor<FDataLayerActorTreeItem>(FDataLayerActorTreeItemData(InActor, const_cast<UDataLayer*>(DataLayer))));
 			}
+			HierarchyChangedEvent.Broadcast(EventData);
 		}
 	}
 }
@@ -181,11 +183,13 @@ void FDataLayerHierarchy::OnLevelActorDeleted(AActor* InActor)
 			FSceneOutlinerHierarchyChangedData EventData;
 			EventData.Type = FSceneOutlinerHierarchyChangedData::Removed;
 
-			for (const UDataLayer* DataLayer : InActor->GetDataLayerObjects())
+			TArray<const UDataLayer*> DataLayers = InActor->GetDataLayerObjects();
+			EventData.ItemIDs.Reserve(DataLayers.Num());
+			for (const UDataLayer* DataLayer : DataLayers)
 			{
-				EventData.ItemID = FDataLayerActorTreeItem::ComputeTreeItemID(InActor, DataLayer);
-				HierarchyChangedEvent.Broadcast(EventData);
+				EventData.ItemIDs.Add(FDataLayerActorTreeItem::ComputeTreeItemID(InActor, DataLayer));
 			}
+			HierarchyChangedEvent.Broadcast(EventData);
 		}
 	}
 }
diff --git a/Engine/Source/Editor/SceneOutliner/Private/ActorHierarchy.cpp b/Engine/Source/Editor/SceneOutliner/Private/ActorHierarchy.cpp
--- a/Engine/Source/Editor/SceneOutliner/Private/ActorHierarchy.cpp
+++ b/Engine/Source/Editor/SceneOutliner/Private/ActorHierarchy.cpp
@@ -391,44 +391,44 @@ void FActorHierarchy::FullRefreshEvent()
 
 void FActorHierarchy::OnLevelActorAdded(AActor* InActor)
 {
-	if (InActor && RepresentingWorld.Get() == InActor->GetWorld())
+	if (InActor != nullptr && RepresentingWorld.Get() == InActor->GetWorld())
 	{
 		FSceneOutlinerHierarchyChangedData EventData;
 		EventData.Type = FSceneOutlinerHierarchyChangedData::Added;
-		EventData.Item = Mode->CreateItemFor<FActorTreeItem>(InActor);
+		EventData.Items.Add(Mode->CreateItemFor<FActorTreeItem>(InActor));
 		HierarchyChangedEvent.Broadcast(EventData);
 	}
 }
 
 void FActorHierarchy::OnLevelActorDeleted(AActor* InActor)
 {
-	if (RepresentingWorld.Get() == InActor->GetWorld())
+	if (InActor != nullptr && RepresentingWorld.Get() == InActor->GetWorld())
 	{
 		FSceneOutlinerHierarchyChangedData EventData;
 		EventData.Type = FSceneOutlinerHierarchyChangedData::Removed;
-		EventData.ItemID = InActor;
+		EventData.ItemIDs.Add(InActor);
 		HierarchyChangedEvent.Broadcast(EventData);
 	}
 }
 
 void FActorHierarchy::OnLevelActorAttached(AActor* InActor, const AActor* InParent)
 {
-	if (RepresentingWorld.Get() == InActor->GetWorld())
+	if (InActor != nullptr && RepresentingWorld.Get() == InActor->GetWorld())
 	{
 		FSceneOutlinerHierarchyChangedData EventData;
 		EventData.Type = FSceneOutlinerHierarchyChangedData::Moved;
-		EventData.ItemID = InActor;
+		EventData.ItemIDs.Add(InActor);
 		HierarchyChangedEvent.Broadcast(EventData);
 	}
 }
 
 void FActorHierarchy::OnLevelActorDetached(AActor* InActor, const AActor* InParent)
 {
-	if (RepresentingWorld.Get() == InActor->GetWorld())
+	if (InActor != nullptr && RepresentingWorld.Get() == InActor->GetWorld())
 	{
 		FSceneOutlinerHierarchyChangedData EventData;
 		EventData.Type = FSceneOutlinerHierarchyChangedData::Moved;
-		EventData.ItemID = InActor;
+		EventData.ItemIDs.Add(InActor);
 		HierarchyChangedEvent.Broadcast(EventData);
 	}
 }
@@ -445,17 +445,38 @@ void FActorHierarchy::OnLevelActorListChanged()
 
 void FActorHierarchy::OnLevelAdded(ULevel* InLevel, UWorld* InWorld)
 {
-	if (RepresentingWorld.Get() == InWorld)
+	if (InLevel != nullptr && RepresentingWorld.Get() == InWorld)
 	{
-		FullRefreshEvent();
+		FSceneOutlinerHierarchyChangedData EventData;
+		EventData.Type = FSceneOutlinerHierarchyChangedData::Added;
+
+		EventData.Items.Reserve(InLevel->Actors.Num());
+		for (AActor* Actor : InLevel->Actors)
+		{
+			if (Actor != nullptr)
+			{
+				EventData.Items.Add(Mode->CreateItemFor<FActorTreeItem>(Actor));
+			}
+		}
+		HierarchyChangedEvent.Broadcast(EventData);
 	}
 }
 
 void FActorHierarchy::OnLevelRemoved(ULevel* InLevel, UWorld* InWorld)
 {
-	if (RepresentingWorld.Get() == InWorld)
+	if (InLevel != nullptr && RepresentingWorld.Get() == InWorld)
 	{
-		FullRefreshEvent();
+		FSceneOutlinerHierarchyChangedData EventData;
+		EventData.Type = FSceneOutlinerHierarchyChangedData::Removed;
+
+		EventData.ItemIDs.Reserve(InLevel->Actors.Num());
+		for (AActor* Actor : InLevel->Actors)
+		{
+			if (Actor != nullptr)
+			{
+				EventData.ItemIDs.Add(Actor);		    
+			}
+		}
 	}
 }
 
@@ -466,7 +487,7 @@ void FActorHierarchy::OnBroadcastFolderCreate(UWorld& InWorld, FName NewPath)
 	{
 		FSceneOutlinerHierarchyChangedData EventData;
 		EventData.Type = FSceneOutlinerHierarchyChangedData::Added;
-		EventData.Item = Mode->CreateItemFor<FActorFolderTreeItem>(FActorFolderTreeItem(NewPath, &InWorld));
+		EventData.Items.Add(Mode->CreateItemFor<FActorFolderTreeItem>(FActorFolderTreeItem(NewPath, &InWorld)));
 		EventData.ItemActions = SceneOutliner::ENewItemAction::Select | SceneOutliner::ENewItemAction::Rename;
 		HierarchyChangedEvent.Broadcast(EventData);
 	}
@@ -479,8 +500,8 @@ void FActorHierarchy::OnBroadcastFolderMove(UWorld& InWorld, FName OldPath, FNam
 	{
 		FSceneOutlinerHierarchyChangedData EventData;
 		EventData.Type = FSceneOutlinerHierarchyChangedData::FolderMoved;
-		EventData.ItemID = OldPath;
-		EventData.NewPath = NewPath;
+		EventData.ItemIDs.Add(OldPath);
+		EventData.NewPaths.Add(NewPath);
 		HierarchyChangedEvent.Broadcast(EventData);
 	}
 }
@@ -492,7 +513,7 @@ void FActorHierarchy::OnBroadcastFolderDelete(UWorld& InWorld, FName Path)
 	{
 		FSceneOutlinerHierarchyChangedData EventData;
 		EventData.Type = FSceneOutlinerHierarchyChangedData::Removed;
-		EventData.ItemID = Path;
+		EventData.ItemIDs.Add(Path);
 		HierarchyChangedEvent.Broadcast(EventData);
 	}
 }
@@ -503,7 +524,7 @@ void FActorHierarchy::OnLevelActorFolderChanged(const AActor* InActor, FName Old
 	{
 		FSceneOutlinerHierarchyChangedData EventData;
 		EventData.Type = FSceneOutlinerHierarchyChangedData::Moved;
-		EventData.ItemID = FSceneOutlinerTreeItemID(InActor);
+		EventData.ItemIDs.Add(FSceneOutlinerTreeItemID(InActor));
 		HierarchyChangedEvent.Broadcast(EventData);
 	}
 }
diff --git a/Engine/Source/Editor/SceneOutliner/Private/SSceneOutliner.cpp b/Engine/Source/Editor/SceneOutliner/Private/SSceneOutliner.cpp
--- a/Engine/Source/Editor/SceneOutliner/Private/SSceneOutliner.cpp
+++ b/Engine/Source/Editor/SceneOutliner/Private/SSceneOutliner.cpp
@@ -1559,58 +1559,79 @@ void SSceneOutliner::OnHierarchyChangedEvent(FSceneOutlinerHierarchyChangedData
 {
 	if (Event.Type == FSceneOutlinerHierarchyChangedData::Added)
 	{
-		if (Event.Item.IsValid() && !TreeItemMap.Find(Event.Item->GetID()))
+		for (const auto& TreeItemPtr : Event.Items)
 		{
-			AddPendingItemAndChildren(Event.Item);
-			if (Event.ItemActions)
+			if (TreeItemPtr.IsValid() && !TreeItemMap.Find(TreeItemPtr->GetID()))
 			{
-				NewItemActions.Add(Event.Item->GetID(), Event.ItemActions);
+				AddPendingItemAndChildren(TreeItemPtr);
+				if (Event.ItemActions)
+				{
+					NewItemActions.Add(TreeItemPtr->GetID(), Event.ItemActions);
+				}
 			}
 		}
 	}
 	else if (Event.Type == FSceneOutlinerHierarchyChangedData::Removed)
 	{
-		FSceneOutlinerTreeItemPtr* Item = TreeItemMap.Find(Event.ItemID);
-		if (!Item)
+		for (const auto& TreeItemID : Event.ItemIDs)
 		{
-			Item = PendingTreeItemMap.Find(Event.ItemID);
-		}
+			FSceneOutlinerTreeItemPtr* Item = TreeItemMap.Find(TreeItemID);
+			if (!Item)
+			{
+				Item = PendingTreeItemMap.Find(TreeItemID);
+			}
 
-		if (Item)
-		{
-			PendingOperations.Emplace(SceneOutliner::FPendingTreeOperation::Removed, Item->ToSharedRef());
-			Refresh();
+			if (Item)
+			{
+				PendingOperations.Emplace(SceneOutliner::FPendingTreeOperation::Removed, Item->ToSharedRef());
+			}
 		}
+		Refresh();
 	}
 	else if (Event.Type == FSceneOutlinerHierarchyChangedData::Moved)
 	{
-		FSceneOutlinerTreeItemPtr* Item = Event.Item ? &Event.Item : TreeItemMap.Find(Event.ItemID);
+		for (const auto& TreeItemID : Event.ItemIDs)
+		{
+			FSceneOutlinerTreeItemPtr* Item = TreeItemMap.Find(TreeItemID);
 
-		if (Item)
+			if (Item)
+			{
+				PendingOperations.Emplace(SceneOutliner::FPendingTreeOperation::Moved, Item->ToSharedRef());
+			}
+		}
+
+		for (const auto& TreeItemPtr : Event.Items)
 		{
-			PendingOperations.Emplace(SceneOutliner::FPendingTreeOperation::Moved, Item->ToSharedRef());
-			Refresh();
+			if (TreeItemPtr.IsValid())
+			{
+				PendingOperations.Emplace(SceneOutliner::FPendingTreeOperation::Moved, TreeItemPtr.ToSharedRef());
+			}
 		}
+		Refresh();
 	}
 	else if (Event.Type == FSceneOutlinerHierarchyChangedData::FolderMoved)
 	{
-		FSceneOutlinerTreeItemPtr Item = TreeItemMap.FindRef(Event.ItemID);
-		if (Item.IsValid())
+		check (Event.ItemIDs.Num() == Event.NewPaths.Num());
+		for (int32 i = 0; i < Event.ItemIDs.Num(); ++i)
 		{
-			// Remove it from the map under the old ID (which is derived from the folder path)
-			TreeItemMap.Remove(Item->GetID());
+			FSceneOutlinerTreeItemPtr Item = TreeItemMap.FindRef(Event.ItemIDs[i]);
+			if (Item.IsValid())
+			{
+				// Remove it from the map under the old ID (which is derived from the folder path)
+				TreeItemMap.Remove(Item->GetID());
 
-			// Now change the path and put it back in the map with its new ID
-			auto Folder = StaticCastSharedPtr<FFolderTreeItem>(Item);
-			Folder->Path = Event.NewPath;
-			Folder->LeafName = FEditorFolderUtils::GetLeafName(Event.NewPath);
+				// Now change the path and put it back in the map with its new ID
+				auto Folder = StaticCastSharedPtr<FFolderTreeItem>(Item);
+				Folder->Path = Event.NewPaths[i];
+				Folder->LeafName = FEditorFolderUtils::GetLeafName(Event.NewPaths[i]);
 
-			TreeItemMap.Add(Item->GetID(), Item);
+				TreeItemMap.Add(Item->GetID(), Item);
 
-			// Add an operation to move the item in the hierarchy
-			PendingOperations.Emplace(SceneOutliner::FPendingTreeOperation::Moved, Item.ToSharedRef());
-			Refresh();
+				// Add an operation to move the item in the hierarchy
+				PendingOperations.Emplace(SceneOutliner::FPendingTreeOperation::Moved, Item.ToSharedRef());
+			}
 		}
+		Refresh();
 	}
 	else if (Event.Type == FSceneOutlinerHierarchyChangedData::FullRefresh)
 	{
@@ -1623,7 +1644,7 @@ void SSceneOutliner::PostUndo(bool bSuccess)
 	// Refresh our tree in case any changes have been made to the scene that might effect our list
 	if( !bIsReentrant )
 	{
-        bDisableIntermediateSorting = true;
+		bDisableIntermediateSorting = true;
 		FullRefresh();
 	}
 }
