Added missing file from 15541886

diff --git a/Engine/Source/Editor/SceneOutliner/Public/SceneOutlinerStandaloneTypes.h b/Engine/Source/Editor/SceneOutliner/Public/SceneOutlinerStandaloneTypes.h
--- a/Engine/Source/Editor/SceneOutliner/Public/SceneOutlinerStandaloneTypes.h
+++ b/Engine/Source/Editor/SceneOutliner/Public/SceneOutlinerStandaloneTypes.h
@@ -192,11 +192,11 @@ struct FSceneOutlinerHierarchyChangedData
 	} Type;
 
 	// This event may pass one of two kinds of data, depending on the type of event
-	FSceneOutlinerTreeItemPtr Item;
+	TArray<FSceneOutlinerTreeItemPtr> Items;
 
-	FSceneOutlinerTreeItemID ItemID;
+	TArray<FSceneOutlinerTreeItemID> ItemIDs;
 	// Used for FolderMoved events
-	FName NewPath;
+	TArray<FName> NewPaths;
 	/** Actions to apply to items */
 	uint8 ItemActions = 0;
 };
