Fixes Scene Outliner Type Column Dropdown to be a proper menu

diff --git a/Engine/Source/Editor/SceneOutliner/Private/SceneOutlinerActorInfoColumn.cpp b/Engine/Source/Editor/SceneOutliner/Private/SceneOutlinerActorInfoColumn.cpp
--- a/Engine/Source/Editor/SceneOutliner/Private/SceneOutlinerActorInfoColumn.cpp
+++ b/Engine/Source/Editor/SceneOutliner/Private/SceneOutlinerActorInfoColumn.cpp
@@ -156,9 +156,6 @@ struct FGetInfo
 	}
 };
 
-
-TArray< TSharedPtr< SceneOutliner::ECustomColumnMode::Type > > FActorInfoColumn::ModeOptions;
-
 FActorInfoColumn::FActorInfoColumn( ISceneOutliner& Outliner, SceneOutliner::ECustomColumnMode::Type InDefaultCustomColumnMode )
 	: CurrentMode( InDefaultCustomColumnMode )
 	, SceneOutlinerWeak( StaticCastSharedRef<ISceneOutliner>(Outliner.AsShared()) )
@@ -175,31 +172,11 @@ FName FActorInfoColumn::GetColumnID()
 
 SHeaderRow::FColumn::FArguments FActorInfoColumn::ConstructHeaderRowColumn()
 {
-	if( ModeOptions.Num() == 0 )
-	{
-		for(SceneOutliner::ECustomColumnMode::Type CurMode : TEnumRange<SceneOutliner::ECustomColumnMode::Type>() )
-		{
-			ModeOptions.Add( MakeShareable( new SceneOutliner::ECustomColumnMode::Type( CurMode ) ) );
-		}
-	}
-
 	/** Customizable actor data column */
 	return SHeaderRow::Column( GetColumnID() )
 		.FillWidth(2)
 		.HeaderComboVisibility(EHeaderComboVisibility::OnHover)
-		.MenuContent()
-		[
-			SNew(SBorder)
-			.Padding(FMargin(5))
-			.BorderImage(FEditorStyle::GetBrush("Menu.Background"))
-			[
-				SNew(SListView<TSharedPtr<SceneOutliner::ECustomColumnMode::Type>>)
-				.ListItemsSource(&ModeOptions)
-				.SelectionMode(ESelectionMode::Single)
-				.OnGenerateRow( this, &FActorInfoColumn::MakeComboButtonItemWidget )
-				.OnSelectionChanged( this, &FActorInfoColumn::OnModeChanged )
-			]
-		]
+		.OnGetMenuContent(this, &FActorInfoColumn::GenerateActorTypeMenu)
 		.HeaderContent()
 		[
 			SNew( SHorizontalBox )
@@ -213,6 +190,31 @@ SHeaderRow::FColumn::FArguments FActorInfoColumn::ConstructHeaderRowColumn()
 		];
 }
 
+TSharedRef<SWidget> FActorInfoColumn::GenerateActorTypeMenu() 
+{
+	FMenuBuilder MenuBuilder(true, nullptr);
+
+	for(SceneOutliner::ECustomColumnMode::Type CurMode : TEnumRange<SceneOutliner::ECustomColumnMode::Type>() )
+	{
+		FText CurModeText = MakeComboText(CurMode);
+		MenuBuilder.AddMenuEntry(
+			CurModeText,
+			MakeComboToolTipText( CurMode ),
+			FSlateIcon(), 
+			FUIAction(
+				FExecuteAction::CreateSP(this, &FActorInfoColumn::OnModeChanged, CurMode),
+				FCanExecuteAction(),
+				FIsActionChecked::CreateSP(this, &FActorInfoColumn::IsCurrentMode, CurMode)
+			),
+			NAME_None, 
+			EUserInterfaceActionType::RadioButton
+		);
+	}
+
+	return MenuBuilder.MakeWidget();
+
+}
+
 const TSharedRef< SWidget > FActorInfoColumn::ConstructRowWidget( FSceneOutlinerTreeItemRef TreeItem, const STableRow<FSceneOutlinerTreeItemPtr>& Row )
 {
 	auto SceneOutliner = SceneOutlinerWeak.Pin();
@@ -330,15 +332,20 @@ void FActorInfoColumn::SortItems(TArray<FSceneOutlinerTreeItemPtr>& RootItems, c
 		.Sort(RootItems);
 }
 
-void FActorInfoColumn::OnModeChanged( TSharedPtr< SceneOutliner::ECustomColumnMode::Type > NewSelection, ESelectInfo::Type /*SelectInfo*/ )
+void FActorInfoColumn::OnModeChanged( const SceneOutliner::ECustomColumnMode::Type NewSelection )
 {
-	CurrentMode = *NewSelection;
+	CurrentMode = NewSelection;
 
 	// Refresh and refilter the list
 	SceneOutlinerWeak.Pin()->Refresh();
 	FSlateApplication::Get().DismissAllMenus();
 }
 
+bool FActorInfoColumn::IsCurrentMode( const SceneOutliner::ECustomColumnMode::Type Mode ) const 
+{
+	return CurrentMode == Mode;
+}
+
 EVisibility FActorInfoColumn::GetColumnDataVisibility( bool bIsClassHyperlink ) const
 {
 	if ( CurrentMode == SceneOutliner::ECustomColumnMode::Class )
@@ -418,7 +425,7 @@ FText FActorInfoColumn::MakeComboText( const SceneOutliner::ECustomColumnMode::T
 }
 
 
-FText FActorInfoColumn::MakeComboToolTipText( const SceneOutliner::ECustomColumnMode::Type& Mode )
+FText FActorInfoColumn::MakeComboToolTipText( const SceneOutliner::ECustomColumnMode::Type& Mode ) const
 {
 	FText ToolTipText;
 
diff --git a/Engine/Source/Editor/SceneOutliner/Private/SceneOutlinerActorInfoColumn.h b/Engine/Source/Editor/SceneOutliner/Private/SceneOutlinerActorInfoColumn.h
--- a/Engine/Source/Editor/SceneOutliner/Private/SceneOutlinerActorInfoColumn.h
+++ b/Engine/Source/Editor/SceneOutliner/Private/SceneOutlinerActorInfoColumn.h
@@ -97,13 +97,17 @@ private:
 
 	TSharedPtr<SWidget> ConstructClassHyperlink( ISceneOutlinerTreeItem& TreeItem );
 
-	void OnModeChanged( TSharedPtr< SceneOutliner::ECustomColumnMode::Type > NewSelection, ESelectInfo::Type SelectInfo );
+	TSharedRef<SWidget> GenerateActorTypeMenu();
+
+	void OnModeChanged( const SceneOutliner::ECustomColumnMode::Type NewSelection );
+
+	bool IsCurrentMode( const SceneOutliner::ECustomColumnMode::Type Mode) const;
 
 	EVisibility GetColumnDataVisibility( bool bIsClassHyperlink ) const;
 
 	FText MakeComboText( const SceneOutliner::ECustomColumnMode::Type& Mode ) const;
 
-	FText MakeComboToolTipText( const SceneOutliner::ECustomColumnMode::Type& Mode );
+	FText MakeComboToolTipText( const SceneOutliner::ECustomColumnMode::Type& Mode ) const;
 
 	TSharedRef< ITableRow > MakeComboButtonItemWidget( TSharedPtr< SceneOutliner::ECustomColumnMode::Type > Mode, const TSharedRef<STableViewBase> & );
 
@@ -117,9 +121,6 @@ private:
 	 */
 	mutable SceneOutliner::ECustomColumnMode::Type CurrentMode;
 
-	/** A list of available custom column modes for Slate */
-	static TArray< TSharedPtr< SceneOutliner::ECustomColumnMode::Type > > ModeOptions;
-
 	/** Weak reference to the outliner widget that owns our list */
 	TWeakPtr< ISceneOutliner > SceneOutlinerWeak;
 };
