SceneOutliner: Fixed issue where unloaded actors were still being displayed in the outliner as deleted actors.

diff --git a/Engine/Source/Editor/SceneOutliner/Private/ActorHierarchy.cpp b/Engine/Source/Editor/SceneOutliner/Private/ActorHierarchy.cpp
--- a/Engine/Source/Editor/SceneOutliner/Private/ActorHierarchy.cpp
+++ b/Engine/Source/Editor/SceneOutliner/Private/ActorHierarchy.cpp
@@ -474,9 +474,10 @@ void FActorHierarchy::OnLevelRemoved(ULevel* InLevel, UWorld* InWorld)
 		{
 			if (Actor != nullptr)
 			{
-				EventData.ItemIDs.Add(Actor);		    
+				EventData.ItemIDs.Add(Actor);
 			}
 		}
+		HierarchyChangedEvent.Broadcast(EventData);
 	}
 }
 
