SceneOutliner: Fixed crash when LandscapeStreamingProxies are unloaded

diff --git a/Engine/Source/Editor/SceneOutliner/Private/ActorHierarchy.cpp b/Engine/Source/Editor/SceneOutliner/Private/ActorHierarchy.cpp
--- a/Engine/Source/Editor/SceneOutliner/Private/ActorHierarchy.cpp
+++ b/Engine/Source/Editor/SceneOutliner/Private/ActorHierarchy.cpp
@@ -122,7 +122,7 @@ FSceneOutlinerTreeItemPtr FActorHierarchy::FindParent(const ISceneOutlinerTreeIt
 			}
 
 			// Default to the world
-			if (const FSceneOutlinerTreeItemPtr* ParentItem = Items.Find(ActorTreeItem->Actor->GetWorld()))
+			if (const FSceneOutlinerTreeItemPtr* ParentItem = Items.Find(RepresentingWorld.Get()))
 			{
 				return *ParentItem;
 			}
