Fix Non-Unity compile error

diff --git a/Engine/Source/Editor/SceneOutliner/Private/SceneOutlinerActorInfoColumn.cpp b/Engine/Source/Editor/SceneOutliner/Private/SceneOutlinerActorInfoColumn.cpp
--- a/Engine/Source/Editor/SceneOutliner/Private/SceneOutlinerActorInfoColumn.cpp
+++ b/Engine/Source/Editor/SceneOutliner/Private/SceneOutlinerActorInfoColumn.cpp
@@ -6,6 +6,7 @@
 #include "Misc/PackageName.h"
 #include "Misc/StringBuilder.h"
 #include "Framework/Application/SlateApplication.h"
+#include "Framework/MultiBox/MultiBoxBuilder.h"
 #include "Widgets/Text/STextBlock.h"
 #include "Widgets/Views/SListView.h"
 #include "EditorStyleSet.h"
