Fix an issue in scene outliner world selection logic where it incorrectly assumes which PIE instance ID the first client is (reimplementation of UE-107043 from UE4)

diff --git a/Engine/Source/Editor/SceneOutliner/Private/ActorMode.cpp b/Engine/Source/Editor/SceneOutliner/Private/ActorMode.cpp
--- a/Engine/Source/Editor/SceneOutliner/Private/ActorMode.cpp
+++ b/Engine/Source/Editor/SceneOutliner/Private/ActorMode.cpp
@@ -135,6 +135,7 @@ void FActorMode::ChooseRepresentingWorld()
 	if (!RepresentingWorld.IsValid())
 	{
 		// ideally we want a PIE world that is standalone or the first client
+		int32 LowestClientInstanceSeen = MAX_int32;
 		for (const FWorldContext& Context : GEngine->GetWorldContexts())
 		{
 			UWorld* World = Context.World();
@@ -145,10 +146,10 @@ void FActorMode::ChooseRepresentingWorld()
 					RepresentingWorld = World;
 					break;
 				}
-				else if (World->GetNetMode() == NM_Client && Context.PIEInstance == 2)	// Slightly dangerous: assumes server is always PIEInstance = 1;
+				else if ((World->GetNetMode() == NM_Client) && (Context.PIEInstance < LowestClientInstanceSeen))
 				{
 					RepresentingWorld = World;
-					break;
+					LowestClientInstanceSeen = Context.PIEInstance;
 				}
 			}
 		}
@@ -156,7 +157,7 @@ void FActorMode::ChooseRepresentingWorld()
 
 	if (RepresentingWorld == nullptr)
 	{
-		// still not world so fallback to old logic where we just prefer PIE over Editor
+		// still no world so fallback to old logic where we just prefer PIE over Editor
 
 		for (const FWorldContext& Context : GEngine->GetWorldContexts())
 		{
