Marked FActorFolderPickingMode for export

diff --git a/Engine/Source/Editor/SceneOutliner/Public/ActorFolderPickingMode.h b/Engine/Source/Editor/SceneOutliner/Public/ActorFolderPickingMode.h
--- a/Engine/Source/Editor/SceneOutliner/Public/ActorFolderPickingMode.h
+++ b/Engine/Source/Editor/SceneOutliner/Public/ActorFolderPickingMode.h
@@ -4,7 +4,7 @@
 
 #include "ActorMode.h"
 
-class FActorFolderPickingMode : public FActorMode
+class SCENEOUTLINER_API FActorFolderPickingMode : public FActorMode
 {
 public:
 	FActorFolderPickingMode(SSceneOutliner* InSceneOutliner, FOnSceneOutlinerItemPicked InOnItemPicked, TWeakObjectPtr<UWorld> InSpecifiedWorldToDisplay = nullptr);
