SceneOutliner: Show unloaded world partition actors in the outliner

diff --git a/Engine/Source/Editor/SceneOutliner/Private/ActorBrowsingMode.cpp b/Engine/Source/Editor/SceneOutliner/Private/ActorBrowsingMode.cpp
--- a/Engine/Source/Editor/SceneOutliner/Private/ActorBrowsingMode.cpp
+++ b/Engine/Source/Editor/SceneOutliner/Private/ActorBrowsingMode.cpp
@@ -28,6 +28,7 @@
 #include "WorldTreeItem.h"
 #include "ComponentTreeItem.h"
 #include "ActorBrowsingModeSettings.h"
+#include "ActorDescTreeItem.h"
 #include "ScopedTransaction.h"
 #include "LevelInstance/LevelInstanceActor.h"
 #include "LevelInstance/LevelInstanceSubsystem.h"
@@ -42,6 +43,7 @@ UActorBrowsingModeSettings::UActorBrowsingModeSettings(const FObjectInitializer&
 { }
 
 using FActorFilter = TSceneOutlinerPredicateFilter<FActorTreeItem>;
+using FActorDescFilter = TSceneOutlinerPredicateFilter<FActorDescTreeItem>;
 
 FActorBrowsingMode::FActorBrowsingMode(SSceneOutliner* InSceneOutliner, TWeakObjectPtr<UWorld> InSpecifiedWorldToDisplay)
 	: FActorModeInteractive(FActorModeParams(InSceneOutliner, InSpecifiedWorldToDisplay, true, false))
@@ -121,6 +123,15 @@ FActorBrowsingMode::FActorBrowsingMode(SSceneOutliner* InSceneOutliner, TWeakObj
 			}
 		});
 	FilterInfoMap.Add(TEXT("HideLevelInstancesFilter"), HideLevelInstancesInfo);
+	
+	FSceneOutlinerFilterInfo HideUnloadedActorsInfo(LOCTEXT("ToggleHideUnloadedActors", "Hide Unloaded Actors"), LOCTEXT("ToggleHideUnloadedActorsToolTip", "When enabled, hides all unloaded world partition actors."), SharedSettings->bHideUnloadedActors, FCreateSceneOutlinerFilter::CreateStatic(&FActorBrowsingMode::CreateHideUnloadedActorsFilter));
+	HideUnloadedActorsInfo.OnToggle().AddLambda([this] (bool bIsActive)
+		{
+			UActorBrowsingModeSettings* Settings = GetMutableDefault<UActorBrowsingModeSettings>();
+	        Settings->bHideLevelInstanceHierarchy = bIsActive;
+	        Settings->PostEditChange();
+		});
+	FilterInfoMap.Add(TEXT("HideUnloadedActorsFilter"), HideUnloadedActorsInfo);
 
 	// Add a filter which sets the interactive mode of LevelInstance items and their children
 	SceneOutliner->AddFilter(MakeShared<FActorFilter>(FActorTreeItem::FFilterPredicate::CreateStatic([](const AActor* Actor) {return true; }), FSceneOutlinerFilter::EDefaultBehaviour::Pass, FActorTreeItem::FFilterPredicate::CreateLambda([this](const AActor* Actor)
@@ -313,6 +324,12 @@ TSharedRef<FSceneOutlinerFilter> FActorBrowsingMode::CreateHideLevelInstancesFil
 		}), FSceneOutlinerFilter::EDefaultBehaviour::Pass));
 }
 
+TSharedRef<FSceneOutlinerFilter> FActorBrowsingMode::CreateHideUnloadedActorsFilter()
+{
+	return MakeShareable(new FActorDescFilter(FActorDescTreeItem::FFilterPredicate::CreateStatic(
+		[](const FWorldPartitionActorDesc* ActorDesc) { return false; }), FSceneOutlinerFilter::EDefaultBehaviour::Pass));
+}
+
 static const FName DefaultContextBaseMenuName("SceneOutliner.DefaultContextMenuBase");
 static const FName DefaultContextMenuName("SceneOutliner.DefaultContextMenu");
 
diff --git a/Engine/Source/Editor/SceneOutliner/Private/ActorBrowsingModeSettings.h b/Engine/Source/Editor/SceneOutliner/Private/ActorBrowsingModeSettings.h
--- a/Engine/Source/Editor/SceneOutliner/Private/ActorBrowsingModeSettings.h
+++ b/Engine/Source/Editor/SceneOutliner/Private/ActorBrowsingModeSettings.h
@@ -36,4 +36,8 @@ class UActorBrowsingModeSettings
 	/** True when the Scene Outliner is not displaying LevelInstances */
 	UPROPERTY(config)
 	uint32 bHideLevelInstanceHierarchy : 1;
+
+	/** True when the Scene Outliner is not displaying unloaded actors */
+	UPROPERTY(config)
+	uint32 bHideUnloadedActors : 1;
 };
diff --git a/Engine/Source/Editor/SceneOutliner/Private/ActorDescTreeItem.cpp b/Engine/Source/Editor/SceneOutliner/Private/ActorDescTreeItem.cpp
new file mode 100644
--- /dev/null
+++ b/Engine/Source/Editor/SceneOutliner/Private/ActorDescTreeItem.cpp
@@ -0,0 +1,234 @@
+﻿// Copyright Epic Games, Inc. All Rights Reserved.
+
+#include "ActorDescTreeItem.h"
+#include "WorldPartition/WorldPartitionActorDesc.h"
+#include "Widgets/DeclarativeSyntaxSupport.h"
+#include "Layout/WidgetPath.h"
+#include "Framework/Application/MenuStack.h"
+#include "Framework/Application/SlateApplication.h"
+#include "Editor.h"
+#include "ScopedTransaction.h"
+#include "SceneOutlinerPublicTypes.h"
+#include "SceneOutlinerDragDrop.h"
+#include "SceneOutlinerStandaloneTypes.h"
+#include "Widgets/Text/SInlineEditableTextBlock.h"
+#include "ActorEditorUtils.h"
+#include "ClassIconFinder.h"
+#include "ISceneOutliner.h"
+#include "ISceneOutlinerMode.h"
+#include "Logging/MessageLog.h"
+#include "SSocketChooser.h"
+
+#define LOCTEXT_NAMESPACE "SceneOutliner_ActorDescTreeItem"
+
+const FSceneOutlinerTreeItemType FActorDescTreeItem::Type(&ISceneOutlinerTreeItem::Type);
+
+struct SActorDescTreeLabel : FSceneOutlinerCommonLabelData, public SCompoundWidget
+{
+	SLATE_BEGIN_ARGS(SActorDescTreeLabel) {}
+	SLATE_END_ARGS()
+
+	void Construct(const FArguments& InArgs, FActorDescTreeItem& ActorDescItem, ISceneOutliner& SceneOutliner, const STableRow<FSceneOutlinerTreeItemPtr>& InRow)
+	{
+		WeakSceneOutliner = StaticCastSharedRef<ISceneOutliner>(SceneOutliner.AsShared());
+
+		TreeItemPtr = StaticCastSharedRef<FActorDescTreeItem>(ActorDescItem.AsShared());
+		ActorDescPtr = ActorDescItem.ActorDesc;
+
+		HighlightText = SceneOutliner.GetFilterHighlightText();
+
+		TSharedPtr<SInlineEditableTextBlock> InlineTextBlock;
+
+		auto MainContent = SNew(SHorizontalBox)
+			// Main actor desc label
+			+SHorizontalBox::Slot()
+			.VAlign(VAlign_Center)
+			[
+				SAssignNew(InlineTextBlock, SInlineEditableTextBlock)
+				.Text(this, &SActorDescTreeLabel::GetDisplayText)
+				.ToolTipText(this, &SActorDescTreeLabel::GetTooltipText)
+				.HighlightText(HighlightText)
+				.ColorAndOpacity(this, &SActorDescTreeLabel::GetForegroundColor)
+				.OnTextCommitted(this, &SActorDescTreeLabel::OnLabelCommitted)
+				.OnVerifyTextChanged(this, &SActorDescTreeLabel::OnVerifyItemLabelChanged)
+				.IsSelected(FIsSelected::CreateSP(&InRow, &STableRow<FSceneOutlinerTreeItemPtr>::IsSelectedExclusively))
+				.IsReadOnly_Lambda([Item = ActorDescItem.AsShared(), this]()
+				{
+					return !CanExecuteRenameRequest(Item.Get());
+				})
+			]
+
+			+ SHorizontalBox::Slot()
+			.VAlign(VAlign_Center)
+			.AutoWidth()
+			.Padding(0.0f, 0.f, 3.0f, 0.0f)
+			[
+				SNew(STextBlock)
+				.Text(this, &SActorDescTreeLabel::GetTypeText)
+				.Visibility(this, &SActorDescTreeLabel::GetTypeTextVisibility)
+				.HighlightText(HighlightText)
+			];
+
+		if (WeakSceneOutliner.Pin()->GetMode()->IsInteractive())
+		{
+			ActorDescItem.RenameRequestEvent.BindSP(InlineTextBlock.Get(), &SInlineEditableTextBlock::EnterEditingMode);
+		}
+
+		ChildSlot
+			[
+				SNew(SHorizontalBox)
+
+				+ SHorizontalBox::Slot()
+				.AutoWidth()
+				.VAlign(VAlign_Center)
+				.Padding(FSceneOutlinerDefaultTreeItemMetrics::IconPadding())
+				[
+					SNew(SBox)
+					.WidthOverride(FSceneOutlinerDefaultTreeItemMetrics::IconSize())
+					.HeightOverride(FSceneOutlinerDefaultTreeItemMetrics::IconSize())
+					[
+						SNew(SImage)
+						.Image(this, &SActorDescTreeLabel::GetIcon)
+						.ToolTipText(this, &SActorDescTreeLabel::GetIconTooltip)
+						.ColorAndOpacity(FSlateColor::UseForeground())
+					]
+				]
+
+				+ SHorizontalBox::Slot()
+				.FillWidth(1.0f)
+				.VAlign(VAlign_Center)
+				.Padding(0.0f, 0.0f)
+				[
+					MainContent
+				]
+			];
+	}
+
+private:
+	TWeakPtr<FActorDescTreeItem> TreeItemPtr;
+	const FWorldPartitionActorDesc* ActorDescPtr = nullptr;
+	TAttribute<FText> HighlightText;
+
+	FText GetDisplayText() const
+	{
+		FFormatNamedArguments Args;
+		Args.Add(TEXT("ActorLabel"), FText::FromName(ActorDescPtr->GetActorLabel()));
+		Args.Add(TEXT("UnloadedTag"), LOCTEXT("UnloadedActorLabel", "(Unloaded)"));
+		return FText::Format(LOCTEXT("UnloadedActorDisplay", "{ActorLabel} {UnloadedTag}"), Args);
+	}
+
+	FText GetTooltipText() const
+	{
+		return FText();
+	}
+
+	FText GetTypeText() const
+	{
+		return ActorDescPtr ? FText::FromName(ActorDescPtr->GetActorClass()->GetFName()) : FText();
+	}
+
+	EVisibility GetTypeTextVisibility() const
+	{
+		return HighlightText.Get().IsEmpty() ? EVisibility::Collapsed : EVisibility::Visible;
+	}
+
+	const FSlateBrush* GetIcon() const
+	{
+		if (ActorDescPtr && WeakSceneOutliner.IsValid())
+		{
+			const FName IconName = ActorDescPtr->GetActorClass()->GetFName();
+			
+			const FSlateBrush* CachedBrush = WeakSceneOutliner.Pin()->GetCachedIconForClass(IconName);
+			if (CachedBrush != nullptr)
+			{
+				return CachedBrush;
+			}
+			else if (IconName != NAME_None)
+			{
+
+				const FSlateBrush* FoundSlateBrush = FSlateIconFinder::FindIconForClass(ActorDescPtr->GetActorClass()).GetIcon();
+				WeakSceneOutliner.Pin()->CacheIconForClass(IconName, FoundSlateBrush);
+				return FoundSlateBrush;
+			}
+		}
+		return nullptr;
+	}
+
+	const FSlateBrush* GetIconOverlay() const
+	{
+		return nullptr;
+	}
+
+	FText GetIconTooltip() const
+	{
+		return FText();
+	}
+
+	virtual FSlateColor GetForegroundColor() const override
+	{
+		const auto TreeItem = TreeItemPtr.Pin();
+		if (auto BaseColor = FSceneOutlinerCommonLabelData::GetForegroundColor(*TreeItem))
+		{
+			return BaseColor.GetValue();
+		}
+
+		if (ActorDescPtr == nullptr)
+		{
+			// Deleted actor!
+			return FLinearColor(0.2f, 0.2f, 0.25f);
+		}
+
+		return FSceneOutlinerCommonLabelData::DarkColor;
+	}
+
+	bool OnVerifyItemLabelChanged(const FText&, FText&)
+	{
+		// don't allow label change for unloaded actor items
+		return false;
+	}
+
+	void OnLabelCommitted(const FText& InLabel, ETextCommit::Type InCommitInfo)
+	{
+		// not supported.
+	}
+};
+
+FActorDescTreeItem::FActorDescTreeItem(const FWorldPartitionActorDesc* InActorDesc)
+	: ISceneOutlinerTreeItem(Type)
+	, ActorDesc(InActorDesc)
+{
+	check (InActorDesc != nullptr);
+	ID = FSceneOutlinerTreeItemID(ActorDesc->GetGuid());
+}
+
+FSceneOutlinerTreeItemID FActorDescTreeItem::GetID() const
+{
+	return ID;
+}
+
+FString FActorDescTreeItem::GetDisplayString() const
+{
+	return ActorDesc ? ActorDesc->GetActorLabel().ToString() : LOCTEXT("ActorLabelForMissingActor", "(Deleted Actor)").ToString();
+}
+
+bool FActorDescTreeItem::CanInteract() const
+{
+	return false;
+}
+
+TSharedRef<SWidget> FActorDescTreeItem::GenerateLabelWidget(ISceneOutliner& Outliner, const STableRow<FSceneOutlinerTreeItemPtr>& InRow)
+{
+	return SNew(SActorDescTreeLabel, *this, Outliner, InRow);
+}
+
+void FActorDescTreeItem::OnVisibilityChanged(const bool bNewVisibility)
+{
+
+}
+
+bool FActorDescTreeItem::GetVisibility() const
+{
+	return true;
+}
+
+#undef LOCTEXT_NAMESPACE
diff --git a/Engine/Source/Editor/SceneOutliner/Private/ActorHierarchy.cpp b/Engine/Source/Editor/SceneOutliner/Private/ActorHierarchy.cpp
--- a/Engine/Source/Editor/SceneOutliner/Private/ActorHierarchy.cpp
+++ b/Engine/Source/Editor/SceneOutliner/Private/ActorHierarchy.cpp
@@ -6,6 +6,7 @@
 #include "EngineUtils.h"
 #include "WorldTreeItem.h"
 #include "ActorTreeItem.h"
+#include "ActorDescTreeItem.h"
 #include "ComponentTreeItem.h"
 #include "ActorFolderTreeItem.h"
 #include "ISceneOutlinerMode.h"
@@ -16,6 +17,8 @@
 #include "EditorFolderUtils.h"
 #include "LevelInstance/LevelInstanceActor.h"
 #include "LevelInstance/LevelInstanceSubsystem.h"
+#include "WorldPartition/WorldPartitionSubsystem.h"
+#include "WorldPartition/WorldPartitionActorDesc.h"
 
 TUniquePtr<FActorHierarchy> FActorHierarchy::Create(ISceneOutlinerMode* Mode, const TWeakObjectPtr<UWorld>& World)
 {
@@ -162,6 +165,27 @@ FSceneOutlinerTreeItemPtr FActorHierarchy::FindParent(const ISceneOutlinerTreeIt
 			}
 		}
 	}
+	else if (const FActorDescTreeItem* ActorDescItem = Item.CastTo<FActorDescTreeItem>())
+	{
+		const FWorldPartitionActorDesc* ActorDesc = ActorDescItem->ActorDesc;
+		check(ActorDesc != nullptr);
+		const FName FolderPath = ActorDesc->GetFolderPath();
+		if (!FolderPath.IsNone())
+		{
+			if (const FSceneOutlinerTreeItemPtr* UnloadedActorItem = Items.Find(FolderPath))
+			{
+				return *UnloadedActorItem;
+			}
+		}
+		else
+		{
+			// Default to the world
+			if (const FSceneOutlinerTreeItemPtr* ParentItem = Items.Find(RepresentingWorld.Get()))
+			{
+				return *ParentItem;
+			}
+		}
+	}
 	return nullptr;
 }
 
@@ -236,8 +260,26 @@ void FActorHierarchy::CreateItems(TArray<FSceneOutlinerTreeItemPtr>& OutItems) c
 		OutItems.Add(WorldItem);
 	}
 
+	UWorld* RepresentingWorldPtr = RepresentingWorld.Get();
+	check(RepresentingWorldPtr);
+
 	// Create world children regardless of if a world item was created
-	CreateWorldChildren(RepresentingWorld.Get(), OutItems);
+	CreateWorldChildren(RepresentingWorldPtr, OutItems);
+
+	if (UWorldPartitionSubsystem* WorldPartitionSubsystem = RepresentingWorldPtr->GetSubsystem<UWorldPartitionSubsystem>())
+	{
+		WorldPartitionSubsystem->ForEachActorDesc(AActor::StaticClass(), [this, &OutItems](const FWorldPartitionActorDesc* ActorDesc)
+		{
+			if (ActorDesc != nullptr && !ActorDesc->IsLoaded())
+			{
+				if (const FSceneOutlinerTreeItemPtr ActorDescItem = Mode->CreateItemFor<FActorDescTreeItem>(ActorDesc))
+				{
+					OutItems.Add(ActorDescItem);
+				}
+			}
+			return true;
+		});
+	}
 }
 
 void FActorHierarchy::CreateChildren(const FSceneOutlinerTreeItemPtr& Item, TArray<FSceneOutlinerTreeItemPtr>& OutChildren) const
@@ -378,6 +420,15 @@ FSceneOutlinerTreeItemPtr FActorHierarchy::CreateParentItem(const FSceneOutliner
 			return Mode->CreateItemFor<FActorFolderTreeItem>(FActorFolderTreeItem(ParentPath, FolderTreeItem->World), true);
 		}
 	}
+	else if (const FActorDescTreeItem* ActorDescItem = Item->CastTo<FActorDescTreeItem>())
+	{
+		const FName ActorDescPath = ActorDescItem->ActorDesc->GetFolderPath();
+		if (Mode->ShouldShowFolders() && !ActorDescPath.IsNone())
+		{
+			return Mode->CreateItemFor<FActorFolderTreeItem>(FActorFolderTreeItem(ActorDescPath, RepresentingWorld.Get()), true);
+		}
+	}
+	
 	return nullptr;
 }
 
diff --git a/Engine/Source/Editor/SceneOutliner/Private/ActorMode.cpp b/Engine/Source/Editor/SceneOutliner/Private/ActorMode.cpp
--- a/Engine/Source/Editor/SceneOutliner/Private/ActorMode.cpp
+++ b/Engine/Source/Editor/SceneOutliner/Private/ActorMode.cpp
@@ -13,6 +13,7 @@
 #include "ActorTreeItem.h"
 #include "FolderTreeItem.h"
 #include "ComponentTreeItem.h"
+#include "ActorDescTreeItem.h"
 #include "WorldTreeItem.h"
 #include "LevelInstance/LevelInstanceEditorInstanceActor.h"
 
@@ -333,6 +334,11 @@ int32 FActorMode::GetTypeSortPriority(const ISceneOutlinerTreeItem& Item) const
 	{
 		return EItemSortOrder::Actor;
 	}
+	else if (Item.IsA<FActorDescTreeItem>())
+	{
+		return EItemSortOrder::Unloaded;
+	}
+
 	// Warning: using actor mode with an unsupported item type!
 	check(false);
 	return -1;
diff --git a/Engine/Source/Editor/SceneOutliner/Private/SceneOutlinerActorInfoColumn.cpp b/Engine/Source/Editor/SceneOutliner/Private/SceneOutlinerActorInfoColumn.cpp
--- a/Engine/Source/Editor/SceneOutliner/Private/SceneOutlinerActorInfoColumn.cpp
+++ b/Engine/Source/Editor/SceneOutliner/Private/SceneOutlinerActorInfoColumn.cpp
@@ -14,11 +14,13 @@
 #include "SortHelper.h"
 #include "ISceneOutliner.h"
 #include "ActorTreeItem.h"
+#include "ActorDescTreeItem.h"
 #include "ComponentTreeItem.h"
 #include "FolderTreeItem.h"
 #include "WorldTreeItem.h"
 #include "WorldPartition/DataLayer/DataLayer.h"
 #include "WorldPartition/DataLayer/WorldDataLayers.h"
+#include "WorldPartition/WorldPartitionActorDesc.h"
 
 #define LOCTEXT_NAMESPACE "SceneOutlinerActorInfoColumn"
 
@@ -152,6 +154,24 @@ struct FGetInfo
 				return FString();
 			}
 		}
+		else if (const FActorDescTreeItem* ActorDescItem = Item.CastTo<FActorDescTreeItem>())
+		{
+			const FWorldPartitionActorDesc* ActorDesc = ActorDescItem->ActorDesc;
+
+			if (!ActorDesc)
+			{
+				return FString();
+			}
+			
+			switch (CurrentMode)
+			{
+			case SceneOutliner::ECustomColumnMode::Class:
+				return ActorDesc->GetActorClass()->GetName();
+			default:
+				return FString();
+			}
+		}
+
 
 		return FString();
 	}
diff --git a/Engine/Source/Editor/SceneOutliner/Public/ActorBrowsingMode.h b/Engine/Source/Editor/SceneOutliner/Public/ActorBrowsingMode.h
--- a/Engine/Source/Editor/SceneOutliner/Public/ActorBrowsingMode.h
+++ b/Engine/Source/Editor/SceneOutliner/Public/ActorBrowsingMode.h
@@ -99,6 +99,7 @@ private:
 	static TSharedRef<FSceneOutlinerFilter> CreateIsInCurrentLevelFilter();
 	static TSharedRef<FSceneOutlinerFilter> CreateHideComponentsFilter();
 	static TSharedRef<FSceneOutlinerFilter> CreateHideLevelInstancesFilter();
+	static TSharedRef<FSceneOutlinerFilter> CreateHideUnloadedActorsFilter();
 	
 	void ToggleActorSCCStatusColumn();
 	bool isActorSCCStatusColumnActive() const { return bActorSCCStatusColumnActive; }
diff --git a/Engine/Source/Editor/SceneOutliner/Public/ActorDescTreeItem.h b/Engine/Source/Editor/SceneOutliner/Public/ActorDescTreeItem.h
new file mode 100644
--- /dev/null
+++ b/Engine/Source/Editor/SceneOutliner/Public/ActorDescTreeItem.h
@@ -0,0 +1,51 @@
+// Copyright Epic Games, Inc. All Rights Reserved.
+
+#pragma once
+
+#include "CoreMinimal.h"
+#include "Widgets/SWidget.h"
+#include "ISceneOutlinerTreeItem.h"
+#include "UObject/ObjectKey.h"
+
+class FWorldPartitionActorDesc;
+
+/** A tree item that represents an actor in the world */
+struct SCENEOUTLINER_API FActorDescTreeItem : ISceneOutlinerTreeItem
+{
+public:
+	DECLARE_DELEGATE_RetVal_OneParam(bool, FFilterPredicate, const FWorldPartitionActorDesc*);
+	DECLARE_DELEGATE_RetVal_OneParam(bool, FInteractivePredicate, const FWorldPartitionActorDesc*);
+		
+	bool Filter(FFilterPredicate Pred) const
+	{
+		return Pred.Execute(ActorDesc);
+	}
+
+	bool GetInteractiveState(FInteractivePredicate Pred) const 
+	{
+		return Pred.Execute(ActorDesc);
+	}
+
+	/** The actor this tree item is associated with. */
+	const FWorldPartitionActorDesc* ActorDesc;
+
+	/** Constant identifier for this tree item */
+	FSceneOutlinerTreeItemID ID;
+
+	/** Static type identifier for this tree item class */
+	static const FSceneOutlinerTreeItemType Type;
+
+	/** Construct this item from an actor */
+	FActorDescTreeItem(const FWorldPartitionActorDesc* InActorDesc);
+
+	/* Begin ISceneOutlinerTreeItem Implementation */
+	virtual bool IsValid() const override { return ActorDesc != nullptr; }
+	virtual FSceneOutlinerTreeItemID GetID() const override;
+	virtual FString GetDisplayString() const override;
+	virtual bool CanInteract() const override;
+	virtual TSharedRef<SWidget> GenerateLabelWidget(ISceneOutliner& Outliner, const STableRow<FSceneOutlinerTreeItemPtr>& InRow) override;
+	virtual void OnVisibilityChanged(const bool bNewVisibility) override;
+	virtual bool HasVisibilityInfo() const override { return true; }
+	virtual bool GetVisibility() const override;
+	/* End ISceneOutlinerTreeItem Implementation */
+};
diff --git a/Engine/Source/Editor/SceneOutliner/Public/ActorMode.h b/Engine/Source/Editor/SceneOutliner/Public/ActorMode.h
--- a/Engine/Source/Editor/SceneOutliner/Public/ActorMode.h
+++ b/Engine/Source/Editor/SceneOutliner/Public/ActorMode.h
@@ -41,7 +41,7 @@ class SCENEOUTLINER_API FActorMode : public ISceneOutlinerMode
 public:
 	struct EItemSortOrder
 	{
-		enum Type { World = 0, Folder = 10, Actor = 20 };
+		enum Type { World = 0, Folder = 10, Actor = 20, Unloaded = 30 };
 	};
 	
 	FActorMode(const FActorModeParams& Params);
diff --git a/Engine/Source/Editor/SceneOutliner/Public/SceneOutlinerStandaloneTypes.h b/Engine/Source/Editor/SceneOutliner/Public/SceneOutlinerStandaloneTypes.h
--- a/Engine/Source/Editor/SceneOutliner/Public/SceneOutlinerStandaloneTypes.h
+++ b/Engine/Source/Editor/SceneOutliner/Public/SceneOutlinerStandaloneTypes.h
@@ -14,7 +14,7 @@ struct FSceneOutlinerTreeItemID
 public:
 	using FUniqueID = uint32;
 
-	enum class EType : uint8 { Object, Folder, UniqueID, Null };
+	enum class EType : uint8 { Object, Folder, UniqueID, Guid, Null };
 
 	/** Default constructed null item ID */
 	FSceneOutlinerTreeItemID() : Type(EType::Null), CachedHash(0) {}
@@ -46,6 +46,12 @@ public:
 		CachedHash = CalculateTypeHash();
 	}
 
+	FSceneOutlinerTreeItemID(const FGuid& InGuid) : Type(EType::Guid)
+	{
+		new (Data) FGuid(InGuid);
+		CachedHash = CalculateTypeHash();
+	}
+
 	/** Copy construction / assignment */
 	FSceneOutlinerTreeItemID(const FSceneOutlinerTreeItemID& Other)
 	{
@@ -58,7 +64,8 @@ public:
 		{
 			case EType::Object:			new (Data) FObjectKey(Other.GetAsObjectKey());		break;
 			case EType::Folder:			new (Data) FName(Other.GetAsFolderRef());			break;
-			case EType::UniqueID:		new (Data) FUniqueID(Other.GetAsHash());	break;
+			case EType::Guid:			new (Data) FGuid(Other.GetAsGuid());				break;
+			case EType::UniqueID:		new (Data) FUniqueID(Other.GetAsHash());			break;
 			default:																		break;
 		}
 
@@ -83,6 +90,7 @@ public:
 		{
 			case EType::Object:			GetAsObjectKey().~FObjectKey();							break;
 			case EType::Folder:			GetAsFolderRef().~FName();								break;
+			case EType::Guid:			GetAsGuid().~FGuid();									break;
 			case EType::UniqueID:		/* NOP */												break;
 			default:																			break;
 		}
@@ -104,7 +112,8 @@ public:
 		{
 			case EType::Object:			Hash = GetTypeHash(GetAsObjectKey());				break;
 			case EType::Folder:			Hash = GetTypeHash(GetAsFolderRef());				break;
-			case EType::UniqueID:		Hash = GetAsHash();									break;
+			case EType::Guid:			Hash = GetTypeHash(GetAsGuid());					break;
+			case EType::UniqueID:		Hash = GetTypeHash(GetAsHash());					break;
 			default:																		break;
 		}
 
@@ -120,6 +129,7 @@ private:
 
 	FObjectKey& 	GetAsObjectKey() const 		{ return *reinterpret_cast<FObjectKey*>(Data); }
 	FName& 			GetAsFolderRef() const		{ return *reinterpret_cast<FName*>(Data); }
+	FGuid&			GetAsGuid() const			{ return *reinterpret_cast<FGuid*>(Data); }
 	FUniqueID&		GetAsHash() const			{ return *reinterpret_cast<FUniqueID*>(Data); }
 
 	/** Compares the specified ID with this one - assumes matching types */
@@ -129,6 +139,7 @@ private:
 		{
 			case EType::Object:			return GetAsObjectKey() == Other.GetAsObjectKey();
 			case EType::Folder:			return GetAsFolderRef() == Other.GetAsFolderRef();
+			case EType::Guid:			return GetAsGuid() == Other.GetAsGuid();
 			case EType::UniqueID:		return GetAsHash() == Other.GetAsHash();
 			case EType::Null:			return true;
 			default: check(false);		return false;
@@ -138,7 +149,7 @@ private:
 	EType Type;
 
 	uint32 CachedHash;
-	static const uint32 MaxSize = TMaxSizeof<FObjectKey, FName, FUniqueID>::Value;
+	static const uint32 MaxSize = TMaxSizeof<FObjectKey, FName, FGuid, FUniqueID>::Value;
 	mutable uint8 Data[MaxSize];
 };
 	
