Backout 15629005.

diff --git a/Engine/Source/Editor/SceneOutliner/Private/ActorBrowsingMode.cpp b/Engine/Source/Editor/SceneOutliner/Private/ActorBrowsingMode.cpp
--- a/Engine/Source/Editor/SceneOutliner/Private/ActorBrowsingMode.cpp
+++ b/Engine/Source/Editor/SceneOutliner/Private/ActorBrowsingMode.cpp
@@ -28,7 +28,6 @@
 #include "WorldTreeItem.h"
 #include "ComponentTreeItem.h"
 #include "ActorBrowsingModeSettings.h"
-#include "ActorDescTreeItem.h"
 #include "ScopedTransaction.h"
 #include "LevelInstance/LevelInstanceActor.h"
 #include "LevelInstance/LevelInstanceSubsystem.h"
@@ -43,7 +42,6 @@ UActorBrowsingModeSettings::UActorBrowsingModeSettings(const FObjectInitializer&
 { }
 
 using FActorFilter = TSceneOutlinerPredicateFilter<FActorTreeItem>;
-using FActorDescFilter = TSceneOutlinerPredicateFilter<FActorDescTreeItem>;
 
 FActorBrowsingMode::FActorBrowsingMode(SSceneOutliner* InSceneOutliner, TWeakObjectPtr<UWorld> InSpecifiedWorldToDisplay)
 	: FActorModeInteractive(FActorModeParams(InSceneOutliner, InSpecifiedWorldToDisplay, true, false))
@@ -123,15 +121,6 @@ FActorBrowsingMode::FActorBrowsingMode(SSceneOutliner* InSceneOutliner, TWeakObj
 			}
 		});
 	FilterInfoMap.Add(TEXT("HideLevelInstancesFilter"), HideLevelInstancesInfo);
-	
-	FSceneOutlinerFilterInfo HideUnloadedActorsInfo(LOCTEXT("ToggleHideUnloadedActors", "Hide Unloaded Actors"), LOCTEXT("ToggleHideUnloadedActorsToolTip", "When enabled, hides all unloaded world partition actors."), SharedSettings->bHideUnloadedActors, FCreateSceneOutlinerFilter::CreateStatic(&FActorBrowsingMode::CreateHideUnloadedActorsFilter));
-	HideUnloadedActorsInfo.OnToggle().AddLambda([this] (bool bIsActive)
-		{
-			UActorBrowsingModeSettings* Settings = GetMutableDefault<UActorBrowsingModeSettings>();
-	        Settings->bHideLevelInstanceHierarchy = bIsActive;
-	        Settings->PostEditChange();
-		});
-	FilterInfoMap.Add(TEXT("HideUnloadedActorsFilter"), HideUnloadedActorsInfo);
 
 	// Add a filter which sets the interactive mode of LevelInstance items and their children
 	SceneOutliner->AddFilter(MakeShared<FActorFilter>(FActorTreeItem::FFilterPredicate::CreateStatic([](const AActor* Actor) {return true; }), FSceneOutlinerFilter::EDefaultBehaviour::Pass, FActorTreeItem::FFilterPredicate::CreateLambda([this](const AActor* Actor)
@@ -324,12 +313,6 @@ TSharedRef<FSceneOutlinerFilter> FActorBrowsingMode::CreateHideLevelInstancesFil
 		}), FSceneOutlinerFilter::EDefaultBehaviour::Pass));
 }
 
-TSharedRef<FSceneOutlinerFilter> FActorBrowsingMode::CreateHideUnloadedActorsFilter()
-{
-	return MakeShareable(new FActorDescFilter(FActorDescTreeItem::FFilterPredicate::CreateStatic(
-		[](const FWorldPartitionActorDesc* ActorDesc) { return false; }), FSceneOutlinerFilter::EDefaultBehaviour::Pass));
-}
-
 static const FName DefaultContextBaseMenuName("SceneOutliner.DefaultContextMenuBase");
 static const FName DefaultContextMenuName("SceneOutliner.DefaultContextMenu");
 
diff --git a/Engine/Source/Editor/SceneOutliner/Private/ActorBrowsingModeSettings.h b/Engine/Source/Editor/SceneOutliner/Private/ActorBrowsingModeSettings.h
--- a/Engine/Source/Editor/SceneOutliner/Private/ActorBrowsingModeSettings.h
+++ b/Engine/Source/Editor/SceneOutliner/Private/ActorBrowsingModeSettings.h
@@ -36,8 +36,4 @@ class UActorBrowsingModeSettings
 	/** True when the Scene Outliner is not displaying LevelInstances */
 	UPROPERTY(config)
 	uint32 bHideLevelInstanceHierarchy : 1;
-
-	/** True when the Scene Outliner is not displaying unloaded actors */
-	UPROPERTY(config)
-	uint32 bHideUnloadedActors : 1;
 };
diff --git a/Engine/Source/Editor/SceneOutliner/Private/ActorDescTreeItem.cpp b/Engine/Source/Editor/SceneOutliner/Private/ActorDescTreeItem.cpp
deleted file mode 100644
diff --git a/Engine/Source/Editor/SceneOutliner/Private/ActorHierarchy.cpp b/Engine/Source/Editor/SceneOutliner/Private/ActorHierarchy.cpp
--- a/Engine/Source/Editor/SceneOutliner/Private/ActorHierarchy.cpp
+++ b/Engine/Source/Editor/SceneOutliner/Private/ActorHierarchy.cpp
@@ -6,7 +6,6 @@
 #include "EngineUtils.h"
 #include "WorldTreeItem.h"
 #include "ActorTreeItem.h"
-#include "ActorDescTreeItem.h"
 #include "ComponentTreeItem.h"
 #include "ActorFolderTreeItem.h"
 #include "ISceneOutlinerMode.h"
@@ -17,8 +16,6 @@
 #include "EditorFolderUtils.h"
 #include "LevelInstance/LevelInstanceActor.h"
 #include "LevelInstance/LevelInstanceSubsystem.h"
-#include "WorldPartition/WorldPartitionSubsystem.h"
-#include "WorldPartition/WorldPartitionActorDesc.h"
 
 TUniquePtr<FActorHierarchy> FActorHierarchy::Create(ISceneOutlinerMode* Mode, const TWeakObjectPtr<UWorld>& World)
 {
@@ -165,27 +162,6 @@ FSceneOutlinerTreeItemPtr FActorHierarchy::FindParent(const ISceneOutlinerTreeIt
 			}
 		}
 	}
-	else if (const FActorDescTreeItem* ActorDescItem = Item.CastTo<FActorDescTreeItem>())
-	{
-		const FWorldPartitionActorDesc* ActorDesc = ActorDescItem->ActorDesc;
-		check(ActorDesc != nullptr);
-		const FName FolderPath = ActorDesc->GetFolderPath();
-		if (!FolderPath.IsNone())
-		{
-			if (const FSceneOutlinerTreeItemPtr* UnloadedActorItem = Items.Find(FolderPath))
-			{
-				return *UnloadedActorItem;
-			}
-		}
-		else
-		{
-			// Default to the world
-			if (const FSceneOutlinerTreeItemPtr* ParentItem = Items.Find(RepresentingWorld.Get()))
-			{
-				return *ParentItem;
-			}
-		}
-	}
 	return nullptr;
 }
 
@@ -260,26 +236,8 @@ void FActorHierarchy::CreateItems(TArray<FSceneOutlinerTreeItemPtr>& OutItems) c
 		OutItems.Add(WorldItem);
 	}
 
-	UWorld* RepresentingWorldPtr = RepresentingWorld.Get();
-	check(RepresentingWorldPtr);
-
 	// Create world children regardless of if a world item was created
-	CreateWorldChildren(RepresentingWorldPtr, OutItems);
-
-	if (UWorldPartitionSubsystem* WorldPartitionSubsystem = RepresentingWorldPtr->GetSubsystem<UWorldPartitionSubsystem>())
-	{
-		WorldPartitionSubsystem->ForEachActorDesc(AActor::StaticClass(), [this, &OutItems](const FWorldPartitionActorDesc* ActorDesc)
-		{
-			if (ActorDesc != nullptr && !ActorDesc->IsLoaded())
-			{
-				if (const FSceneOutlinerTreeItemPtr ActorDescItem = Mode->CreateItemFor<FActorDescTreeItem>(ActorDesc))
-				{
-					OutItems.Add(ActorDescItem);
-				}
-			}
-			return true;
-		});
-	}
+	CreateWorldChildren(RepresentingWorld.Get(), OutItems);
 }
 
 void FActorHierarchy::CreateChildren(const FSceneOutlinerTreeItemPtr& Item, TArray<FSceneOutlinerTreeItemPtr>& OutChildren) const
@@ -420,15 +378,6 @@ FSceneOutlinerTreeItemPtr FActorHierarchy::CreateParentItem(const FSceneOutliner
 			return Mode->CreateItemFor<FActorFolderTreeItem>(FActorFolderTreeItem(ParentPath, FolderTreeItem->World), true);
 		}
 	}
-	else if (const FActorDescTreeItem* ActorDescItem = Item->CastTo<FActorDescTreeItem>())
-	{
-		const FName ActorDescPath = ActorDescItem->ActorDesc->GetFolderPath();
-		if (Mode->ShouldShowFolders() && !ActorDescPath.IsNone())
-		{
-			return Mode->CreateItemFor<FActorFolderTreeItem>(FActorFolderTreeItem(ActorDescPath, RepresentingWorld.Get()), true);
-		}
-	}
-	
 	return nullptr;
 }
 
diff --git a/Engine/Source/Editor/SceneOutliner/Private/ActorMode.cpp b/Engine/Source/Editor/SceneOutliner/Private/ActorMode.cpp
--- a/Engine/Source/Editor/SceneOutliner/Private/ActorMode.cpp
+++ b/Engine/Source/Editor/SceneOutliner/Private/ActorMode.cpp
@@ -13,7 +13,6 @@
 #include "ActorTreeItem.h"
 #include "FolderTreeItem.h"
 #include "ComponentTreeItem.h"
-#include "ActorDescTreeItem.h"
 #include "WorldTreeItem.h"
 #include "LevelInstance/LevelInstanceEditorInstanceActor.h"
 
@@ -334,11 +333,6 @@ int32 FActorMode::GetTypeSortPriority(const ISceneOutlinerTreeItem& Item) const
 	{
 		return EItemSortOrder::Actor;
 	}
-	else if (Item.IsA<FActorDescTreeItem>())
-	{
-		return EItemSortOrder::Unloaded;
-	}
-
 	// Warning: using actor mode with an unsupported item type!
 	check(false);
 	return -1;
diff --git a/Engine/Source/Editor/SceneOutliner/Private/SceneOutlinerActorInfoColumn.cpp b/Engine/Source/Editor/SceneOutliner/Private/SceneOutlinerActorInfoColumn.cpp
--- a/Engine/Source/Editor/SceneOutliner/Private/SceneOutlinerActorInfoColumn.cpp
+++ b/Engine/Source/Editor/SceneOutliner/Private/SceneOutlinerActorInfoColumn.cpp
@@ -14,13 +14,11 @@
 #include "SortHelper.h"
 #include "ISceneOutliner.h"
 #include "ActorTreeItem.h"
-#include "ActorDescTreeItem.h"
 #include "ComponentTreeItem.h"
 #include "FolderTreeItem.h"
 #include "WorldTreeItem.h"
 #include "WorldPartition/DataLayer/DataLayer.h"
 #include "WorldPartition/DataLayer/WorldDataLayers.h"
-#include "WorldPartition/WorldPartitionActorDesc.h"
 
 #define LOCTEXT_NAMESPACE "SceneOutlinerActorInfoColumn"
 
@@ -154,24 +152,6 @@ struct FGetInfo
 				return FString();
 			}
 		}
-		else if (const FActorDescTreeItem* ActorDescItem = Item.CastTo<FActorDescTreeItem>())
-		{
-			const FWorldPartitionActorDesc* ActorDesc = ActorDescItem->ActorDesc;
-
-			if (!ActorDesc)
-			{
-				return FString();
-			}
-			
-			switch (CurrentMode)
-			{
-			case SceneOutliner::ECustomColumnMode::Class:
-				return ActorDesc->GetActorClass()->GetName();
-			default:
-				return FString();
-			}
-		}
-
 
 		return FString();
 	}
diff --git a/Engine/Source/Editor/SceneOutliner/Public/ActorBrowsingMode.h b/Engine/Source/Editor/SceneOutliner/Public/ActorBrowsingMode.h
--- a/Engine/Source/Editor/SceneOutliner/Public/ActorBrowsingMode.h
+++ b/Engine/Source/Editor/SceneOutliner/Public/ActorBrowsingMode.h
@@ -99,7 +99,6 @@ private:
 	static TSharedRef<FSceneOutlinerFilter> CreateIsInCurrentLevelFilter();
 	static TSharedRef<FSceneOutlinerFilter> CreateHideComponentsFilter();
 	static TSharedRef<FSceneOutlinerFilter> CreateHideLevelInstancesFilter();
-	static TSharedRef<FSceneOutlinerFilter> CreateHideUnloadedActorsFilter();
 	
 	void ToggleActorSCCStatusColumn();
 	bool isActorSCCStatusColumnActive() const { return bActorSCCStatusColumnActive; }
diff --git a/Engine/Source/Editor/SceneOutliner/Public/ActorDescTreeItem.h b/Engine/Source/Editor/SceneOutliner/Public/ActorDescTreeItem.h
deleted file mode 100644
diff --git a/Engine/Source/Editor/SceneOutliner/Public/ActorMode.h b/Engine/Source/Editor/SceneOutliner/Public/ActorMode.h
--- a/Engine/Source/Editor/SceneOutliner/Public/ActorMode.h
+++ b/Engine/Source/Editor/SceneOutliner/Public/ActorMode.h
@@ -41,7 +41,7 @@ class SCENEOUTLINER_API FActorMode : public ISceneOutlinerMode
 public:
 	struct EItemSortOrder
 	{
-		enum Type { World = 0, Folder = 10, Actor = 20, Unloaded = 30 };
+		enum Type { World = 0, Folder = 10, Actor = 20 };
 	};
 	
 	FActorMode(const FActorModeParams& Params);
diff --git a/Engine/Source/Editor/SceneOutliner/Public/SceneOutlinerStandaloneTypes.h b/Engine/Source/Editor/SceneOutliner/Public/SceneOutlinerStandaloneTypes.h
--- a/Engine/Source/Editor/SceneOutliner/Public/SceneOutlinerStandaloneTypes.h
+++ b/Engine/Source/Editor/SceneOutliner/Public/SceneOutlinerStandaloneTypes.h
@@ -14,7 +14,7 @@ struct FSceneOutlinerTreeItemID
 public:
 	using FUniqueID = uint32;
 
-	enum class EType : uint8 { Object, Folder, UniqueID, Guid, Null };
+	enum class EType : uint8 { Object, Folder, UniqueID, Null };
 
 	/** Default constructed null item ID */
 	FSceneOutlinerTreeItemID() : Type(EType::Null), CachedHash(0) {}
@@ -46,12 +46,6 @@ public:
 		CachedHash = CalculateTypeHash();
 	}
 
-	FSceneOutlinerTreeItemID(const FGuid& InGuid) : Type(EType::Guid)
-	{
-		new (Data) FGuid(InGuid);
-		CachedHash = CalculateTypeHash();
-	}
-
 	/** Copy construction / assignment */
 	FSceneOutlinerTreeItemID(const FSceneOutlinerTreeItemID& Other)
 	{
@@ -64,8 +58,7 @@ public:
 		{
 			case EType::Object:			new (Data) FObjectKey(Other.GetAsObjectKey());		break;
 			case EType::Folder:			new (Data) FName(Other.GetAsFolderRef());			break;
-			case EType::Guid:			new (Data) FGuid(Other.GetAsGuid());				break;
-			case EType::UniqueID:		new (Data) FUniqueID(Other.GetAsHash());			break;
+			case EType::UniqueID:		new (Data) FUniqueID(Other.GetAsHash());	break;
 			default:																		break;
 		}
 
@@ -90,7 +83,6 @@ public:
 		{
 			case EType::Object:			GetAsObjectKey().~FObjectKey();							break;
 			case EType::Folder:			GetAsFolderRef().~FName();								break;
-			case EType::Guid:			GetAsGuid().~FGuid();									break;
 			case EType::UniqueID:		/* NOP */												break;
 			default:																			break;
 		}
@@ -112,8 +104,7 @@ public:
 		{
 			case EType::Object:			Hash = GetTypeHash(GetAsObjectKey());				break;
 			case EType::Folder:			Hash = GetTypeHash(GetAsFolderRef());				break;
-			case EType::Guid:			Hash = GetTypeHash(GetAsGuid());					break;
-			case EType::UniqueID:		Hash = GetTypeHash(GetAsHash());					break;
+			case EType::UniqueID:		Hash = GetAsHash();									break;
 			default:																		break;
 		}
 
@@ -129,7 +120,6 @@ private:
 
 	FObjectKey& 	GetAsObjectKey() const 		{ return *reinterpret_cast<FObjectKey*>(Data); }
 	FName& 			GetAsFolderRef() const		{ return *reinterpret_cast<FName*>(Data); }
-	FGuid&			GetAsGuid() const			{ return *reinterpret_cast<FGuid*>(Data); }
 	FUniqueID&		GetAsHash() const			{ return *reinterpret_cast<FUniqueID*>(Data); }
 
 	/** Compares the specified ID with this one - assumes matching types */
@@ -139,7 +129,6 @@ private:
 		{
 			case EType::Object:			return GetAsObjectKey() == Other.GetAsObjectKey();
 			case EType::Folder:			return GetAsFolderRef() == Other.GetAsFolderRef();
-			case EType::Guid:			return GetAsGuid() == Other.GetAsGuid();
 			case EType::UniqueID:		return GetAsHash() == Other.GetAsHash();
 			case EType::Null:			return true;
 			default: check(false);		return false;
@@ -149,7 +138,7 @@ private:
 	EType Type;
 
 	uint32 CachedHash;
-	static const uint32 MaxSize = TMaxSizeof<FObjectKey, FName, FGuid, FUniqueID>::Value;
+	static const uint32 MaxSize = TMaxSizeof<FObjectKey, FName, FUniqueID>::Value;
 	mutable uint8 Data[MaxSize];
 };
 	
