SceneOutliner: GetSceneOutlinerParent discarded if it's hidden in Outliner

diff --git a/Engine/Source/Editor/SceneOutliner/Private/ActorHierarchy.cpp b/Engine/Source/Editor/SceneOutliner/Private/ActorHierarchy.cpp
--- a/Engine/Source/Editor/SceneOutliner/Private/ActorHierarchy.cpp
+++ b/Engine/Source/Editor/SceneOutliner/Private/ActorHierarchy.cpp
@@ -85,8 +85,8 @@ FSceneOutlinerTreeItemPtr FActorHierarchy::FindParent(const ISceneOutlinerTreeIt
 				if (const FSceneOutlinerTreeItemPtr* ParentItem = Items.Find(ParentActor))
 				{
 					return *ParentItem;
-				}
-				else
+				} // IF Parent can be listed in SceneOutliner return nullptr so it gets created
+				else if(ParentActor->IsListedInSceneOutliner())
 				{
 					return nullptr;
 				}
