SceneOutliner: Fixed issue where PIE with netmode Client would show an empty hierarchy

diff --git a/Engine/Source/Editor/SceneOutliner/Private/ActorModeInteractive.cpp b/Engine/Source/Editor/SceneOutliner/Private/ActorModeInteractive.cpp
--- a/Engine/Source/Editor/SceneOutliner/Private/ActorModeInteractive.cpp
+++ b/Engine/Source/Editor/SceneOutliner/Private/ActorModeInteractive.cpp
@@ -15,6 +15,7 @@ FActorModeInteractive::FActorModeInteractive(const FActorModeParams& Params)
 	FEditorDelegates::NewCurrentLevel.AddRaw(this, &FActorModeInteractive::OnNewCurrentLevel);
 
 	FCoreDelegates::OnActorLabelChanged.AddRaw(this, &FActorModeInteractive::OnActorLabelChanged);
+	FCoreUObjectDelegates::PostLoadMapWithWorld.AddRaw(this, &FActorModeInteractive::OnPostLoadMapWithWorld);
 	GEngine->OnLevelActorRequestRename().AddRaw(this, &FActorModeInteractive::OnLevelActorRequestsRename);
 }
 
@@ -27,6 +28,7 @@ FActorModeInteractive::~FActorModeInteractive()
 	FEditorDelegates::NewCurrentLevel.RemoveAll(this);
 
 	FCoreDelegates::OnActorLabelChanged.RemoveAll(this);
+	FCoreUObjectDelegates::PostLoadMapWithWorld.RemoveAll(this);
 	GEngine->OnLevelActorRequestRename().RemoveAll(this);
 }
 
@@ -87,6 +89,11 @@ void FActorModeInteractive::OnLevelActorRequestsRename(const AActor* Actor)
 	}
 }
 
+void FActorModeInteractive::OnPostLoadMapWithWorld(UWorld* World)
+{
+	SceneOutliner->FullRefresh();
+}
+
 void FActorModeInteractive::OnActorLabelChanged(AActor* ChangedActor)
 {
 	if (!ensure(ChangedActor))
diff --git a/Engine/Source/Editor/SceneOutliner/Public/ActorModeInteractive.h b/Engine/Source/Editor/SceneOutliner/Public/ActorModeInteractive.h
--- a/Engine/Source/Editor/SceneOutliner/Public/ActorModeInteractive.h
+++ b/Engine/Source/Editor/SceneOutliner/Public/ActorModeInteractive.h
@@ -20,4 +20,5 @@ private:
 	void OnLevelSelectionChanged(UObject* Obj);
 	void OnActorLabelChanged(AActor* ChangedActor);
 	void OnLevelActorRequestsRename(const AActor* Actor);
+	void OnPostLoadMapWithWorld(UWorld* World);
 };
