SceneOutliner: Show unloaded actor source control status in outliner

diff --git a/Engine/Source/Editor/SceneOutliner/Private/SceneOutlinerActorSCCColumn.cpp b/Engine/Source/Editor/SceneOutliner/Private/SceneOutlinerActorSCCColumn.cpp
--- a/Engine/Source/Editor/SceneOutliner/Private/SceneOutlinerActorSCCColumn.cpp
+++ b/Engine/Source/Editor/SceneOutliner/Private/SceneOutlinerActorSCCColumn.cpp
@@ -4,6 +4,8 @@
 #include "Widgets/Images/SImage.h"
 #include "Widgets/Views/STreeView.h"
 #include "ActorTreeItem.h"
+#include "ActorDescTreeItem.h"
+#include "WorldPartition/WorldPartitionActorDesc.h"
 #include "ISourceControlProvider.h"
 #include "ISourceControlModule.h"
 #include "SourceControlHelpers.h"
@@ -38,7 +40,6 @@ public:
 					if (Actor->IsPackageExternal())
 					{
 						ExternalPackageName = USourceControlHelpers::PackageFilename(Actor->GetExternalPackage());
-						ConnectSourceControl();
 					}
 
 					ActorPackingModeChangedDelegateHandle = Actor->OnPackagingModeChanged.AddLambda([this](AActor* InActor, bool bExternal)
@@ -56,6 +57,18 @@ public:
 						});
 				}
 			}
+			else if (FActorDescTreeItem* ActorDescItem = TreeItemPtr->CastTo<FActorDescTreeItem>())
+			{
+				if (const FWorldPartitionActorDesc* ActorDesc = ActorDescItem->ActorDescHandle.GetActorDesc())
+				{
+					ExternalPackageName =  USourceControlHelpers::PackageFilename(ActorDesc->GetActorPackage().ToString());
+				}
+			}
+		}
+
+		if (!ExternalPackageName.IsEmpty())
+		{
+			ConnectSourceControl();
 		}
 	}
 
@@ -87,7 +100,7 @@ private:
 
 		// Check if there is already a cached state for this item
 		FSourceControlStatePtr SourceControlState = ISourceControlModule::Get().GetProvider().GetState(ExternalPackageName, EStateCacheUsage::Use);
-		if (SourceControlState.IsValid())
+		if (SourceControlState.IsValid() && !SourceControlState->IsUnknown())
 		{
 			UpdateSourceControlStateIcon(SourceControlState);
 		}
@@ -189,7 +202,7 @@ SHeaderRow::FColumn::FArguments FSceneOutlinerActorSCCColumn::ConstructHeaderRow
 
 const TSharedRef<SWidget> FSceneOutlinerActorSCCColumn::ConstructRowWidget(FSceneOutlinerTreeItemRef TreeItem, const STableRow<FSceneOutlinerTreeItemPtr>& Row)
 {
-	if (TreeItem->IsA<FActorTreeItem>())
+	if (TreeItem->IsA<FActorTreeItem>() || TreeItem->IsA<FActorDescTreeItem>())
 	{
 		return SNew(SHorizontalBox)
 			+ SHorizontalBox::Slot()
