SceneOutliner: Hide unloaded actors in non-World Outliner contexts

diff --git a/Engine/Source/Editor/SceneOutliner/Private/ActorBrowsingMode.cpp b/Engine/Source/Editor/SceneOutliner/Private/ActorBrowsingMode.cpp
--- a/Engine/Source/Editor/SceneOutliner/Private/ActorBrowsingMode.cpp
+++ b/Engine/Source/Editor/SceneOutliner/Private/ActorBrowsingMode.cpp
@@ -46,7 +46,7 @@ using FActorFilter = TSceneOutlinerPredicateFilter<FActorTreeItem>;
 using FActorDescFilter = TSceneOutlinerPredicateFilter<FActorDescTreeItem>;
 
 FActorBrowsingMode::FActorBrowsingMode(SSceneOutliner* InSceneOutliner, TWeakObjectPtr<UWorld> InSpecifiedWorldToDisplay)
-	: FActorModeInteractive(FActorModeParams(InSceneOutliner, InSpecifiedWorldToDisplay, true, false))
+	: FActorModeInteractive(FActorModeParams(InSceneOutliner, InSpecifiedWorldToDisplay,  /* bHideComponents */ true, /* bHideLevelInstanceHierarchy */ false, /* bHideUnloadedActors */ false))
 	, FilteredActorCount(0)
 {
 	// Capture selection changes of bones from mesh selection in fracture tools
@@ -128,8 +128,13 @@ FActorBrowsingMode::FActorBrowsingMode(SSceneOutliner* InSceneOutliner, TWeakObj
 	HideUnloadedActorsInfo.OnToggle().AddLambda([this] (bool bIsActive)
 		{
 			UActorBrowsingModeSettings* Settings = GetMutableDefault<UActorBrowsingModeSettings>();
-	        Settings->bHideLevelInstanceHierarchy = bIsActive;
-	        Settings->PostEditChange();
+			Settings->bHideUnloadedActors = bHideUnloadedActors = bIsActive;
+			Settings->PostEditChange();
+		
+			if (auto ActorHierarchy = StaticCast<FActorHierarchy*>(Hierarchy.Get()))
+			{
+				ActorHierarchy->SetShowingUnloadedActors(!bIsActive);
+			}
 		});
 	FilterInfoMap.Add(TEXT("HideUnloadedActorsFilter"), HideUnloadedActorsInfo);
 
diff --git a/Engine/Source/Editor/SceneOutliner/Private/ActorHierarchy.cpp b/Engine/Source/Editor/SceneOutliner/Private/ActorHierarchy.cpp
--- a/Engine/Source/Editor/SceneOutliner/Private/ActorHierarchy.cpp
+++ b/Engine/Source/Editor/SceneOutliner/Private/ActorHierarchy.cpp
@@ -254,21 +254,24 @@ void FActorHierarchy::CreateWorldChildren(UWorld* World, TArray<FSceneOutlinerTr
 		}
 	}
 
-	if (UWorldPartitionSubsystem* WorldPartitionSubsystem = World->GetSubsystem<UWorldPartitionSubsystem>())
+	if (bShowingUnloadedActors)
 	{
-		UWorldPartition* WorldPartition = World->GetWorldPartition();
-		check(WorldPartition);
-		WorldPartitionSubsystem->ForEachActorDesc(AActor::StaticClass(), [this, WorldPartition, &OutItems](const FWorldPartitionActorDesc* ActorDesc)
-			{
-				if (ActorDesc != nullptr && !ActorDesc->IsLoaded())
-				{
-					if (const FSceneOutlinerTreeItemPtr ActorDescItem = Mode->CreateItemFor<FActorDescTreeItem>(FActorDescTreeItem(ActorDesc, WorldPartition)))
-					{
-						OutItems.Add(ActorDescItem);
-					}
-				}
-				return true;
-			});
+		if (UWorldPartitionSubsystem* WorldPartitionSubsystem = World->GetSubsystem<UWorldPartitionSubsystem>())
+		{
+			UWorldPartition* WorldPartition = World->GetWorldPartition();
+			check(WorldPartition);
+			WorldPartitionSubsystem->ForEachActorDesc(AActor::StaticClass(), [this, WorldPartition, &OutItems](const FWorldPartitionActorDesc* ActorDesc)
+                {
+                    if (ActorDesc != nullptr && !ActorDesc->IsLoaded())
+                    {
+                        if (const FSceneOutlinerTreeItemPtr ActorDescItem = Mode->CreateItemFor<FActorDescTreeItem>(FActorDescTreeItem(ActorDesc, WorldPartition)))
+                        {
+                            OutItems.Add(ActorDescItem);
+                        }
+                    }
+                    return true;
+                });
+		}
 	}
 }
 
diff --git a/Engine/Source/Editor/SceneOutliner/Private/ActorMode.cpp b/Engine/Source/Editor/SceneOutliner/Private/ActorMode.cpp
--- a/Engine/Source/Editor/SceneOutliner/Private/ActorMode.cpp
+++ b/Engine/Source/Editor/SceneOutliner/Private/ActorMode.cpp
@@ -76,9 +76,10 @@ namespace SceneOutliner
 
 FActorMode::FActorMode(const FActorModeParams& Params)
 	: ISceneOutlinerMode(Params.SceneOutliner)
+	, SpecifiedWorldToDisplay(Params.SpecifiedWorldToDisplay)
 	, bHideComponents(Params.bHideComponents)
 	, bHideLevelInstanceHierarchy(Params.bHideLevelInstanceHierarchy)
-	, SpecifiedWorldToDisplay(Params.SpecifiedWorldToDisplay)
+	, bHideUnloadedActors(Params.bHideUnloadedActors)
 {
 	SceneOutliner->AddFilter(MakeShared<FActorFilter>(FActorTreeItem::FFilterPredicate::CreateLambda([this](const AActor* Actor)
 		{
@@ -95,6 +96,7 @@ TUniquePtr<ISceneOutlinerHierarchy> FActorMode::CreateHierarchy()
 	TUniquePtr<FActorHierarchy> ActorHierarchy = FActorHierarchy::Create(this, RepresentingWorld);
 	ActorHierarchy->SetShowingComponents(!bHideComponents);
 	ActorHierarchy->SetShowingLevelInstances(!bHideLevelInstanceHierarchy);
+	ActorHierarchy->SetShowingUnloadedActors(!bHideUnloadedActors);
 
 	return ActorHierarchy;
 }
diff --git a/Engine/Source/Editor/SceneOutliner/Private/SceneOutlinerModule.cpp b/Engine/Source/Editor/SceneOutliner/Private/SceneOutlinerModule.cpp
--- a/Engine/Source/Editor/SceneOutliner/Private/SceneOutlinerModule.cpp
+++ b/Engine/Source/Editor/SceneOutliner/Private/SceneOutlinerModule.cpp
@@ -65,6 +65,7 @@ TSharedRef<ISceneOutliner> FSceneOutlinerModule::CreateActorPicker(const FSceneO
 			Params.SpecifiedWorldToDisplay = SpecifiedWorld;
 			Params.bHideComponents = true;
 			Params.bHideLevelInstanceHierarchy = true;
+			Params.bHideUnloadedActors = true;
 			return new FActorPickingMode(Params, OnItemPicked);
 		});
 
@@ -100,6 +101,7 @@ TSharedRef<ISceneOutliner> FSceneOutlinerModule::CreateComponentPicker(const FSc
 			Params.SpecifiedWorldToDisplay = SpecifiedWorld;
 			Params.bHideComponents = false;
 			Params.bHideLevelInstanceHierarchy = true;
+			Params.bHideUnloadedActors = true;
 			return new FActorPickingMode(Params, OnItemPicked);
 		});
 
diff --git a/Engine/Source/Editor/SceneOutliner/Public/ActorHierarchy.h b/Engine/Source/Editor/SceneOutliner/Public/ActorHierarchy.h
--- a/Engine/Source/Editor/SceneOutliner/Public/ActorHierarchy.h
+++ b/Engine/Source/Editor/SceneOutliner/Public/ActorHierarchy.h
@@ -24,6 +24,7 @@ public:
 
 	void SetShowingComponents(bool bInShowingComponents) { bShowingComponents = bInShowingComponents; }
 	void SetShowingLevelInstances(bool bInShowingLevelInstances) { bShowingLevelInstances = bInShowingLevelInstances; }
+	void SetShowingUnloadedActors(bool bInShowingUnloadedActors) { bShowingUnloadedActors = bInShowingUnloadedActors; }
 private:
 	/** Adds all the direct and indirect children of a world to OutItems */
 	void CreateWorldChildren(UWorld* World, TArray<FSceneOutlinerTreeItemPtr>& OutItems) const;
@@ -60,10 +61,11 @@ private:
 	/** Send a an event indicating a full refresh of the hierarchy is required */
 	void FullRefreshEvent();
 
+	TWeakObjectPtr<UWorld> RepresentingWorld;
+	
 	bool bShowingComponents = false;
 	bool bShowingLevelInstances = false;
-
-	TWeakObjectPtr<UWorld> RepresentingWorld;
+	bool bShowingUnloadedActors = false;
 private:
 	FActorHierarchy(ISceneOutlinerMode* Mode, const TWeakObjectPtr<UWorld>& Worlds);
 
diff --git a/Engine/Source/Editor/SceneOutliner/Public/ActorMode.h b/Engine/Source/Editor/SceneOutliner/Public/ActorMode.h
--- a/Engine/Source/Editor/SceneOutliner/Public/ActorMode.h
+++ b/Engine/Source/Editor/SceneOutliner/Public/ActorMode.h
@@ -23,17 +23,19 @@ struct SCENEOUTLINER_API FActorModeParams
 {
 	FActorModeParams() {}
 
-	FActorModeParams(SSceneOutliner* InSceneOutliner, const TWeakObjectPtr<UWorld>& InSpecifiedWorldToDisplay = nullptr, bool bInHideComponents = true, bool bInHideLevelInstanceHierarchy = true)
+	FActorModeParams(SSceneOutliner* InSceneOutliner, const TWeakObjectPtr<UWorld>& InSpecifiedWorldToDisplay = nullptr, bool bInHideComponents = true, bool bInHideLevelInstanceHierarchy = true, bool bInHideUnloadedActors = true)
 		: SpecifiedWorldToDisplay(InSpecifiedWorldToDisplay)
 		, SceneOutliner(InSceneOutliner)
 		, bHideComponents(bInHideComponents)
 		, bHideLevelInstanceHierarchy(bInHideLevelInstanceHierarchy)
+		, bHideUnloadedActors(bInHideUnloadedActors)
 	{}
 
 	TWeakObjectPtr<UWorld> SpecifiedWorldToDisplay = nullptr;
-	SSceneOutliner* SceneOutliner;
-	bool bHideComponents;
-	bool bHideLevelInstanceHierarchy;
+	SSceneOutliner* SceneOutliner = nullptr;
+	bool bHideComponents = true;
+	bool bHideLevelInstanceHierarchy = true;
+	bool bHideUnloadedActors = true;
 };
 
 class SCENEOUTLINER_API FActorMode : public ISceneOutlinerMode
@@ -72,10 +74,6 @@ protected:
 
 	virtual TUniquePtr<ISceneOutlinerHierarchy> CreateHierarchy() override;
 protected:
-	// Should the hide components filter be enabled
-	bool bHideComponents;
-	// Are LevelInstances being hidden
-	bool bHideLevelInstanceHierarchy;
 	/** The world which we are currently representing */
 	TWeakObjectPtr<UWorld> RepresentingWorld;
 	/** The world which the user manually selected */
@@ -83,4 +81,11 @@ protected:
 
 	/** If this mode was created to display a specific world, don't allow it to be reassigned */
 	const TWeakObjectPtr<UWorld> SpecifiedWorldToDisplay;
+
+	/** Should components be hidden */
+	bool bHideComponents;
+	/** Should the level instance hierarchy be hidden */
+	bool bHideLevelInstanceHierarchy;
+	/** Should unloaded actors be hidden */
+	bool bHideUnloadedActors;
 };
