SceneOutliner: deleted actors now display a cached label instead of the generic "(Deleted Actor)" label.

diff --git a/Engine/Source/Editor/SceneOutliner/Private/ActorTreeItem.cpp b/Engine/Source/Editor/SceneOutliner/Private/ActorTreeItem.cpp
--- a/Engine/Source/Editor/SceneOutliner/Private/ActorTreeItem.cpp
+++ b/Engine/Source/Editor/SceneOutliner/Private/ActorTreeItem.cpp
@@ -111,22 +111,26 @@ private:
 	TWeakPtr<FActorTreeItem> TreeItemPtr;
 	TWeakObjectPtr<AActor> ActorPtr;
 	TAttribute<FText> HighlightText;
-
+	
 	FText GetDisplayText() const
 	{
-		const AActor* Actor = ActorPtr.Get();
-		if (const ALevelInstance* LevelInstanceActor = Cast<ALevelInstance>(Actor))
+		if (const FSceneOutlinerTreeItemPtr TreeItem = TreeItemPtr.Pin())
 		{
-			if (LevelInstanceActor->IsDirty() && !bInEditingMode)
+			const AActor* Actor = ActorPtr.Get();
+			if (const ALevelInstance* LevelInstanceActor = Cast<ALevelInstance>(Actor))
 			{
-				FFormatNamedArguments Args;
-				Args.Add(TEXT("ActorLabel"), FText::FromString(LevelInstanceActor->GetActorLabel()));
-				Args.Add(TEXT("EditTag"), LOCTEXT("EditingLevelInstanceLabel", "*"));
-				return FText::Format(LOCTEXT("LevelInstanceDisplay", "{ActorLabel}{EditTag}"), Args);
+				if (LevelInstanceActor->IsDirty() && !bInEditingMode)
+				{
+					FFormatNamedArguments Args;
+					Args.Add(TEXT("ActorLabel"), FText::FromString(TreeItem->GetDisplayString()));
+					Args.Add(TEXT("EditTag"), LOCTEXT("EditingLevelInstanceLabel", "*"));
+					return FText::Format(LOCTEXT("LevelInstanceDisplay", "{ActorLabel}{EditTag}"), Args);
+				}
 			}
+			return FText::FromString(TreeItem->GetDisplayString());
 		}
 
-		return Actor ? FText::FromString(Actor->GetActorLabel()) : LOCTEXT("ActorLabelForMissingActor", "(Deleted Actor)");
+		return FText();
 	}
 
 	FText GetTooltipText() const
@@ -189,7 +193,7 @@ private:
 		}
 		else
 		{
-			return nullptr;
+			return FSlateIconFinder::FindIconForClass(AActor::StaticClass()).GetOptionalIcon();
 		}
 	}
 
@@ -340,6 +344,9 @@ FActorTreeItem::FActorTreeItem(AActor* InActor)
 	, Actor(InActor)
 	, ID(InActor)
 {
+	check(InActor);
+	ActorLabel = InActor->GetActorLabel();
+	
 	bExistsInCurrentWorldAndPIE = GEditor->ObjectsThatExistInEditorWorld.Get(InActor);
 }
 
@@ -350,8 +357,7 @@ FSceneOutlinerTreeItemID FActorTreeItem::GetID() const
 
 FString FActorTreeItem::GetDisplayString() const
 {
-	const AActor* ActorPtr = Actor.Get();
-	return ActorPtr ? ActorPtr->GetActorLabel() : LOCTEXT("ActorLabelForMissingActor", "(Deleted Actor)").ToString();
+	return ActorLabel;
 }
 
 bool FActorTreeItem::CanInteract() const
@@ -388,7 +394,16 @@ void FActorTreeItem::OnVisibilityChanged(const bool bNewVisibility)
 
 bool FActorTreeItem::GetVisibility() const
 {
-	return Actor.IsValid() && !Actor->IsTemporarilyHiddenInEditor(true);
+	// We want deleted actors to appear as if they are visible to minimize visual clutter.
+	return !Actor.IsValid() || !Actor->IsTemporarilyHiddenInEditor(true);
+}
+
+void FActorTreeItem::OnLabelChanged()
+{
+	if (Actor.IsValid())
+	{
+		ActorLabel = Actor->GetActorLabel();
+	}
 }
 
 #undef LOCTEXT_NAMESPACE
diff --git a/Engine/Source/Editor/SceneOutliner/Private/SSceneOutliner.cpp b/Engine/Source/Editor/SceneOutliner/Private/SSceneOutliner.cpp
--- a/Engine/Source/Editor/SceneOutliner/Private/SSceneOutliner.cpp
+++ b/Engine/Source/Editor/SceneOutliner/Private/SSceneOutliner.cpp
@@ -1654,6 +1654,8 @@ void SSceneOutliner::OnItemLabelChanged(FSceneOutlinerTreeItemPtr ChangedItem)
 	// If the item already exists
 	if (FSceneOutlinerTreeItemPtr* ExistingItem = TreeItemMap.Find(ChangedItem->GetID()))
 	{
+		(*ExistingItem)->OnLabelChanged();
+		
 		// The changed item flags will have been set already
 		if (!ChangedItem->Flags.bIsFilteredOut)
 		{
diff --git a/Engine/Source/Editor/SceneOutliner/Public/ActorTreeItem.h b/Engine/Source/Editor/SceneOutliner/Public/ActorTreeItem.h
--- a/Engine/Source/Editor/SceneOutliner/Public/ActorTreeItem.h
+++ b/Engine/Source/Editor/SceneOutliner/Public/ActorTreeItem.h
@@ -45,8 +45,11 @@ public:
 	virtual void OnVisibilityChanged(const bool bNewVisibility) override;
 	virtual bool HasVisibilityInfo() const override { return true; }
 	virtual bool GetVisibility() const override;
+	virtual void OnLabelChanged() override;
 	/* End ISceneOutlinerTreeItem Implementation */
 public:
 	/** true if this item exists in both the current world and PIE. */
 	bool bExistsInCurrentWorldAndPIE;
+	/** Cached actor label */
+	FString ActorLabel;
 };
diff --git a/Engine/Source/Editor/SceneOutliner/Public/ISceneOutlinerTreeItem.h b/Engine/Source/Editor/SceneOutliner/Public/ISceneOutlinerTreeItem.h
--- a/Engine/Source/Editor/SceneOutliner/Public/ISceneOutlinerTreeItem.h
+++ b/Engine/Source/Editor/SceneOutliner/Public/ISceneOutlinerTreeItem.h
@@ -147,4 +147,6 @@ public:
 	/** Query this items visibility state. Only called if the item type has visibility info */
 	virtual bool GetVisibility() const { return false; }
 
+	/** Called when this item's label has changed */
+	virtual void OnLabelChanged() {}
 };
