SceneOutliner: Fixed issue where item labels are incorrectly set for new actors

diff --git a/Engine/Source/Editor/SceneOutliner/Private/SSceneOutliner.cpp b/Engine/Source/Editor/SceneOutliner/Private/SSceneOutliner.cpp
--- a/Engine/Source/Editor/SceneOutliner/Private/SSceneOutliner.cpp
+++ b/Engine/Source/Editor/SceneOutliner/Private/SSceneOutliner.cpp
@@ -1669,6 +1669,10 @@ void SSceneOutliner::OnItemLabelChanged(FSceneOutlinerTreeItemPtr ChangedItem)
 			Refresh();
 		}
 	}
+	else if (FSceneOutlinerTreeItemPtr* PendingItem = PendingTreeItemMap.Find(ChangedItem->GetID()))
+	{
+		(*PendingItem)->OnLabelChanged();
+	}
 	else
 	{
 		// Attempt to add the item if we didn't find it - perhaps it now matches the filter?
