SceneOutliner: Context menu for unloaded actors now contains a "focus actor bounds"

diff --git a/Engine/Source/Editor/SceneOutliner/Private/ActorDescTreeItem.cpp b/Engine/Source/Editor/SceneOutliner/Private/ActorDescTreeItem.cpp
--- a/Engine/Source/Editor/SceneOutliner/Private/ActorDescTreeItem.cpp
+++ b/Engine/Source/Editor/SceneOutliner/Private/ActorDescTreeItem.cpp
@@ -18,6 +18,8 @@
 #include "ISceneOutlinerMode.h"
 #include "Logging/MessageLog.h"
 #include "SSocketChooser.h"
+#include "ToolMenus.h"
+#include "LevelEditorViewport.h"
 
 #define LOCTEXT_NAMESPACE "SceneOutliner_ActorDescTreeItem"
 
@@ -28,7 +30,7 @@ struct SActorDescTreeLabel : FSceneOutlinerCommonLabelData, public SCompoundWidg
 	SLATE_BEGIN_ARGS(SActorDescTreeLabel) {}
 	SLATE_END_ARGS()
 
-		void Construct(const FArguments& InArgs, FActorDescTreeItem& ActorDescItem, ISceneOutliner& SceneOutliner, const STableRow<FSceneOutlinerTreeItemPtr>& InRow)
+	void Construct(const FArguments& InArgs, FActorDescTreeItem& ActorDescItem, ISceneOutliner& SceneOutliner, const STableRow<FSceneOutlinerTreeItemPtr>& InRow)
 	{
 		WeakSceneOutliner = StaticCastSharedRef<ISceneOutliner>(SceneOutliner.AsShared());
 
@@ -45,16 +47,16 @@ struct SActorDescTreeLabel : FSceneOutlinerCommonLabelData, public SCompoundWidg
 			[
 				SAssignNew(InlineTextBlock, SInlineEditableTextBlock)
 				.Text(this, &SActorDescTreeLabel::GetDisplayText)
-			.ToolTipText(this, &SActorDescTreeLabel::GetTooltipText)
-			.HighlightText(HighlightText)
-			.ColorAndOpacity(this, &SActorDescTreeLabel::GetForegroundColor)
-			.OnTextCommitted(this, &SActorDescTreeLabel::OnLabelCommitted)
-			.OnVerifyTextChanged(this, &SActorDescTreeLabel::OnVerifyItemLabelChanged)
-			.IsSelected(FIsSelected::CreateSP(&InRow, &STableRow<FSceneOutlinerTreeItemPtr>::IsSelectedExclusively))
-			.IsReadOnly_Lambda([Item = ActorDescItem.AsShared(), this]()
-		{
-			return !CanExecuteRenameRequest(Item.Get());
-		})
+				.ToolTipText(this, &SActorDescTreeLabel::GetTooltipText)
+				.HighlightText(HighlightText)
+				.ColorAndOpacity(this, &SActorDescTreeLabel::GetForegroundColor)
+				.OnTextCommitted(this, &SActorDescTreeLabel::OnLabelCommitted)
+				.OnVerifyTextChanged(this, &SActorDescTreeLabel::OnVerifyItemLabelChanged)
+				.IsSelected(FIsSelected::CreateSP(&InRow, &STableRow<FSceneOutlinerTreeItemPtr>::IsSelectedExclusively))
+				.IsReadOnly_Lambda([Item = ActorDescItem.AsShared(), this]()
+			{
+				return !CanExecuteRenameRequest(Item.Get());
+			})
 			]
 
 		+ SHorizontalBox::Slot()
@@ -64,8 +66,8 @@ struct SActorDescTreeLabel : FSceneOutlinerCommonLabelData, public SCompoundWidg
 			[
 				SNew(STextBlock)
 				.Text(this, &SActorDescTreeLabel::GetTypeText)
-			.Visibility(this, &SActorDescTreeLabel::GetTypeTextVisibility)
-			.HighlightText(HighlightText)
+				.Visibility(this, &SActorDescTreeLabel::GetTypeTextVisibility)
+				.HighlightText(HighlightText)
 			];
 
 		if (WeakSceneOutliner.Pin()->GetMode()->IsInteractive())
@@ -78,28 +80,28 @@ struct SActorDescTreeLabel : FSceneOutlinerCommonLabelData, public SCompoundWidg
 				SNew(SHorizontalBox)
 
 				+ SHorizontalBox::Slot()
-			.AutoWidth()
-			.VAlign(VAlign_Center)
-			.Padding(FSceneOutlinerDefaultTreeItemMetrics::IconPadding())
-			[
-				SNew(SBox)
-				.WidthOverride(FSceneOutlinerDefaultTreeItemMetrics::IconSize())
-			.HeightOverride(FSceneOutlinerDefaultTreeItemMetrics::IconSize())
-			[
-				SNew(SImage)
-				.Image(this, &SActorDescTreeLabel::GetIcon)
-			.ToolTipText(this, &SActorDescTreeLabel::GetIconTooltip)
-			.ColorAndOpacity(FSlateColor::UseForeground())
-			]
-			]
+				.AutoWidth()
+				.VAlign(VAlign_Center)
+				.Padding(FSceneOutlinerDefaultTreeItemMetrics::IconPadding())
+				[
+					SNew(SBox)
+					.WidthOverride(FSceneOutlinerDefaultTreeItemMetrics::IconSize())
+					.HeightOverride(FSceneOutlinerDefaultTreeItemMetrics::IconSize())
+					[
+						SNew(SImage)
+						.Image(this, &SActorDescTreeLabel::GetIcon)
+					.ToolTipText(this, &SActorDescTreeLabel::GetIconTooltip)
+					.ColorAndOpacity(FSlateColor::UseForeground())
+					]
+				]
 
-		+ SHorizontalBox::Slot()
-			.FillWidth(1.0f)
-			.VAlign(VAlign_Center)
-			.Padding(0.0f, 0.0f)
-			[
-				MainContent
-			]
+			+ SHorizontalBox::Slot()
+				.FillWidth(1.0f)
+				.VAlign(VAlign_Center)
+				.Padding(0.0f, 0.0f)
+				[
+					MainContent
+				]
 			];
 	}
 
@@ -229,7 +231,7 @@ FString FActorDescTreeItem::GetDisplayString() const
 
 bool FActorDescTreeItem::CanInteract() const
 {
-	return false;
+	return ActorDescHandle.IsValid();
 }
 
 TSharedRef<SWidget> FActorDescTreeItem::GenerateLabelWidget(ISceneOutliner& Outliner, const STableRow<FSceneOutlinerTreeItemPtr>& InRow)
@@ -237,6 +239,23 @@ TSharedRef<SWidget> FActorDescTreeItem::GenerateLabelWidget(ISceneOutliner& Outl
 	return SNew(SActorDescTreeLabel, *this, Outliner, InRow);
 }
 
+void FActorDescTreeItem::FocusActorBounds() const
+{
+	if (FWorldPartitionActorDesc const* ActorDesc = ActorDescHandle.GetActorDesc())
+	{
+		if (FLevelEditorViewportClient* LevelViewportClient = GCurrentLevelEditingViewportClient)
+		{
+			LevelViewportClient->FocusViewportOnBox(ActorDesc->GetBounds(), false);
+		}	
+	}
+}
+
+void FActorDescTreeItem::GenerateContextMenu(UToolMenu* Menu, SSceneOutliner&)
+{
+	FToolMenuSection& Section = Menu->AddSection("Section");
+	Section.AddMenuEntry("FocusActorBounds", LOCTEXT("FocusActorBounds", "Focus Actor Bounds"), FText(), FSlateIcon(), FUIAction(FExecuteAction::CreateSP(this, &FActorDescTreeItem::FocusActorBounds)));
+}
+
 void FActorDescTreeItem::OnVisibilityChanged(const bool bNewVisibility)
 {
 
diff --git a/Engine/Source/Editor/SceneOutliner/Public/ActorDescTreeItem.h b/Engine/Source/Editor/SceneOutliner/Public/ActorDescTreeItem.h
--- a/Engine/Source/Editor/SceneOutliner/Public/ActorDescTreeItem.h
+++ b/Engine/Source/Editor/SceneOutliner/Public/ActorDescTreeItem.h
@@ -9,6 +9,7 @@
 #include "WorldPartition/WorldPartition.h"
 
 class FWorldPartitionActorDesc;
+class UToolMenu;
 
 /** A tree item that represents an actor in the world */
 struct SCENEOUTLINER_API FActorDescTreeItem : ISceneOutlinerTreeItem
@@ -21,6 +22,9 @@ struct SCENEOUTLINER_API FActorDescTreeItem : ISceneOutlinerTreeItem
 			, Container(InContainer)
 		{
 		}
+
+		bool IsValid() const { return Container.IsValid(); }
+		
 		const FWorldPartitionActorDesc* GetActorDesc() const
 		{
 			return Container.IsValid() ? ActorDesc : nullptr;
@@ -62,8 +66,11 @@ public:
 	virtual FString GetDisplayString() const override;
 	virtual bool CanInteract() const override;
 	virtual TSharedRef<SWidget> GenerateLabelWidget(ISceneOutliner& Outliner, const STableRow<FSceneOutlinerTreeItemPtr>& InRow) override;
+	virtual void GenerateContextMenu(UToolMenu* Menu, SSceneOutliner& Outliner) override;
 	virtual void OnVisibilityChanged(const bool bNewVisibility) override;
 	virtual bool HasVisibilityInfo() const override { return true; }
 	virtual bool GetVisibility() const override;
 	/* End ISceneOutlinerTreeItem Implementation */
+private:
+	void FocusActorBounds() const;	
 };
