SceneOutliner: Fixed an issue where an actor tree item and an actor desc tree item appear at the same time for the same actor.

diff --git a/Engine/Source/Editor/SceneOutliner/Private/ActorHierarchy.cpp b/Engine/Source/Editor/SceneOutliner/Private/ActorHierarchy.cpp
--- a/Engine/Source/Editor/SceneOutliner/Private/ActorHierarchy.cpp
+++ b/Engine/Source/Editor/SceneOutliner/Private/ActorHierarchy.cpp
@@ -31,6 +31,12 @@ TUniquePtr<FActorHierarchy> FActorHierarchy::Create(ISceneOutlinerMode* Mode, co
 	GEngine->OnLevelActorFolderChanged().AddRaw(Hierarchy, &FActorHierarchy::OnLevelActorFolderChanged);
 	GEngine->OnLevelActorListChanged().AddRaw(Hierarchy, &FActorHierarchy::OnLevelActorListChanged);
 
+	if (World.IsValid() && World->PersistentLevel)
+	{
+		World->PersistentLevel->OnLoadedActorAddedToLevelEvent.AddRaw(Hierarchy, &FActorHierarchy::OnLoadedActorAdded);
+		World->PersistentLevel->OnLoadedActorRemovedFromLevelEvent.AddRaw(Hierarchy, &FActorHierarchy::OnLoadedActorRemoved);
+	}
+
 	FWorldDelegates::LevelAddedToWorld.AddRaw(Hierarchy, &FActorHierarchy::OnLevelAdded);
 	FWorldDelegates::LevelRemovedFromWorld.AddRaw(Hierarchy, &FActorHierarchy::OnLevelRemoved);
 
@@ -60,6 +66,12 @@ FActorHierarchy::~FActorHierarchy()
 		GEngine->OnLevelActorListChanged().RemoveAll(this);
 	}
 
+	if (RepresentingWorld.IsValid() && RepresentingWorld->PersistentLevel)
+	{
+		RepresentingWorld->PersistentLevel->OnLoadedActorAddedToLevelEvent.RemoveAll(this);
+		RepresentingWorld->PersistentLevel->OnLoadedActorRemovedFromLevelEvent.RemoveAll(this);
+	}
+
 	FWorldDelegates::LevelAddedToWorld.RemoveAll(this);
 	FWorldDelegates::LevelRemovedFromWorld.RemoveAll(this);
 
@@ -496,6 +508,35 @@ void FActorHierarchy::OnLevelActorDetached(AActor* InActor, const AActor* InPare
 	}
 }
 
+void FActorHierarchy::OnLoadedActorAdded(AActor& InActor)
+{
+	OnLevelActorAdded(&InActor);
+
+	FSceneOutlinerHierarchyChangedData EventData;
+	EventData.Type = FSceneOutlinerHierarchyChangedData::Removed;
+	EventData.ItemIDs.Add(InActor.GetActorGuid());
+	HierarchyChangedEvent.Broadcast(EventData);
+}
+
+void FActorHierarchy::OnLoadedActorRemoved(AActor& InActor)
+{
+	OnLevelActorDeleted(&InActor);
+
+	if (bShowingUnloadedActors)
+	{
+		if (UWorldPartition* WorldPartition = RepresentingWorld->GetWorldPartition())
+		{		
+			if (FWorldPartitionActorDesc* ActorDesc = WorldPartition->GetActorDesc(InActor.GetActorGuid()))
+			{
+				FSceneOutlinerHierarchyChangedData EventData;
+				EventData.Type = FSceneOutlinerHierarchyChangedData::Added;
+				EventData.Items.Add(Mode->CreateItemFor<FActorDescTreeItem>(FActorDescTreeItem(ActorDesc, WorldPartition)));
+				HierarchyChangedEvent.Broadcast(EventData);
+			}
+		}
+	}
+}
+
 void FActorHierarchy::OnComponentsUpdated()
 {
 	FullRefreshEvent();
diff --git a/Engine/Source/Editor/SceneOutliner/Public/ActorHierarchy.h b/Engine/Source/Editor/SceneOutliner/Public/ActorHierarchy.h
--- a/Engine/Source/Editor/SceneOutliner/Public/ActorHierarchy.h
+++ b/Engine/Source/Editor/SceneOutliner/Public/ActorHierarchy.h
@@ -41,6 +41,9 @@ private:
 
 	void OnLevelActorFolderChanged(const AActor* InActor, FName OldPath);
 
+	void OnLoadedActorAdded(AActor& InActor);
+	void OnLoadedActorRemoved(AActor& InActor);
+	
 	void OnComponentsUpdated();
 
 	void OnLevelActorListChanged();
