SceneOutliner: load unloaded actors in folder

diff --git a/Engine/Source/Editor/SceneOutliner/Private/ActorFolderTreeItem.cpp b/Engine/Source/Editor/SceneOutliner/Private/ActorFolderTreeItem.cpp
--- a/Engine/Source/Editor/SceneOutliner/Private/ActorFolderTreeItem.cpp
+++ b/Engine/Source/Editor/SceneOutliner/Private/ActorFolderTreeItem.cpp
@@ -10,6 +10,8 @@
 #include "ISceneOutliner.h"
 #include "ISceneOutlinerMode.h"
 #include "ScopedTransaction.h"
+#include "ToolMenus.h"
+#include "ActorDescTreeItem.h"
 
 #define LOCTEXT_NAMESPACE "SceneOutliner_ActorFolderTreeItem"
 
@@ -275,4 +277,40 @@ TSharedRef<SWidget> FActorFolderTreeItem::GenerateLabelWidget(ISceneOutliner& Ou
 	return SNew(SActorFolderTreeLabel, *this, Outliner, InRow);
 }
 
-#undef LOCTEXT_NAMESPACE
\ No newline at end of file
+void FActorFolderTreeItem::LoadUnloadedActorChildren() const
+{
+	// #todo (outliner): This functionality should be moved to the mode.
+	
+	for (const auto& Child : Children)
+	{
+		if (FSceneOutlinerTreeItemPtr ChildPtr = Child.Pin())
+		{
+			if (const FActorFolderTreeItem* ActorFolderChild = ChildPtr->CastTo<FActorFolderTreeItem>())
+			{
+				ActorFolderChild->LoadUnloadedActorChildren();
+			}
+
+			if (const FActorDescTreeItem* ActorDescChild = ChildPtr->CastTo<FActorDescTreeItem>())
+			{
+				if (const FWorldPartitionActorDesc* ActorDesc = ActorDescChild->ActorDescHandle.GetActorDesc())
+				{
+					if (UActorDescContainer* ActorDescContainer = ActorDescChild->ActorDescHandle.GetActorDescContainer().Get())
+					{
+						new FWorldPartitionReference(ActorDescContainer, ActorDesc->GetGuid());
+					}
+				}
+			}
+		}
+	}
+}
+
+
+void FActorFolderTreeItem::GenerateContextMenu(UToolMenu* Menu, SSceneOutliner& Outliner)
+{
+	FFolderTreeItem::GenerateContextMenu(Menu, Outliner);
+
+	FToolMenuSection& Section = Menu->AddSection("Section");
+	Section.AddMenuEntry("LoadUnloadedActors", LOCTEXT("LoadUnloadedActors", "Load Unloaded Actors"), FText(), FSlateIcon(), FUIAction(FExecuteAction::CreateSP(this, &FActorFolderTreeItem::LoadUnloadedActorChildren)));
+}
+
+#undef LOCTEXT_NAMESPACE
diff --git a/Engine/Source/Editor/SceneOutliner/Public/ActorDescTreeItem.h b/Engine/Source/Editor/SceneOutliner/Public/ActorDescTreeItem.h
--- a/Engine/Source/Editor/SceneOutliner/Public/ActorDescTreeItem.h
+++ b/Engine/Source/Editor/SceneOutliner/Public/ActorDescTreeItem.h
@@ -29,6 +29,11 @@ struct SCENEOUTLINER_API FActorDescTreeItem : ISceneOutlinerTreeItem
 		{
 			return Container.IsValid() ? ActorDesc : nullptr;
 		}
+
+		TWeakObjectPtr<UActorDescContainer> GetActorDescContainer() const
+		{
+			return Container;
+		}
 	private:
 		const FWorldPartitionActorDesc* ActorDesc;
 		TWeakObjectPtr<UActorDescContainer> Container;
diff --git a/Engine/Source/Editor/SceneOutliner/Public/ActorFolderTreeItem.h b/Engine/Source/Editor/SceneOutliner/Public/ActorFolderTreeItem.h
--- a/Engine/Source/Editor/SceneOutliner/Public/ActorFolderTreeItem.h
+++ b/Engine/Source/Editor/SceneOutliner/Public/ActorFolderTreeItem.h
@@ -21,11 +21,13 @@ public:
 	virtual void OnExpansionChanged() override;
 	virtual void Delete(FName InNewParentPath) override;
 	virtual TSharedRef<SWidget> GenerateLabelWidget(ISceneOutliner& Outliner, const STableRow<FSceneOutlinerTreeItemPtr>& InRow) override;
+	virtual void GenerateContextMenu(UToolMenu* Menu, SSceneOutliner& Outliner) override;
 	/* End FFolderTreeItem Implementation */
 		
 	/* Begin FFolderTreeItem Implementation */
 	virtual FName MoveTo(const FName& NewParent) override;
 private:
 	virtual void CreateSubFolder(TWeakPtr<SSceneOutliner> WeakOutliner) override;
+	void LoadUnloadedActorChildren() const;
 	/* End FFolderTreeItem Implementation */
 };
