SceneOutliner: Extended unloaded actor item context menu with option to load single unloaded actor

diff --git a/Engine/Source/Editor/SceneOutliner/Private/ActorDescTreeItem.cpp b/Engine/Source/Editor/SceneOutliner/Private/ActorDescTreeItem.cpp
--- a/Engine/Source/Editor/SceneOutliner/Private/ActorDescTreeItem.cpp
+++ b/Engine/Source/Editor/SceneOutliner/Private/ActorDescTreeItem.cpp
@@ -250,10 +250,22 @@ void FActorDescTreeItem::FocusActorBounds() const
 	}
 }
 
+void FActorDescTreeItem::LoadUnloadedActor() const
+{
+	if (const FWorldPartitionActorDesc* ActorDesc = ActorDescHandle.GetActorDesc())
+	{
+		if (UActorDescContainer* ActorDescContainer = ActorDescHandle.GetActorDescContainer().Get())
+		{
+			new FWorldPartitionReference(ActorDescContainer, ActorDesc->GetGuid());
+		}
+	}
+}
+
 void FActorDescTreeItem::GenerateContextMenu(UToolMenu* Menu, SSceneOutliner&)
 {
 	FToolMenuSection& Section = Menu->AddSection("Section");
 	Section.AddMenuEntry("FocusActorBounds", LOCTEXT("FocusActorBounds", "Focus Actor Bounds"), FText(), FSlateIcon(), FUIAction(FExecuteAction::CreateSP(this, &FActorDescTreeItem::FocusActorBounds)));
+	Section.AddMenuEntry("LoadUnloadedActor", LOCTEXT("LoadUnloadedActor", "Load Unloaded Actor"), FText(), FSlateIcon(), FUIAction(FExecuteAction::CreateSP(this, &FActorDescTreeItem::LoadUnloadedActor)));
 }
 
 void FActorDescTreeItem::OnVisibilityChanged(const bool bNewVisibility)
diff --git a/Engine/Source/Editor/SceneOutliner/Public/ActorDescTreeItem.h b/Engine/Source/Editor/SceneOutliner/Public/ActorDescTreeItem.h
--- a/Engine/Source/Editor/SceneOutliner/Public/ActorDescTreeItem.h
+++ b/Engine/Source/Editor/SceneOutliner/Public/ActorDescTreeItem.h
@@ -77,5 +77,6 @@ public:
 	virtual bool GetVisibility() const override;
 	/* End ISceneOutlinerTreeItem Implementation */
 private:
-	void FocusActorBounds() const;	
+	void FocusActorBounds() const;
+	void LoadUnloadedActor() const;
 };
