SceneOutliner: Fixed crash caused by deleted unloaded actor descs.

diff --git a/Engine/Source/Editor/SceneOutliner/Private/ActorDescTreeItem.cpp b/Engine/Source/Editor/SceneOutliner/Private/ActorDescTreeItem.cpp
--- a/Engine/Source/Editor/SceneOutliner/Private/ActorDescTreeItem.cpp
+++ b/Engine/Source/Editor/SceneOutliner/Private/ActorDescTreeItem.cpp
@@ -210,12 +210,11 @@ private:
 	}
 };
 
-FActorDescTreeItem::FActorDescTreeItem(const FWorldPartitionActorDesc* InActorDesc, UActorDescContainer* Container)
+FActorDescTreeItem::FActorDescTreeItem(const FGuid& InActorGuid, UActorDescContainer* Container)
 	: ISceneOutlinerTreeItem(Type)
-	, ActorDescHandle(InActorDesc, Container)
+	, ActorDescHandle(InActorGuid, Container)
 {
-	check(InActorDesc != nullptr);
-	ID = FSceneOutlinerTreeItemID(InActorDesc->GetGuid());
+	ID = FSceneOutlinerTreeItemID(InActorGuid);
 }
 
 FSceneOutlinerTreeItemID FActorDescTreeItem::GetID() const
diff --git a/Engine/Source/Editor/SceneOutliner/Private/ActorHierarchy.cpp b/Engine/Source/Editor/SceneOutliner/Private/ActorHierarchy.cpp
--- a/Engine/Source/Editor/SceneOutliner/Private/ActorHierarchy.cpp
+++ b/Engine/Source/Editor/SceneOutliner/Private/ActorHierarchy.cpp
@@ -276,7 +276,7 @@ void FActorHierarchy::CreateWorldChildren(UWorld* World, TArray<FSceneOutlinerTr
                 {
                     if (ActorDesc != nullptr && !ActorDesc->IsLoaded())
                     {
-                        if (const FSceneOutlinerTreeItemPtr ActorDescItem = Mode->CreateItemFor<FActorDescTreeItem>(FActorDescTreeItem(ActorDesc, WorldPartition)))
+                        if (const FSceneOutlinerTreeItemPtr ActorDescItem = Mode->CreateItemFor<FActorDescTreeItem>(FActorDescTreeItem(ActorDesc->GetGuid(), WorldPartition)))
                         {
                             OutItems.Add(ActorDescItem);
                         }
@@ -525,12 +525,13 @@ void FActorHierarchy::OnLoadedActorRemoved(AActor& InActor)
 	if (bShowingUnloadedActors)
 	{
 		if (UWorldPartition* WorldPartition = RepresentingWorld->GetWorldPartition())
-		{		
-			if (FWorldPartitionActorDesc* ActorDesc = WorldPartition->GetActorDesc(InActor.GetActorGuid()))
+		{
+			const FGuid& ActorGuid = InActor.GetActorGuid();
+			if (WorldPartition->GetActorDesc(ActorGuid) != nullptr)
 			{
 				FSceneOutlinerHierarchyChangedData EventData;
 				EventData.Type = FSceneOutlinerHierarchyChangedData::Added;
-				EventData.Items.Add(Mode->CreateItemFor<FActorDescTreeItem>(FActorDescTreeItem(ActorDesc, WorldPartition)));
+				EventData.Items.Add(Mode->CreateItemFor<FActorDescTreeItem>(FActorDescTreeItem(ActorGuid, WorldPartition)));
 				HierarchyChangedEvent.Broadcast(EventData);
 			}
 		}
diff --git a/Engine/Source/Editor/SceneOutliner/Public/ActorDescTreeItem.h b/Engine/Source/Editor/SceneOutliner/Public/ActorDescTreeItem.h
--- a/Engine/Source/Editor/SceneOutliner/Public/ActorDescTreeItem.h
+++ b/Engine/Source/Editor/SceneOutliner/Public/ActorDescTreeItem.h
@@ -17,8 +17,8 @@ struct SCENEOUTLINER_API FActorDescTreeItem : ISceneOutlinerTreeItem
 	struct FActorDescHandle
 	{
 	public:
-		FActorDescHandle(const FWorldPartitionActorDesc* InActorDesc, UActorDescContainer* InContainer)
-			: ActorDesc(InActorDesc)
+		FActorDescHandle(const FGuid& ActorGuid, UActorDescContainer* InContainer)
+			: ActorDesc(InContainer, ActorGuid)
 			, Container(InContainer)
 		{
 		}
@@ -27,7 +27,7 @@ struct SCENEOUTLINER_API FActorDescTreeItem : ISceneOutlinerTreeItem
 		
 		const FWorldPartitionActorDesc* GetActorDesc() const
 		{
-			return Container.IsValid() ? ActorDesc : nullptr;
+			return Container.IsValid() ? ActorDesc.Get() : nullptr;
 		}
 
 		TWeakObjectPtr<UActorDescContainer> GetActorDescContainer() const
@@ -35,7 +35,7 @@ struct SCENEOUTLINER_API FActorDescTreeItem : ISceneOutlinerTreeItem
 			return Container;
 		}
 	private:
-		const FWorldPartitionActorDesc* ActorDesc;
+		FWorldPartitionHandle ActorDesc;
 		TWeakObjectPtr<UActorDescContainer> Container;
 	};
 
@@ -63,7 +63,7 @@ public:
 	static const FSceneOutlinerTreeItemType Type;
 
 	/** Construct this item from an actor desc */
-	FActorDescTreeItem(const FWorldPartitionActorDesc* InActorDesc, UActorDescContainer* Container);
+	FActorDescTreeItem(const FGuid& InActorGuid, UActorDescContainer* Container);
 
 	/* Begin ISceneOutlinerTreeItem Implementation */
 	virtual bool IsValid() const override { return ActorDescHandle.GetActorDesc() != nullptr; }
