Added public WorldPartition events for adding and removing actor descs

diff --git a/Engine/Source/Editor/SceneOutliner/Private/ActorHierarchy.cpp b/Engine/Source/Editor/SceneOutliner/Private/ActorHierarchy.cpp
--- a/Engine/Source/Editor/SceneOutliner/Private/ActorHierarchy.cpp
+++ b/Engine/Source/Editor/SceneOutliner/Private/ActorHierarchy.cpp
@@ -31,10 +31,19 @@ TUniquePtr<FActorHierarchy> FActorHierarchy::Create(ISceneOutlinerMode* Mode, co
 	GEngine->OnLevelActorFolderChanged().AddRaw(Hierarchy, &FActorHierarchy::OnLevelActorFolderChanged);
 	GEngine->OnLevelActorListChanged().AddRaw(Hierarchy, &FActorHierarchy::OnLevelActorListChanged);
 
-	if (World.IsValid() && World->PersistentLevel)
+	if (World.IsValid())
 	{
-		World->PersistentLevel->OnLoadedActorAddedToLevelEvent.AddRaw(Hierarchy, &FActorHierarchy::OnLoadedActorAdded);
-		World->PersistentLevel->OnLoadedActorRemovedFromLevelEvent.AddRaw(Hierarchy, &FActorHierarchy::OnLoadedActorRemoved);
+		if (World->PersistentLevel)
+		{
+			World->PersistentLevel->OnLoadedActorAddedToLevelEvent.AddRaw(Hierarchy, &FActorHierarchy::OnLoadedActorAdded);
+			World->PersistentLevel->OnLoadedActorRemovedFromLevelEvent.AddRaw(Hierarchy, &FActorHierarchy::OnLoadedActorRemoved);
+		}
+
+		if (UWorldPartition* WorldPartition = World->GetWorldPartition())
+		{
+			WorldPartition->OnActorDescAddedEvent.AddRaw(Hierarchy, &FActorHierarchy::OnActorDescAdded);
+			WorldPartition->OnActorDescRemovedEvent.AddRaw(Hierarchy, &FActorHierarchy::OnActorDescRemoved);
+		}
 	}
 
 	FWorldDelegates::LevelAddedToWorld.AddRaw(Hierarchy, &FActorHierarchy::OnLevelAdded);
@@ -66,10 +75,19 @@ FActorHierarchy::~FActorHierarchy()
 		GEngine->OnLevelActorListChanged().RemoveAll(this);
 	}
 
-	if (RepresentingWorld.IsValid() && RepresentingWorld->PersistentLevel)
+	if (RepresentingWorld.IsValid())
 	{
-		RepresentingWorld->PersistentLevel->OnLoadedActorAddedToLevelEvent.RemoveAll(this);
-		RepresentingWorld->PersistentLevel->OnLoadedActorRemovedFromLevelEvent.RemoveAll(this);
+		if (RepresentingWorld->PersistentLevel)
+		{
+			RepresentingWorld->PersistentLevel->OnLoadedActorAddedToLevelEvent.RemoveAll(this);
+			RepresentingWorld->PersistentLevel->OnLoadedActorRemovedFromLevelEvent.RemoveAll(this);
+		}
+
+		if (UWorldPartition* WorldPartition = RepresentingWorld->GetWorldPartition())
+		{
+			WorldPartition->OnActorDescAddedEvent.RemoveAll(this);
+			WorldPartition->OnActorDescRemovedEvent.RemoveAll(this);
+		}
 	}
 
 	FWorldDelegates::LevelAddedToWorld.RemoveAll(this);
@@ -538,6 +556,31 @@ void FActorHierarchy::OnLoadedActorRemoved(AActor& InActor)
 	}
 }
 
+void FActorHierarchy::OnActorDescAdded(FWorldPartitionActorDesc* ActorDesc)
+{
+	if (bShowingUnloadedActors && ActorDesc && !ActorDesc->IsLoaded())
+	{
+		if (UWorldPartition* WorldPartition = RepresentingWorld->GetWorldPartition())
+		{
+			FSceneOutlinerHierarchyChangedData EventData;
+			EventData.Type = FSceneOutlinerHierarchyChangedData::Added;
+			EventData.Items.Add(Mode->CreateItemFor<FActorDescTreeItem>(FActorDescTreeItem(ActorDesc->GetGuid(), WorldPartition)));
+			HierarchyChangedEvent.Broadcast(EventData);
+		}
+	}
+}
+
+void FActorHierarchy::OnActorDescRemoved(FWorldPartitionActorDesc* ActorDesc)
+{
+	if (bShowingUnloadedActors && ActorDesc)
+	{
+		FSceneOutlinerHierarchyChangedData EventData;
+		EventData.Type = FSceneOutlinerHierarchyChangedData::Removed;
+		EventData.ItemIDs.Add(ActorDesc->GetGuid());
+		HierarchyChangedEvent.Broadcast(EventData);
+	}
+}
+
 void FActorHierarchy::OnComponentsUpdated()
 {
 	FullRefreshEvent();
diff --git a/Engine/Source/Editor/SceneOutliner/Public/ActorHierarchy.h b/Engine/Source/Editor/SceneOutliner/Public/ActorHierarchy.h
--- a/Engine/Source/Editor/SceneOutliner/Public/ActorHierarchy.h
+++ b/Engine/Source/Editor/SceneOutliner/Public/ActorHierarchy.h
@@ -43,6 +43,9 @@ private:
 
 	void OnLoadedActorAdded(AActor& InActor);
 	void OnLoadedActorRemoved(AActor& InActor);
+
+	void OnActorDescAdded(FWorldPartitionActorDesc* ActorDesc);
+	void OnActorDescRemoved(FWorldPartitionActorDesc* ActorDesc);
 	
 	void OnComponentsUpdated();
 
diff --git a/Engine/Source/Runtime/Engine/Private/WorldPartition/ActorDescContainer.cpp b/Engine/Source/Runtime/Engine/Private/WorldPartition/ActorDescContainer.cpp
--- a/Engine/Source/Runtime/Engine/Private/WorldPartition/ActorDescContainer.cpp
+++ b/Engine/Source/Runtime/Engine/Private/WorldPartition/ActorDescContainer.cpp
@@ -189,7 +189,9 @@ void UActorDescContainer::OnObjectPreSave(UObject* Object)
 			// New actor
 			else
 			{
-				OnActorDescAdded(AddActor(Actor));
+				FWorldPartitionActorDesc* const AddedActorDesc = AddActor(Actor);
+				OnActorDescAdded(AddedActorDesc);
+				OnActorDescAddedEvent.Broadcast(AddedActorDesc);
 			}
 		}
 	}
@@ -203,8 +205,10 @@ void UActorDescContainer::OnPackageDeleted(UPackage* Package)
 	{
 		if (TUniquePtr<FWorldPartitionActorDesc>* ExistingActorDesc = GetActorDescriptor(Actor->GetActorGuid()))
 		{
-			OnActorDescRemoved(ExistingActorDesc->Get());
-			RemoveActorDescriptor(ExistingActorDesc->Get());
+			FWorldPartitionActorDesc* const RemovedActorDesc = ExistingActorDesc->Get();
+			OnActorDescRemovedEvent.Broadcast(RemovedActorDesc);
+			OnActorDescRemoved(RemovedActorDesc);
+			RemoveActorDescriptor(RemovedActorDesc);
 			ExistingActorDesc->Reset();
 		}
 	}
@@ -239,4 +243,14 @@ void UActorDescContainer::UnregisterDelegates()
 		FEditorDelegates::OnPackageDeleted.RemoveAll(this);
 	}
 }
-#endif
\ No newline at end of file
+
+void UActorDescContainer::OnActorDescAdded(FWorldPartitionActorDesc* NewActorDesc)
+{
+	OnActorDescAddedEvent.Broadcast(NewActorDesc);
+}
+
+void UActorDescContainer::OnActorDescRemoved(FWorldPartitionActorDesc* ActorDesc)
+{
+	OnActorDescRemovedEvent.Broadcast(ActorDesc);
+}
+#endif
diff --git a/Engine/Source/Runtime/Engine/Public/WorldPartition/ActorDescContainer.h b/Engine/Source/Runtime/Engine/Public/WorldPartition/ActorDescContainer.h
--- a/Engine/Source/Runtime/Engine/Public/WorldPartition/ActorDescContainer.h
+++ b/Engine/Source/Runtime/Engine/Public/WorldPartition/ActorDescContainer.h
@@ -31,6 +31,13 @@ public:
 
 	/** Removes an actor desc without the need to load a package */
 	virtual void RemoveActor(const FGuid& ActorGuid);
+
+public:
+	DECLARE_EVENT_OneParam(UWorldPartition, FActorDescAddedEvent, FWorldPartitionActorDesc*);
+	FActorDescAddedEvent OnActorDescAddedEvent;
+	
+	DECLARE_EVENT_OneParam(UWorldPartition, FActorDescRemovedEvent, FWorldPartitionActorDesc*);
+	FActorDescRemovedEvent OnActorDescRemovedEvent;
 #endif
 
 	UPROPERTY(Transient)
@@ -45,8 +52,8 @@ protected:
 	virtual void RegisterDelegates();
 	virtual void UnregisterDelegates();
 
-	virtual void OnActorDescAdded(FWorldPartitionActorDesc* NewActorDesc) {}
-	virtual void OnActorDescRemoved(FWorldPartitionActorDesc* ActorDesc) {}
+	virtual void OnActorDescAdded(FWorldPartitionActorDesc* NewActorDesc);
+	virtual void OnActorDescRemoved(FWorldPartitionActorDesc* ActorDesc);
 	virtual void OnActorDescUpdating(FWorldPartitionActorDesc* ActorDesc) {}
 	virtual void OnActorDescUpdated(FWorldPartitionActorDesc* ActorDesc) {}
 
