Fixed crash when WorldPartitionHandle is used externally by manually resetting

diff --git a/Engine/Source/Editor/SceneOutliner/Public/ActorDescTreeItem.h b/Engine/Source/Editor/SceneOutliner/Public/ActorDescTreeItem.h
--- a/Engine/Source/Editor/SceneOutliner/Public/ActorDescTreeItem.h
+++ b/Engine/Source/Editor/SceneOutliner/Public/ActorDescTreeItem.h
@@ -23,6 +23,15 @@ struct SCENEOUTLINER_API FActorDescTreeItem : ISceneOutlinerTreeItem
 		{
 		}
 
+		~FActorDescHandle()
+		{
+			// HACK: prevent the destructor of the handle from accessing the deleted pointer if the container is invalid by manually resetting it.
+			if (!Container.IsValid())
+			{
+				ActorDesc.Reset();
+			}
+		}
+
 		bool IsValid() const { return Container.IsValid(); }
 		
 		const FWorldPartitionActorDesc* GetActorDesc() const
diff --git a/Engine/Source/Runtime/Engine/Public/WorldPartition/WorldPartitionHandle.h b/Engine/Source/Runtime/Engine/Public/WorldPartition/WorldPartitionHandle.h
--- a/Engine/Source/Runtime/Engine/Public/WorldPartition/WorldPartitionHandle.h
+++ b/Engine/Source/Runtime/Engine/Public/WorldPartition/WorldPartitionHandle.h
@@ -174,7 +174,12 @@ public:
 	{
 		return IsValid() ? ActorDesc->Get() : nullptr;
 	}
-		
+
+	FORCEINLINE void Reset()
+	{
+		ActorDesc = nullptr;
+	}
+	
 	friend FORCEINLINE uint32 GetTypeHash(const TWorldPartitionHandle<Impl>& HandleBase)
 	{
 		return ::PointerHash(HandleBase.ActorDesc);
