Fixed Compile errors in ActorHierarchy.cpp and ActorHierarchy.h in non-unity builds.

diff --git a/Engine/Source/Editor/SceneOutliner/Public/ActorHierarchy.h b/Engine/Source/Editor/SceneOutliner/Public/ActorHierarchy.h
--- a/Engine/Source/Editor/SceneOutliner/Public/ActorHierarchy.h
+++ b/Engine/Source/Editor/SceneOutliner/Public/ActorHierarchy.h
@@ -5,6 +5,8 @@
 #include "CoreMinimal.h"
 #include "ISceneOutlinerHierarchy.h"
 
+class FWorldPartitionActorDesc;
+
 class FActorHierarchy : public ISceneOutlinerHierarchy
 {
 public:
