World Outliner: show the actor's label as tooltip instead of the internal name.
Useful in cases where the actor's label is truncated because of the dialog's width.

diff --git a/Engine/Source/Editor/SceneOutliner/Private/ActorTreeItem.cpp b/Engine/Source/Editor/SceneOutliner/Private/ActorTreeItem.cpp
--- a/Engine/Source/Editor/SceneOutliner/Private/ActorTreeItem.cpp
+++ b/Engine/Source/Editor/SceneOutliner/Private/ActorTreeItem.cpp
@@ -135,12 +135,9 @@ private:
 
 	FText GetTooltipText() const
 	{
-		if (const AActor* Actor = ActorPtr.Get())
+		if (const FSceneOutlinerTreeItemPtr TreeItem = TreeItemPtr.Pin())
 		{
-			FFormatNamedArguments Args;
-			Args.Add(TEXT("ID_Name"), LOCTEXT("CustomColumnMode_InternalName", "ID Name"));
-			Args.Add(TEXT("Name"), FText::FromString(Actor->GetName()));
-			return FText::Format(LOCTEXT("ActorNameTooltip", "{ID_Name}: {Name}"), Args);
+			return FText::FromString(TreeItem->GetDisplayString());
 		}
 
 		return FText();
