(Actually submit the change) WorldOutliner: Change check into checkSlow because it had a huge impact on the refresh. (FullRefresh in Frosty was ~1:30min now down to ~20seconds)

diff --git a/Engine/Source/Editor/SceneOutliner/Public/ISceneOutlinerTreeItem.h b/Engine/Source/Editor/SceneOutliner/Public/ISceneOutlinerTreeItem.h
--- a/Engine/Source/Editor/SceneOutliner/Public/ISceneOutlinerTreeItem.h
+++ b/Engine/Source/Editor/SceneOutliner/Public/ISceneOutlinerTreeItem.h
@@ -73,7 +73,7 @@ public:
 	/** Add a child to this item */
 	void AddChild(FSceneOutlinerTreeItemRef Child)
 	{
-		check(!Children.Contains(Child));
+		checkSlow(!Children.Contains(Child));
 		Child->Parent = AsShared();
 		Children.Add(MoveTemp(Child));
 	}
