SceneOutliner: Disable intermediate sorting when doing a FullRefresh + Use a time based budget for processing items (default to 5ms) makes the full refresh off 100k items < 1 second.

diff --git a/Engine/Source/Editor/SceneOutliner/Private/SSceneOutliner.cpp b/Engine/Source/Editor/SceneOutliner/Private/SSceneOutliner.cpp
--- a/Engine/Source/Editor/SceneOutliner/Private/SSceneOutliner.cpp
+++ b/Engine/Source/Editor/SceneOutliner/Private/SSceneOutliner.cpp
@@ -42,6 +42,13 @@
 
 DEFINE_LOG_CATEGORY_STATIC(LogSceneOutliner, Log, All);
 
+static float GSceneOutlinerProcessingBudgetPerFrame = 5.0f;
+static FAutoConsoleVariableRef CVarGuardBandMultiplier(
+	TEXT("SceneOutliner.ProcessingBudgetPerFrame"),
+	GSceneOutlinerProcessingBudgetPerFrame,
+	TEXT("Maximum time in milliseconds to spend processing operations per frame"));
+
+
 #define LOCTEXT_NAMESPACE "SSceneOutliner"
 
 // The amount of time that must pass before the Scene Outliner will attempt a sort when in PIE/SIE.
@@ -388,6 +395,7 @@ void SSceneOutliner::Refresh()
 void SSceneOutliner::FullRefresh()
 {
 	UE_LOG(LogSceneOutliner, Verbose, TEXT("Full Refresh"));
+	bDisableIntermediateSorting = true;
 	bFullRefresh = true;
 	Refresh();
 }
@@ -420,9 +428,11 @@ void SSceneOutliner::Populate()
 		bFullRefresh = false;
 	}
 
-	// Only deal with 500 at a time
-	const int32 End = FMath::Min(PendingOperations.Num(), 500);
-	for (int32 Index = 0; Index < End; ++Index)
+	const double StartTime = FPlatformTime::Seconds();
+	// To avoid checking the time budget for every item.
+	const int32 CheckBudgetEveryNthItem = 100;
+	int32 Index = 0;
+	while (Index < PendingOperations.Num())
 	{
 		auto& PendingOp = PendingOperations[Index];
 		switch (PendingOp.Type)
@@ -445,10 +455,22 @@ void SSceneOutliner::Populate()
 			check(false);
 			break;
 		}
-	}
 
-	PendingOperations.RemoveAt(0, End);
+		++Index;
 
+		if ((Index % CheckBudgetEveryNthItem) == 0)
+		{
+			const double TimeSpentInMs = (FPlatformTime::Seconds() - StartTime) * 1000.0;
+			if (TimeSpentInMs > GSceneOutlinerProcessingBudgetPerFrame)
+			{
+				UE_LOG(LogSceneOutliner, Verbose, TEXT("Processing out of budget (%.2f ms) : %.2f ms"), GSceneOutlinerProcessingBudgetPerFrame, (float)TimeSpentInMs);
+				break;
+			}
+		}
+	}
+		
+	UE_LOG(LogSceneOutliner, Verbose, TEXT("%d Items Processed"), Index);
+	PendingOperations.RemoveAt(0, Index);
 
 	for (FName Folder : PendingFoldersSelect)
 	{
@@ -1644,7 +1666,6 @@ void SSceneOutliner::PostUndo(bool bSuccess)
 	// Refresh our tree in case any changes have been made to the scene that might effect our list
 	if( !bIsReentrant )
 	{
-		bDisableIntermediateSorting = true;
 		FullRefresh();
 	}
 }
