SceneOutliner: Tree items now store pointers to children in a TSet instead of TArray for quicker lookups. Cached display string for ActorDesc items to prevent construction of temporary strings.

diff --git a/Engine/Source/Editor/SceneOutliner/Private/ActorDescTreeItem.cpp b/Engine/Source/Editor/SceneOutliner/Private/ActorDescTreeItem.cpp
--- a/Engine/Source/Editor/SceneOutliner/Private/ActorDescTreeItem.cpp
+++ b/Engine/Source/Editor/SceneOutliner/Private/ActorDescTreeItem.cpp
@@ -215,6 +215,15 @@ FActorDescTreeItem::FActorDescTreeItem(const FGuid& InActorGuid, UActorDescConta
 	, ActorDescHandle(InActorGuid, Container)
 {
 	ID = FSceneOutlinerTreeItemID(InActorGuid);
+
+	if (const FWorldPartitionActorDesc* const ActorDesc = ActorDescHandle.GetActorDesc())
+	{
+		DisplayString = ActorDesc->GetActorLabel().ToString();
+	}
+	else
+	{
+		DisplayString = LOCTEXT("ActorLabelForMissingActor", "(Deleted Actor)").ToString();
+	}
 }
 
 FSceneOutlinerTreeItemID FActorDescTreeItem::GetID() const
@@ -224,8 +233,7 @@ FSceneOutlinerTreeItemID FActorDescTreeItem::GetID() const
 
 FString FActorDescTreeItem::GetDisplayString() const
 {
-	const FWorldPartitionActorDesc* ActorDesc = ActorDescHandle.GetActorDesc();
-	return ActorDesc ? ActorDesc->GetActorLabel().ToString() : LOCTEXT("ActorLabelForMissingActor", "(Deleted Actor)").ToString();
+	return DisplayString;
 }
 
 bool FActorDescTreeItem::CanInteract() const
diff --git a/Engine/Source/Editor/SceneOutliner/Private/SSceneOutliner.cpp b/Engine/Source/Editor/SceneOutliner/Private/SSceneOutliner.cpp
--- a/Engine/Source/Editor/SceneOutliner/Private/SSceneOutliner.cpp
+++ b/Engine/Source/Editor/SceneOutliner/Private/SSceneOutliner.cpp
@@ -1013,7 +1013,17 @@ TSharedRef<TSet<FName>> SSceneOutliner::GatherInvalidMoveToDestinations() const
 			};
 
 			// Exclude this items direct parent if it is a folder and has no other subfolders we can move to
-			if (!ParentFolderItem->GetChildren().ContainsByPredicate(FolderHasOtherSubFolders))
+			bool bFolderHasSubFolders = false;
+			for (const TWeakPtr<ISceneOutlinerTreeItem>& ItemPtr : ParentFolderItem->GetChildren())
+			{
+				if (FolderHasOtherSubFolders(ItemPtr))
+				{
+					bFolderHasSubFolders = true;
+					break;
+				}
+			}
+			
+			if (!bFolderHasSubFolders)
 			{
 				ExcludedParents->Add(ParentFolderItem->Path);
 			}
diff --git a/Engine/Source/Editor/SceneOutliner/Public/ActorDescTreeItem.h b/Engine/Source/Editor/SceneOutliner/Public/ActorDescTreeItem.h
--- a/Engine/Source/Editor/SceneOutliner/Public/ActorDescTreeItem.h
+++ b/Engine/Source/Editor/SceneOutliner/Public/ActorDescTreeItem.h
@@ -88,4 +88,6 @@ public:
 private:
 	void FocusActorBounds() const;
 	void LoadUnloadedActor() const;
+
+	FString DisplayString;
 };
diff --git a/Engine/Source/Editor/SceneOutliner/Public/ISceneOutlinerTreeItem.h b/Engine/Source/Editor/SceneOutliner/Public/ISceneOutlinerTreeItem.h
--- a/Engine/Source/Editor/SceneOutliner/Public/ISceneOutlinerTreeItem.h
+++ b/Engine/Source/Editor/SceneOutliner/Public/ISceneOutlinerTreeItem.h
@@ -55,7 +55,7 @@ protected:
 	mutable TWeakPtr<ISceneOutlinerTreeItem> Parent;
 
 	/** Array of children contained underneath this item */
-	mutable TArray<TWeakPtr<ISceneOutlinerTreeItem>> Children;
+	mutable TSet<TWeakPtr<ISceneOutlinerTreeItem>> Children;
 
 	/** Static type identifier for the base class tree item */
 	static const FSceneOutlinerTreeItemType Type;
@@ -73,7 +73,7 @@ public:
 	/** Add a child to this item */
 	void AddChild(FSceneOutlinerTreeItemRef Child)
 	{
-		checkSlow(!Children.Contains(Child));
+		check(!Children.Contains(Child));
 		Child->Parent = AsShared();
 		Children.Add(MoveTemp(Child));
 	}
@@ -88,7 +88,7 @@ public:
 	}
 
 	/** Get this item's children, if any. Although we store as weak pointers, they are guaranteed to be valid. */
-	FORCEINLINE const TArray<TWeakPtr<ISceneOutlinerTreeItem>>& GetChildren() const
+	FORCEINLINE const TSet<TWeakPtr<ISceneOutlinerTreeItem>>& GetChildren() const
 	{
 		return Children;
 	}
