WorldOutliner: Pin/unpin world partition actors from the context menu

diff --git a/Engine/Source/Editor/SceneOutliner/Private/ActorBrowsingMode.cpp b/Engine/Source/Editor/SceneOutliner/Private/ActorBrowsingMode.cpp
--- a/Engine/Source/Editor/SceneOutliner/Private/ActorBrowsingMode.cpp
+++ b/Engine/Source/Editor/SceneOutliner/Private/ActorBrowsingMode.cpp
@@ -372,6 +372,33 @@ void FActorBrowsingMode::RegisterContextMenu()
 						{
 							SceneOutliner->GetTree().GetSelectedItems()[0]->GenerateContextMenu(InMenu, *SceneOutliner);
 						}
+						
+						if (Context->NumSelectedItems > 0)
+						{
+							// If selection contains some unpinned items, show the pin option
+							// If the selection contains folders, always show the pin option
+							if (Context->NumPinnedItems != Context->NumSelectedItems || Context->NumSelectedFolders > 0)
+							{
+								InMenu->FindOrAddSection("Section").AddMenuEntry(
+									"PinItems",
+									LOCTEXT("Pin", "Pin"),
+									FText(),
+									FSlateIcon(),
+									FUIAction(FExecuteAction::CreateSP(SceneOutliner, &SSceneOutliner::PinSelectedItems)));
+							}
+							
+							// If the selection contains some pinned items, show the unpin option
+							// If the selection contains folders, always show the unpin option
+							if (Context->NumPinnedItems != 0 || Context->NumSelectedFolders > 0)
+							{
+								InMenu->FindOrAddSection("Section").AddMenuEntry(
+									"UnpinItems",
+									LOCTEXT("Unpin", "Unpin"),
+									FText(),
+									FSlateIcon(),
+									FUIAction(FExecuteAction::CreateSP(SceneOutliner, &SSceneOutliner::UnpinSelectedItems)));
+							}
+						}
 
 						// If we've only got folders selected, show the selection and edit sub menus
 						if (Context->NumSelectedItems > 0 && Context->NumSelectedFolders == Context->NumSelectedItems)
@@ -417,7 +444,7 @@ TSharedPtr<SWidget> FActorBrowsingMode::BuildContextMenu()
 {
 	RegisterContextMenu();
 
-	FSceneOutlinerItemSelection ItemSelection(SceneOutliner->GetSelection());
+	const FSceneOutlinerItemSelection ItemSelection(SceneOutliner->GetSelection());
 
 	USceneOutlinerMenuContext* ContextObject = NewObject<USceneOutlinerMenuContext>();
 	ContextObject->SceneOutliner = StaticCastSharedRef<SSceneOutliner>(SceneOutliner->AsShared());
@@ -425,6 +452,21 @@ TSharedPtr<SWidget> FActorBrowsingMode::BuildContextMenu()
 	ContextObject->NumSelectedItems = ItemSelection.Num();
 	ContextObject->NumSelectedFolders = ItemSelection.Num<FFolderTreeItem>();
 	ContextObject->NumWorldsSelected = ItemSelection.Num<FWorldTreeItem>();
+
+	int32 NumPinnedItems = 0;
+	if (const UWorldPartition* const WorldPartition = RepresentingWorld->GetWorldPartition())
+	{
+		ItemSelection.ForEachItem<FActorTreeItem>([WorldPartition, &NumPinnedItems](const FActorTreeItem& ActorItem)
+		{
+			if (ActorItem.Actor.IsValid() && WorldPartition->IsActorPinned(ActorItem.Actor->GetActorGuid()))
+			{
+				++NumPinnedItems;
+			}
+			return true;
+		});
+	}
+	ContextObject->NumPinnedItems = NumPinnedItems;
+
 	FToolMenuContext Context(ContextObject);
 
 	FName MenuName = DefaultContextMenuName;
@@ -1378,6 +1420,82 @@ void FActorBrowsingMode::SelectFoldersDescendants(const TArray<FFolderTreeItem*>
 	GEditor->NoteSelectionChange();
 }
 
+void FActorBrowsingMode::PinItem(const FSceneOutlinerTreeItemPtr& InItem)
+{
+	if (UWorldPartition* const WorldPartition = RepresentingWorld->GetWorldPartition())
+	{
+		if (const FActorDescTreeItem* const ActorDescTreeItem = InItem->CastTo<FActorDescTreeItem>())
+		{
+			WorldPartition->PinActor(ActorDescTreeItem->GetGuid());
+		}
+		else if (const FActorTreeItem* const ActorTreeItem = InItem->CastTo<FActorTreeItem>())
+		{
+			if (ActorTreeItem->Actor.IsValid())
+			{
+				WorldPartition->PinActor(ActorTreeItem->Actor->GetActorGuid());
+			}
+		}
+	}
+
+	// Recursively pin all children.
+	for (const auto& Child : InItem->GetChildren())
+	{
+		if (Child.IsValid())
+		{
+			PinItem(Child.Pin());
+	    }
+	}
+}
+
+void FActorBrowsingMode::UnpinItem(const FSceneOutlinerTreeItemPtr& InItem)
+{
+	// Recursively unpin all children
+	for (const auto& Child : InItem->GetChildren())
+	{
+		if (Child.IsValid())
+		{
+			UnpinItem(Child.Pin());
+		}
+	}
+	
+	if (UWorldPartition* const WorldPartition = RepresentingWorld->GetWorldPartition())
+	{
+		if (const FActorDescTreeItem* const ActorDescTreeItem = InItem->CastTo<FActorDescTreeItem>())
+		{
+			WorldPartition->UnpinActor(ActorDescTreeItem->GetGuid());
+		}
+		else if (const FActorTreeItem* const ActorTreeItem = InItem->CastTo<FActorTreeItem>())
+		{
+			if (ActorTreeItem->Actor.IsValid())
+			{
+				WorldPartition->UnpinActor(ActorTreeItem->Actor->GetActorGuid());
+			}
+		}
+	}
+}
+
+void FActorBrowsingMode::PinSelectedItems()
+{
+	const FSceneOutlinerItemSelection Selection = SceneOutliner->GetSelection();
+
+	Selection.ForEachItem([this](const FSceneOutlinerTreeItemPtr& TreeItem)
+	{
+		PinItem(TreeItem);
+		return true;
+	});
+}
+
+void FActorBrowsingMode::UnpinSelectedItems()
+{
+	const FSceneOutlinerItemSelection Selection = SceneOutliner->GetSelection();
+
+	Selection.ForEachItem([this](const FSceneOutlinerTreeItemPtr& TreeItem)
+	{
+		UnpinItem(TreeItem);
+		return true;
+	});
+}
+
 FCreateSceneOutlinerMode FActorBrowsingMode::CreateFolderPickerMode() const
 {
 	auto MoveSelectionTo = [this](const FSceneOutlinerTreeItemRef& NewParent)
diff --git a/Engine/Source/Editor/SceneOutliner/Private/ActorDescTreeItem.cpp b/Engine/Source/Editor/SceneOutliner/Private/ActorDescTreeItem.cpp
--- a/Engine/Source/Editor/SceneOutliner/Private/ActorDescTreeItem.cpp
+++ b/Engine/Source/Editor/SceneOutliner/Private/ActorDescTreeItem.cpp
@@ -213,6 +213,7 @@ private:
 FActorDescTreeItem::FActorDescTreeItem(const FGuid& InActorGuid, UActorDescContainer* Container)
 	: ISceneOutlinerTreeItem(Type)
 	, ActorDescHandle(InActorGuid, Container)
+	, ActorGuid(InActorGuid)
 {
 	ID = FSceneOutlinerTreeItemID(InActorGuid);
 
@@ -257,22 +258,10 @@ void FActorDescTreeItem::FocusActorBounds() const
 	}
 }
 
-void FActorDescTreeItem::LoadUnloadedActor() const
-{
-	if (const FWorldPartitionActorDesc* ActorDesc = ActorDescHandle.GetActorDesc())
-	{
-		if (UActorDescContainer* ActorDescContainer = ActorDescHandle.GetActorDescContainer().Get())
-		{
-			new FWorldPartitionReference(ActorDescContainer, ActorDesc->GetGuid());
-		}
-	}
-}
-
 void FActorDescTreeItem::GenerateContextMenu(UToolMenu* Menu, SSceneOutliner&)
 {
 	FToolMenuSection& Section = Menu->AddSection("Section");
 	Section.AddMenuEntry("FocusActorBounds", LOCTEXT("FocusActorBounds", "Focus Actor Bounds"), FText(), FSlateIcon(), FUIAction(FExecuteAction::CreateSP(this, &FActorDescTreeItem::FocusActorBounds)));
-	Section.AddMenuEntry("LoadUnloadedActor", LOCTEXT("LoadUnloadedActor", "Load Unloaded Actor"), FText(), FSlateIcon(), FUIAction(FExecuteAction::CreateSP(this, &FActorDescTreeItem::LoadUnloadedActor)));
 }
 
 void FActorDescTreeItem::OnVisibilityChanged(const bool bNewVisibility)
diff --git a/Engine/Source/Editor/SceneOutliner/Private/ActorFolderTreeItem.cpp b/Engine/Source/Editor/SceneOutliner/Private/ActorFolderTreeItem.cpp
--- a/Engine/Source/Editor/SceneOutliner/Private/ActorFolderTreeItem.cpp
+++ b/Engine/Source/Editor/SceneOutliner/Private/ActorFolderTreeItem.cpp
@@ -277,40 +277,4 @@ TSharedRef<SWidget> FActorFolderTreeItem::GenerateLabelWidget(ISceneOutliner& Ou
 	return SNew(SActorFolderTreeLabel, *this, Outliner, InRow);
 }
 
-void FActorFolderTreeItem::LoadUnloadedActorChildren() const
-{
-	// #todo (outliner): This functionality should be moved to the mode.
-	
-	for (const auto& Child : Children)
-	{
-		if (FSceneOutlinerTreeItemPtr ChildPtr = Child.Pin())
-		{
-			if (const FActorFolderTreeItem* ActorFolderChild = ChildPtr->CastTo<FActorFolderTreeItem>())
-			{
-				ActorFolderChild->LoadUnloadedActorChildren();
-			}
-
-			if (const FActorDescTreeItem* ActorDescChild = ChildPtr->CastTo<FActorDescTreeItem>())
-			{
-				if (const FWorldPartitionActorDesc* ActorDesc = ActorDescChild->ActorDescHandle.GetActorDesc())
-				{
-					if (UActorDescContainer* ActorDescContainer = ActorDescChild->ActorDescHandle.GetActorDescContainer().Get())
-					{
-						new FWorldPartitionReference(ActorDescContainer, ActorDesc->GetGuid());
-					}
-				}
-			}
-		}
-	}
-}
-
-
-void FActorFolderTreeItem::GenerateContextMenu(UToolMenu* Menu, SSceneOutliner& Outliner)
-{
-	FFolderTreeItem::GenerateContextMenu(Menu, Outliner);
-
-	FToolMenuSection& Section = Menu->AddSection("Section");
-	Section.AddMenuEntry("LoadUnloadedActors", LOCTEXT("LoadUnloadedActors", "Load Unloaded Actors"), FText(), FSlateIcon(), FUIAction(FExecuteAction::CreateSP(this, &FActorFolderTreeItem::LoadUnloadedActorChildren)));
-}
-
 #undef LOCTEXT_NAMESPACE
diff --git a/Engine/Source/Editor/SceneOutliner/Private/SSceneOutliner.cpp b/Engine/Source/Editor/SceneOutliner/Private/SSceneOutliner.cpp
--- a/Engine/Source/Editor/SceneOutliner/Private/SSceneOutliner.cpp
+++ b/Engine/Source/Editor/SceneOutliner/Private/SSceneOutliner.cpp
@@ -1996,6 +1996,26 @@ FReply SSceneOutliner::OnDragOverItem(const FDragDropEvent& Event, const ISceneO
 	return Mode->OnDragOverItem(Event, Item); 
 }
 
+void SSceneOutliner::PinItem(const FSceneOutlinerTreeItemPtr& InItem)
+{
+	Mode->PinItem(InItem);
+}
+
+void SSceneOutliner::UnpinItem(const FSceneOutlinerTreeItemPtr& InItem)
+{
+	Mode->UnpinItem(InItem);
+}
+
+void SSceneOutliner::PinSelectedItems()
+{
+	Mode->PinSelectedItems();
+}
+
+void SSceneOutliner::UnpinSelectedItems()
+{
+	Mode->UnpinSelectedItems();
+}
+
 FSceneOutlinerTreeItemPtr SSceneOutliner::FindParent(const ISceneOutlinerTreeItem& InItem) const
 {
 	FSceneOutlinerTreeItemPtr Parent = Mode->GetHierarchy()->FindParent(InItem, TreeItemMap);
diff --git a/Engine/Source/Editor/SceneOutliner/Public/ActorBrowsingMode.h b/Engine/Source/Editor/SceneOutliner/Public/ActorBrowsingMode.h
--- a/Engine/Source/Editor/SceneOutliner/Public/ActorBrowsingMode.h
+++ b/Engine/Source/Editor/SceneOutliner/Public/ActorBrowsingMode.h
@@ -49,6 +49,10 @@ public:
 	virtual FName CreateFolder(const FName& ParentPath, const FName& LeafName) override;
 	virtual bool ReparentItemToFolder(const FName& FolderPath, const FSceneOutlinerTreeItemPtr& Item) override;
 	virtual void SelectFoldersDescendants(const TArray<FFolderTreeItem*>& FolderItems, bool bSelectImmediateChildrenOnly) override;
+	virtual void PinItem(const FSceneOutlinerTreeItemPtr& InItem) override;
+	virtual void UnpinItem(const FSceneOutlinerTreeItemPtr& InItem) override;
+	virtual void PinSelectedItems() override;
+	virtual void UnpinSelectedItems() override;
 	/* End ISceneOutlinerMode Interface */
 public:
 	/* External events this mode must respond to */
diff --git a/Engine/Source/Editor/SceneOutliner/Public/ActorDescTreeItem.h b/Engine/Source/Editor/SceneOutliner/Public/ActorDescTreeItem.h
--- a/Engine/Source/Editor/SceneOutliner/Public/ActorDescTreeItem.h
+++ b/Engine/Source/Editor/SceneOutliner/Public/ActorDescTreeItem.h
@@ -85,9 +85,11 @@ public:
 	virtual bool HasVisibilityInfo() const override { return true; }
 	virtual bool GetVisibility() const override;
 	/* End ISceneOutlinerTreeItem Implementation */
+	
+	const FGuid& GetGuid() const { return ActorGuid; }
 private:
 	void FocusActorBounds() const;
-	void LoadUnloadedActor() const;
 
 	FString DisplayString;
+	FGuid ActorGuid;
 };
diff --git a/Engine/Source/Editor/SceneOutliner/Public/ActorFolderTreeItem.h b/Engine/Source/Editor/SceneOutliner/Public/ActorFolderTreeItem.h
--- a/Engine/Source/Editor/SceneOutliner/Public/ActorFolderTreeItem.h
+++ b/Engine/Source/Editor/SceneOutliner/Public/ActorFolderTreeItem.h
@@ -17,17 +17,15 @@ public:
 	TWeakObjectPtr<UWorld> World;
 
 	/* Begin ISceneOutlinerTreeItem Implementation */
-	virtual bool IsValid() const { return World.IsValid(); }
+	virtual bool IsValid() const override { return World.IsValid(); }
 	virtual void OnExpansionChanged() override;
 	virtual void Delete(FName InNewParentPath) override;
 	virtual TSharedRef<SWidget> GenerateLabelWidget(ISceneOutliner& Outliner, const STableRow<FSceneOutlinerTreeItemPtr>& InRow) override;
-	virtual void GenerateContextMenu(UToolMenu* Menu, SSceneOutliner& Outliner) override;
 	/* End FFolderTreeItem Implementation */
 		
 	/* Begin FFolderTreeItem Implementation */
 	virtual FName MoveTo(const FName& NewParent) override;
 private:
 	virtual void CreateSubFolder(TWeakPtr<SSceneOutliner> WeakOutliner) override;
-	void LoadUnloadedActorChildren() const;
 	/* End FFolderTreeItem Implementation */
 };
diff --git a/Engine/Source/Editor/SceneOutliner/Public/ISceneOutliner.h b/Engine/Source/Editor/SceneOutliner/Public/ISceneOutliner.h
--- a/Engine/Source/Editor/SceneOutliner/Public/ISceneOutliner.h
+++ b/Engine/Source/Editor/SceneOutliner/Public/ISceneOutliner.h
@@ -119,6 +119,18 @@ public:
 	/** Set the item selection of the outliner based on a selector function. Any items which return true will be added */
 	virtual void SetSelection(const TFunctionRef<bool(ISceneOutlinerTreeItem&)> Selector) = 0;
 
+	/** Pins an item in the outliner. */
+	virtual void PinItem(const FSceneOutlinerTreeItemPtr& InItem) = 0;
+
+	/**  Unpins an item in the outliner. */
+	virtual void UnpinItem(const FSceneOutlinerTreeItemPtr& InItem) = 0;
+	
+	/** Pin selected items */
+	virtual void PinSelectedItems() = 0;
+
+	/** Unpins selected items */
+	virtual void UnpinSelectedItems() = 0;
+
 	/** Get the active SceneOutlinerMode */
 	const ISceneOutlinerMode* GetMode() const { return Mode; }
 protected:
diff --git a/Engine/Source/Editor/SceneOutliner/Public/ISceneOutlinerMode.h b/Engine/Source/Editor/SceneOutliner/Public/ISceneOutlinerMode.h
--- a/Engine/Source/Editor/SceneOutliner/Public/ISceneOutlinerMode.h
+++ b/Engine/Source/Editor/SceneOutliner/Public/ISceneOutlinerMode.h
@@ -135,6 +135,15 @@ public:
 	virtual bool ReparentItemToFolder(const FName& FolderPath, const FSceneOutlinerTreeItemPtr& Item) { return false; }
 	/** Select all descendants of a folder. Optionally select only immediate descendants. */
 	virtual void SelectFoldersDescendants(const TArray<FFolderTreeItem*>& FolderItems, bool bSelectImmediateChildrenOnly) {}
+public:
+	/** Pins an item in the outliner */
+	virtual void PinItem(const FSceneOutlinerTreeItemPtr& InItem) {}
+	/** Unpins an item in the outliner */
+	virtual void UnpinItem(const FSceneOutlinerTreeItemPtr& InItem) {}
+	/** Pins all selected items */
+	virtual void PinSelectedItems() {}
+	/** Unpins all selected items */
+	virtual void UnpinSelectedItems() {}
 public:
 	/* Getters */
 
diff --git a/Engine/Source/Editor/SceneOutliner/Public/SSceneOutliner.h b/Engine/Source/Editor/SceneOutliner/Public/SSceneOutliner.h
--- a/Engine/Source/Editor/SceneOutliner/Public/SSceneOutliner.h
+++ b/Engine/Source/Editor/SceneOutliner/Public/SSceneOutliner.h
@@ -173,6 +173,18 @@ struct FSceneOutlinerItemSelection
 		}
 	}
 
+	/** Apply a function to every item in the selection regardless of type. */
+	void ForEachItem(TFunctionRef<void(FSceneOutlinerTreeItemPtr&)> Func) const
+	{
+		for (const TWeakPtr<ISceneOutlinerTreeItem>& Item : SelectedItems)
+		{
+			if (auto ItemPtr = Item.Pin())
+			{
+				Func(ItemPtr);
+			}
+		}
+	}
+
 	/** Use a selector to retrieve a specific data type from items in the selection. Will only add an item's data if the selector returns true for that item. */
 	template <typename DataType>
 	TArray<DataType> GetData(TFunctionRef<bool(const TWeakPtr<ISceneOutlinerTreeItem>&, DataType&)> Selector) const
@@ -369,6 +381,26 @@ public:
 		*/
 	virtual FSceneOutlinerItemSelection GetSelection() const { return FSceneOutlinerItemSelection(*OutlinerTreeView); }
 
+	/**
+	 * Pins an item in the outliner.
+	 */
+	virtual void PinItem(const FSceneOutlinerTreeItemPtr& InItem) override;
+
+	/**
+	 * Unpins an item in the outliner.
+	 */
+	virtual void UnpinItem(const FSceneOutlinerTreeItemPtr& InItem) override;
+
+	/**
+	 * Pin selected items
+	 */
+	virtual void PinSelectedItems() override;
+
+	/**
+	 * Unpins selected items
+	 */
+	virtual void UnpinSelectedItems() override;
+	
 	/**
 	 * Returns the parent tree item for a given item if it exists, nullptr otherwise.
 	 */
diff --git a/Engine/Source/Editor/SceneOutliner/Public/SceneOutlinerMenuContext.h b/Engine/Source/Editor/SceneOutliner/Public/SceneOutlinerMenuContext.h
--- a/Engine/Source/Editor/SceneOutliner/Public/SceneOutlinerMenuContext.h
+++ b/Engine/Source/Editor/SceneOutliner/Public/SceneOutlinerMenuContext.h
@@ -22,6 +22,7 @@ public:
 		NumSelectedItems = 0;
 		NumSelectedFolders = 0;
 		NumWorldsSelected = 0;
+		NumPinnedItems = 0;
 	}
 
 	TWeakPtr<SSceneOutliner> SceneOutliner;
@@ -30,4 +31,5 @@ public:
 	int32 NumSelectedItems;
 	int32 NumSelectedFolders;
 	int32 NumWorldsSelected;
+	int32 NumPinnedItems;
 };
