SceneOutliner: Fully implemented the concept of pinned items as well as a column to display the pinned state of items.

diff --git a/Engine/Source/Editor/EditorStyle/Private/StarshipStyle.cpp b/Engine/Source/Editor/EditorStyle/Private/StarshipStyle.cpp
--- a/Engine/Source/Editor/EditorStyle/Private/StarshipStyle.cpp
+++ b/Engine/Source/Editor/EditorStyle/Private/StarshipStyle.cpp
@@ -1172,6 +1172,11 @@ void FStarshipEditorStyle::FStyle::SetupGeneralStyles()
 		Set( "SceneOutliner.World", 		new CORE_IMAGE_BRUSH_SVG("Starship/Common/world", Icon16x16 ) );
 		Set( "SceneOutliner.ChangedItemHighlight", new FSlateRoundedBoxBrush( FStyleColors::White, 1.0) );
 
+		Set("SceneOutliner.PinnedHighlighIcon", new IMAGE_BRUSH( "Icons/icon_Blueprint_CommentBubblePin_16x", Icon10x10));
+		Set("SceneOutliner.PinnedIcon", new IMAGE_BRUSH( "Icons/icon_Blueprint_CommentBubblePin_16x", Icon10x10));
+		Set("SceneOutliner.UnpinnedHighlighIcon", new IMAGE_BRUSH( "Icons/icon_Blueprint_CommentBubbleUnPin_16x", Icon10x10));
+		Set("SceneOutliner.UnpinnedIcon", new IMAGE_BRUSH( "Icons/icon_Blueprint_CommentBubbleUnPin_16x", Icon10x10));
+
 		Set( "SceneOutliner.TableViewRow", FTableRowStyle(NormalTableRowStyle));
 	}
 
diff --git a/Engine/Source/Editor/SceneOutliner/Private/ActorBrowsingMode.cpp b/Engine/Source/Editor/SceneOutliner/Private/ActorBrowsingMode.cpp
--- a/Engine/Source/Editor/SceneOutliner/Private/ActorBrowsingMode.cpp
+++ b/Engine/Source/Editor/SceneOutliner/Private/ActorBrowsingMode.cpp
@@ -186,6 +186,20 @@ void FActorBrowsingMode::Rebuild()
 
 	FilteredActorCount = 0;
 	ApplicableActors.Empty();
+
+	// Enable the pinned column by default on WP worlds
+	if (RepresentingWorld.IsValid() && RepresentingWorld->GetWorldPartition() != nullptr)
+	{
+		if (!bPinnedColumnActive)
+		{
+			TogglePinnedColumn();
+		}
+	}
+	// Disable it by default on non-WP worlds
+	else if (bPinnedColumnActive)
+	{
+		TogglePinnedColumn();
+	}
 }
 
 FText FActorBrowsingMode::GetStatusText() const 
@@ -253,11 +267,24 @@ void FActorBrowsingMode::CreateViewContent(FMenuBuilder& MenuBuilder)
 			FUIAction(
 				FExecuteAction::CreateRaw(this, &FActorBrowsingMode::ToggleActorSCCStatusColumn),
 				FCanExecuteAction(),
-				FIsActionChecked::CreateRaw(this, &FActorBrowsingMode::isActorSCCStatusColumnActive)
+				FIsActionChecked::CreateRaw(this, &FActorBrowsingMode::IsActorSCCStatusColumnActive)
 			),
 			NAME_None,
 			EUserInterfaceActionType::ToggleButton
 		);
+
+		MenuBuilder.AddMenuEntry(
+			LOCTEXT("PinnedColumnName", "Pinned Column"),
+			LOCTEXT("PinnedColumnToolip", "Displays the pinned state of items"),
+			FSlateIcon(),
+			FUIAction(
+				FExecuteAction::CreateRaw(this, &FActorBrowsingMode::TogglePinnedColumn),
+				FCanExecuteAction(),
+				FIsActionChecked::CreateRaw(this, &FActorBrowsingMode::IsPinnedColumnActive)
+			),
+			NAME_None,
+			EUserInterfaceActionType::ToggleButton
+			);
 	}
 	MenuBuilder.EndSection();
 
@@ -1596,4 +1623,17 @@ void FActorBrowsingMode::ToggleActorSCCStatusColumn()
 	bActorSCCStatusColumnActive = !bActorSCCStatusColumnActive;
 }
 
-#undef LOCTEXT_NAMESPACE
\ No newline at end of file
+void FActorBrowsingMode::TogglePinnedColumn()
+{
+	if (bPinnedColumnActive)
+	{
+		SceneOutliner->RemoveColumn(FSceneOutlinerBuiltInColumnTypes::Pinned());
+	}
+	else
+	{
+		SceneOutliner->AddColumn(FSceneOutlinerBuiltInColumnTypes::Pinned(), FSceneOutlinerColumnInfo(ESceneOutlinerColumnVisibility::Visible, 5));
+	}
+	bPinnedColumnActive = !bPinnedColumnActive;
+}
+
+#undef LOCTEXT_NAMESPACE
diff --git a/Engine/Source/Editor/SceneOutliner/Private/ActorDescTreeItem.cpp b/Engine/Source/Editor/SceneOutliner/Private/ActorDescTreeItem.cpp
--- a/Engine/Source/Editor/SceneOutliner/Private/ActorDescTreeItem.cpp
+++ b/Engine/Source/Editor/SceneOutliner/Private/ActorDescTreeItem.cpp
@@ -274,4 +274,13 @@ bool FActorDescTreeItem::GetVisibility() const
 	return true;
 }
 
+bool FActorDescTreeItem::GetPinnedState() const
+{
+	if (ActorDescHandle.IsValid())
+	{
+		return ActorDescHandle.GetActorDescContainer()->IsActorPinned(GetGuid());
+	}
+	return false;
+}
+
 #undef LOCTEXT_NAMESPACE
diff --git a/Engine/Source/Editor/SceneOutliner/Private/ActorTreeItem.cpp b/Engine/Source/Editor/SceneOutliner/Private/ActorTreeItem.cpp
--- a/Engine/Source/Editor/SceneOutliner/Private/ActorTreeItem.cpp
+++ b/Engine/Source/Editor/SceneOutliner/Private/ActorTreeItem.cpp
@@ -18,6 +18,7 @@
 #include "Logging/MessageLog.h"
 #include "SSocketChooser.h"
 #include "LevelInstance/LevelInstanceActor.h"
+#include "WorldPartition/WorldPartition.h"
 
 #define LOCTEXT_NAMESPACE "SceneOutliner_ActorTreeItem"
 
@@ -395,6 +396,22 @@ bool FActorTreeItem::GetVisibility() const
 	return !Actor.IsValid() || !Actor->IsTemporarilyHiddenInEditor(true);
 }
 
+bool FActorTreeItem::GetPinnedState() const
+{
+	if (Actor.IsValid())
+	{
+		if (const UWorld* const World = Actor->GetWorld())
+		{
+			if (const UWorldPartition* const WorldPartition = World->GetWorldPartition())
+			{
+				return WorldPartition->IsActorPinned(Actor->GetActorGuid());
+			}
+		}
+	}
+
+	return false;
+}
+
 void FActorTreeItem::OnLabelChanged()
 {
 	if (Actor.IsValid())
diff --git a/Engine/Source/Editor/SceneOutliner/Private/SceneOutlinerModule.cpp b/Engine/Source/Editor/SceneOutliner/Private/SceneOutlinerModule.cpp
--- a/Engine/Source/Editor/SceneOutliner/Private/SceneOutlinerModule.cpp
+++ b/Engine/Source/Editor/SceneOutliner/Private/SceneOutlinerModule.cpp
@@ -10,6 +10,7 @@
 #include "SceneOutlinerGutter.h"
 #include "SceneOutlinerItemLabelColumn.h"
 #include "SceneOutlinerActorSCCColumn.h"
+#include "SceneOutlinerPinnedColumn.h"
 #include "SceneOutlinerPublicTypes.h"
 
 #include "ActorPickingMode.h"
@@ -28,6 +29,7 @@ void FSceneOutlinerModule::StartupModule()
 	RegisterColumnType<FSceneOutlinerGutter>();
 	RegisterColumnType<FActorInfoColumn>();
 	RegisterColumnType<FSceneOutlinerActorSCCColumn>();
+	RegisterColumnType<FSceneOutlinerPinnedColumn>();
 }
 
 
@@ -36,6 +38,8 @@ void FSceneOutlinerModule::ShutdownModule()
 	UnRegisterColumnType<FSceneOutlinerGutter>();
 	UnRegisterColumnType<FSceneOutlinerItemLabelColumn>();
 	UnRegisterColumnType<FActorInfoColumn>();
+	UnRegisterColumnType<FSceneOutlinerActorSCCColumn>();
+	UnRegisterColumnType<FSceneOutlinerPinnedColumn>();
 }
 
 TSharedRef<ISceneOutliner> FSceneOutlinerModule::CreateSceneOutliner(const FSceneOutlinerInitializationOptions& InitOptions) const
diff --git a/Engine/Source/Editor/SceneOutliner/Private/SceneOutlinerPinnedColumn.cpp b/Engine/Source/Editor/SceneOutliner/Private/SceneOutlinerPinnedColumn.cpp
new file mode 100644
--- /dev/null
+++ b/Engine/Source/Editor/SceneOutliner/Private/SceneOutlinerPinnedColumn.cpp
@@ -0,0 +1,238 @@
+// Copyright Epic Games, Inc. All Rights Reserved.
+
+#include "SceneOutlinerPinnedColumn.h"
+#include "Widgets/Images/SImage.h"
+#include "Widgets/Views/STreeView.h"
+
+#define LOCTEXT_NAMESPACE "SceneOutlinerPinnedActorColumn"
+
+bool FSceneOutlinerPinnedColumn::FSceneOutlinerPinnedStateCache::CheckChildren(const ISceneOutlinerTreeItem& Item) const
+{
+	if (const bool* const State = PinnedStateInfo.Find(&Item))
+	{
+		return *State;
+	}
+
+	bool bIsPinned = false;
+	for (const auto& ChildPtr : Item.GetChildren())
+	{
+		FSceneOutlinerTreeItemPtr Child = ChildPtr.Pin();
+		if (Child.IsValid() && GetPinnedState(*Child))
+		{
+			bIsPinned = true;
+			break;
+		}
+	}
+	PinnedStateInfo.Add(&Item, bIsPinned);
+	
+	return bIsPinned;
+}
+
+bool FSceneOutlinerPinnedColumn::FSceneOutlinerPinnedStateCache::GetPinnedState(const ISceneOutlinerTreeItem& Item) const
+{
+	if (Item.HasPinnedStateInfo())
+	{
+		if (const bool* const State = PinnedStateInfo.Find(&Item))
+		{
+			return *State;
+		}
+
+		const bool bIsPinned = Item.GetPinnedState();
+		PinnedStateInfo.Add(&Item, bIsPinned);
+		return bIsPinned;
+	}
+
+	return CheckChildren(Item);
+}
+
+void FSceneOutlinerPinnedColumn::FSceneOutlinerPinnedStateCache::Empty()
+{
+	PinnedStateInfo.Empty();
+}
+
+class SPinnedWidget : public SImage
+{
+public:
+	SLATE_BEGIN_ARGS(SPinnedWidget) {}
+	SLATE_END_ARGS()
+
+	/** Construct this widget */
+	void Construct(const FArguments& InArgs, TWeakPtr<ISceneOutliner> InWeakOutliner, TWeakPtr<ISceneOutlinerTreeItem> InWeakTreeItem, const TWeakPtr<FSceneOutlinerPinnedColumn>& InWeakColumn, const STableRow<FSceneOutlinerTreeItemPtr>* InRow)
+	{
+		WeakTreeItem = InWeakTreeItem;
+		WeakOutliner = InWeakOutliner;
+		WeakColumn = InWeakColumn;
+		Row = InRow;
+
+		SImage::Construct(
+			SImage::FArguments()
+            .ColorAndOpacity(this, &SPinnedWidget::GetForegroundColor)
+            .Image(this, &SPinnedWidget::GetBrush)
+        );
+
+		static const FName NAME_PinnedHoveredBrush = TEXT("SceneOutliner.PinnedHighlighIcon");
+		static const FName NAME_PinnedNotHoveredBrush = TEXT("SceneOutliner.PinnedIcon");
+		static const FName NAME_UnpinnedHoveredBrush = TEXT("SceneOutliner.UnpinnedHighlighIcon");
+		static const FName NAME_UnpinnedNotHoveredBrush = TEXT("SceneOutliner.UnpinnedIcon");
+
+		PinnedHoveredBrush = FEditorStyle::GetBrush(NAME_PinnedHoveredBrush);
+		PinnedNotHoveredBrush = FEditorStyle::GetBrush(NAME_PinnedNotHoveredBrush);
+		UnpinnedHoveredBrush = FEditorStyle::GetBrush(NAME_UnpinnedHoveredBrush);
+		UnpinnedNotHoveredBrush = FEditorStyle::GetBrush(NAME_UnpinnedNotHoveredBrush);
+	}
+
+private:
+	bool IsPinned() const
+	{
+		return WeakTreeItem.IsValid() && WeakColumn.IsValid() ? WeakColumn.Pin()->IsItemPinned(*WeakTreeItem.Pin()) : false;
+	}
+	
+	FReply HandleClick() const
+	{
+		if (!WeakTreeItem.IsValid() || !WeakOutliner.IsValid())
+		{
+			return FReply::Unhandled();
+		}
+
+		const auto Outliner = WeakOutliner.Pin();
+		const auto TreeItem = WeakTreeItem.Pin();
+		const auto& Tree = Outliner->GetTree();
+
+		if (!IsPinned())
+		{
+			if (Tree.IsItemSelected(TreeItem.ToSharedRef()))
+			{
+				Outliner->PinSelectedItems();
+			}
+			else
+			{
+				Outliner->PinItem(TreeItem);
+			}
+		}
+		else
+		{
+			if (Tree.IsItemSelected(TreeItem.ToSharedRef()))
+			{
+				Outliner->UnpinSelectedItems();
+			}
+			else
+			{
+				Outliner->UnpinItem(TreeItem);
+			}
+		}
+
+		return FReply::Handled();
+	}
+
+	virtual FReply OnMouseButtonDoubleClick(const FGeometry& InMyGeometry, const FPointerEvent& InMouseEvent) override
+	{
+		return HandleClick();
+	}
+
+	virtual FReply OnMouseButtonDown(const FGeometry& MyGeometry, const FPointerEvent& MouseEvent) override
+	{
+		if (MouseEvent.GetEffectingButton() == EKeys::LeftMouseButton)
+		{
+			return HandleClick();
+		}
+		
+		return FReply::Unhandled();
+	}
+
+	const FSlateBrush* GetBrush() const
+	{
+		if (IsPinned())
+		{
+			return IsHovered() ? PinnedHoveredBrush : PinnedNotHoveredBrush;
+		}
+		else
+		{
+			return IsHovered() ? UnpinnedHoveredBrush : UnpinnedNotHoveredBrush;
+		}
+	}
+
+	virtual FSlateColor GetForegroundColor() const override
+	{
+		const auto Outliner = WeakOutliner.Pin();
+		const auto TreeItem = WeakTreeItem.Pin();
+		const bool bIsSelected = Outliner->GetTree().IsItemSelected(TreeItem.ToSharedRef());
+		
+		if (!IsPinned())
+		{
+			if (!Row->IsHovered() && !bIsSelected)
+			{
+				return FLinearColor::Transparent;
+			}
+		}
+		
+		return IsHovered() ? FSlateColor::UseForeground() : FSlateColor::UseSubduedForeground();
+	}
+
+	/** The tree item we relate to */
+	TWeakPtr<ISceneOutlinerTreeItem> WeakTreeItem;
+
+	/** Weak pointer back to the outliner */
+	TWeakPtr<ISceneOutliner> WeakOutliner;
+	/** Weak pointer back to the column to check cached pinned state */
+	TWeakPtr<FSceneOutlinerPinnedColumn> WeakColumn;
+
+	/** Weak pointer back to the row */
+	const STableRow<FSceneOutlinerTreeItemPtr>* Row = nullptr;
+
+	const FSlateBrush* PinnedHoveredBrush = nullptr;
+	const FSlateBrush* PinnedNotHoveredBrush = nullptr;
+	const FSlateBrush* UnpinnedHoveredBrush = nullptr;
+	const FSlateBrush* UnpinnedNotHoveredBrush = nullptr;
+};
+
+FName FSceneOutlinerPinnedColumn::GetColumnID()
+{
+	return GetID();
+}
+
+SHeaderRow::FColumn::FArguments FSceneOutlinerPinnedColumn::ConstructHeaderRowColumn()
+{
+	return SHeaderRow::Column(GetColumnID())
+		.FixedWidth(24.f)
+		.HAlignHeader(HAlign_Center)
+		.VAlignHeader(VAlign_Center)
+		.HAlignCell(HAlign_Center)
+		.VAlignCell(VAlign_Center)
+		[
+			SNew(SImage)
+			.ColorAndOpacity(FSlateColor::UseForeground())
+			.Image(this, &FSceneOutlinerPinnedColumn::GetHeaderIcon)
+		];
+}
+
+const TSharedRef<SWidget> FSceneOutlinerPinnedColumn::ConstructRowWidget(FSceneOutlinerTreeItemRef TreeItem, const STableRow<FSceneOutlinerTreeItemPtr>& Row)
+{
+	if (TreeItem->ShouldShowPinnedState())
+	{
+		return SNew(SHorizontalBox)
+		+ SHorizontalBox::Slot()
+		.AutoWidth()
+		.VAlign(VAlign_Center)
+		[
+			SNew(SPinnedWidget, WeakSceneOutliner, TreeItem, SharedThis(this), &Row)
+		];
+	}
+	return SNullWidget::NullWidget;
+}
+
+void FSceneOutlinerPinnedColumn::Tick(double InCurrentTime, float InDeltaTime)
+{
+	PinnedStateCache.Empty();
+}
+
+bool FSceneOutlinerPinnedColumn::IsItemPinned(const ISceneOutlinerTreeItem& Item) const
+{
+	return PinnedStateCache.GetPinnedState(Item);
+}
+
+const FSlateBrush* FSceneOutlinerPinnedColumn::GetHeaderIcon() const
+{
+	return FEditorStyle::GetBrush("SceneOutliner.PinnedHighlighIcon");
+}
+
+#undef LOCTEXT_NAMESPACE
\ No newline at end of file
diff --git a/Engine/Source/Editor/SceneOutliner/Private/SceneOutlinerPinnedColumn.h b/Engine/Source/Editor/SceneOutliner/Private/SceneOutlinerPinnedColumn.h
new file mode 100644
--- /dev/null
+++ b/Engine/Source/Editor/SceneOutliner/Private/SceneOutlinerPinnedColumn.h
@@ -0,0 +1,53 @@
+// Copyright Epic Games, Inc. All Rights Reserved.
+
+#pragma once
+
+#include "CoreMinimal.h"
+#include "Widgets/SWidget.h"
+#include "Widgets/Views/SHeaderRow.h"
+#include "ISceneOutliner.h"
+#include "SceneOutlinerPublicTypes.h"
+#include "ISceneOutlinerColumn.h"
+
+template<typename ItemType> class STableRow;
+
+/** A column for the SceneOutliner that displays the item label */
+class FSceneOutlinerPinnedColumn : public ISceneOutlinerColumn
+{
+	struct FSceneOutlinerPinnedStateCache
+	{
+	public:
+		/** Recursive over the children and check the pinned state */
+		bool CheckChildren(const ISceneOutlinerTreeItem& Item) const;
+		/** Get or cache the pinned state of an item */
+		bool GetPinnedState(const ISceneOutlinerTreeItem& Item) const;
+
+		/** Clear the cache */
+		void Empty();
+	private:
+		/** Map of tree item to pinned state */
+		mutable TMap<const ISceneOutlinerTreeItem*, bool> PinnedStateInfo;
+	};
+public:
+	FSceneOutlinerPinnedColumn(ISceneOutliner& SceneOutliner) : WeakSceneOutliner(StaticCastSharedRef<ISceneOutliner>(SceneOutliner.AsShared())) {}
+
+	virtual ~FSceneOutlinerPinnedColumn() {}
+
+	static FName GetID() { return FSceneOutlinerBuiltInColumnTypes::Pinned(); }
+
+	// Begin ISceneOutlinerColumn Implementation
+	virtual FName GetColumnID() override;
+	virtual SHeaderRow::FColumn::FArguments ConstructHeaderRowColumn() override;
+	virtual const TSharedRef<SWidget> ConstructRowWidget(FSceneOutlinerTreeItemRef TreeItem, const STableRow<FSceneOutlinerTreeItemPtr>& Row) override;
+	virtual void Tick(double InCurrentTime, float InDeltaTime) override;
+	virtual bool SupportsSorting() const override { return false; }
+	// End ISceneOutlinerColumn Implementation
+
+	bool IsItemPinned(const ISceneOutlinerTreeItem& Item) const;
+private:
+	const FSlateBrush* GetHeaderIcon() const;
+	TWeakPtr<ISceneOutliner> WeakSceneOutliner;
+
+	/** Cached pinned state of items per-frame */
+	FSceneOutlinerPinnedStateCache PinnedStateCache;
+};
diff --git a/Engine/Source/Editor/SceneOutliner/Public/ActorBrowsingMode.h b/Engine/Source/Editor/SceneOutliner/Public/ActorBrowsingMode.h
--- a/Engine/Source/Editor/SceneOutliner/Public/ActorBrowsingMode.h
+++ b/Engine/Source/Editor/SceneOutliner/Public/ActorBrowsingMode.h
@@ -107,7 +107,10 @@ private:
 	static TSharedRef<FSceneOutlinerFilter> CreateHideUnloadedActorsFilter();
 	
 	void ToggleActorSCCStatusColumn();
-	bool isActorSCCStatusColumnActive() const { return bActorSCCStatusColumnActive; }
+	bool IsActorSCCStatusColumnActive() const { return bActorSCCStatusColumnActive; }
+
+	void TogglePinnedColumn();
+	bool IsPinnedColumnActive() const { return bPinnedColumnActive; }
 private:
 	/** Number of actors which have passed through the filters */
 	uint32 FilteredActorCount = 0;
@@ -115,4 +118,5 @@ private:
 	TSet<TWeakObjectPtr<AActor>> ApplicableActors;
 
 	bool bActorSCCStatusColumnActive = false;
+	bool bPinnedColumnActive = false;
 };
diff --git a/Engine/Source/Editor/SceneOutliner/Public/ActorDescTreeItem.h b/Engine/Source/Editor/SceneOutliner/Public/ActorDescTreeItem.h
--- a/Engine/Source/Editor/SceneOutliner/Public/ActorDescTreeItem.h
+++ b/Engine/Source/Editor/SceneOutliner/Public/ActorDescTreeItem.h
@@ -84,6 +84,9 @@ public:
 	virtual void OnVisibilityChanged(const bool bNewVisibility) override;
 	virtual bool HasVisibilityInfo() const override { return true; }
 	virtual bool GetVisibility() const override;
+	virtual bool ShouldShowPinnedState() const override { return true; }
+	virtual bool HasPinnedStateInfo() const override { return true; }
+	virtual bool GetPinnedState() const override;
 	/* End ISceneOutlinerTreeItem Implementation */
 	
 	const FGuid& GetGuid() const { return ActorGuid; }
diff --git a/Engine/Source/Editor/SceneOutliner/Public/ActorTreeItem.h b/Engine/Source/Editor/SceneOutliner/Public/ActorTreeItem.h
--- a/Engine/Source/Editor/SceneOutliner/Public/ActorTreeItem.h
+++ b/Engine/Source/Editor/SceneOutliner/Public/ActorTreeItem.h
@@ -45,6 +45,9 @@ public:
 	virtual void OnVisibilityChanged(const bool bNewVisibility) override;
 	virtual bool HasVisibilityInfo() const override { return true; }
 	virtual bool GetVisibility() const override;
+	virtual bool ShouldShowPinnedState() const override { return true; }
+	virtual bool HasPinnedStateInfo() const override { return true; };
+	virtual bool GetPinnedState() const override;
 	virtual void OnLabelChanged() override;
 	/* End ISceneOutlinerTreeItem Implementation */
 public:
diff --git a/Engine/Source/Editor/SceneOutliner/Public/FolderTreeItem.h b/Engine/Source/Editor/SceneOutliner/Public/FolderTreeItem.h
--- a/Engine/Source/Editor/SceneOutliner/Public/FolderTreeItem.h
+++ b/Engine/Source/Editor/SceneOutliner/Public/FolderTreeItem.h
@@ -56,6 +56,9 @@ public:
 	virtual void GenerateContextMenu(UToolMenu* Menu, SSceneOutliner& Outliner) override;
 	/** Delete this folder, children will be reparented to provided new parent path */
 	virtual void Delete(FName InNewParentPath) {}
+	
+	virtual bool ShouldShowPinnedState() const override { return true; }
+	virtual bool HasPinnedStateInfo() const override { return false; }
 	/* End ISceneOutlinerTreeItem Implementation */
 
 	/** Move this folder to a new parent */
diff --git a/Engine/Source/Editor/SceneOutliner/Public/ISceneOutlinerTreeItem.h b/Engine/Source/Editor/SceneOutliner/Public/ISceneOutlinerTreeItem.h
--- a/Engine/Source/Editor/SceneOutliner/Public/ISceneOutlinerTreeItem.h
+++ b/Engine/Source/Editor/SceneOutliner/Public/ISceneOutlinerTreeItem.h
@@ -147,6 +147,15 @@ public:
 	/** Query this items visibility state. Only called if the item type has visibility info */
 	virtual bool GetVisibility() const { return false; }
 
+	/** Returns true if this item should show pinned state */
+	virtual bool ShouldShowPinnedState() const { return false; }
+
+	/** Returns true if this item can set its own pinned state */
+	virtual bool HasPinnedStateInfo() const { return false; }
+
+	/** Query the pinned staet of this item. Only called if the item has its own pinned state info */
+	virtual bool GetPinnedState() const { return false; }
+
 	/** Called when this item's label has changed */
 	virtual void OnLabelChanged() {}
 };
diff --git a/Engine/Source/Editor/SceneOutliner/Public/SceneOutlinerPublicTypes.h b/Engine/Source/Editor/SceneOutliner/Public/SceneOutlinerPublicTypes.h
--- a/Engine/Source/Editor/SceneOutliner/Public/SceneOutlinerPublicTypes.h
+++ b/Engine/Source/Editor/SceneOutliner/Public/SceneOutlinerPublicTypes.h
@@ -46,6 +46,12 @@ struct FSceneOutlinerBuiltInColumnTypes
 		static FName SourceControl("SourceControl");
 		return SourceControl;
 	}
+
+	static FName& Pinned()
+	{
+		static FName Pinned("Pinned");
+		return Pinned;
+	}
 };
 
 /** Visibility enum for scene outliner columns */
