SceneOutliner: Include unloaded actor count in status text.

diff --git a/Engine/Source/Editor/SceneOutliner/Private/ActorBrowsingMode.cpp b/Engine/Source/Editor/SceneOutliner/Private/ActorBrowsingMode.cpp
--- a/Engine/Source/Editor/SceneOutliner/Private/ActorBrowsingMode.cpp
+++ b/Engine/Source/Editor/SceneOutliner/Private/ActorBrowsingMode.cpp
@@ -164,6 +164,13 @@ FActorBrowsingMode::FActorBrowsingMode(SSceneOutliner* InSceneOutliner, TWeakObj
 
 FActorBrowsingMode::~FActorBrowsingMode()
 {
+	if (RepresentingWorld.IsValid())
+	{
+		if (UWorldPartition* const WorldPartition = RepresentingWorld->GetWorldPartition())
+		{
+			WorldPartition->OnActorDescRemovedEvent.RemoveAll(this);
+		}
+	}
 	FSceneOutlinerDelegates::Get().OnComponentsUpdated.RemoveAll(this);
 
 	GEngine->OnLevelActorDeleted().RemoveAll(this);
@@ -182,14 +189,30 @@ FActorBrowsingMode::~FActorBrowsingMode()
 
 void FActorBrowsingMode::Rebuild()
 {
+	// If we used to be representing a wp world, unbind delegates before rebuilding begins
+	if (RepresentingWorld.IsValid())
+	{
+		if (UWorldPartition* const WorldPartition = RepresentingWorld->GetWorldPartition())
+		{
+			WorldPartition->OnActorDescRemovedEvent.RemoveAll(this);
+		}
+	}
+	
 	FActorMode::Rebuild();
 
 	FilteredActorCount = 0;
+	FilteredUnloadedActorCount = 0;
+	ApplicableUnloadedActors.Empty();
 	ApplicableActors.Empty();
 
-	// Enable the pinned column by default on WP worlds
-	if (RepresentingWorld.IsValid() && RepresentingWorld->GetWorldPartition() != nullptr)
+	bRepresentingWorldPartitionedWorld = RepresentingWorld.IsValid() && RepresentingWorld->GetWorldPartition() != nullptr;
+
+	if (bRepresentingWorldPartitionedWorld)
 	{
+		UWorldPartition* const WorldPartition = RepresentingWorld->GetWorldPartition();
+		WorldPartition->OnActorDescRemovedEvent.AddRaw(this, &FActorBrowsingMode::OnActorDescRemoved);
+
+		// Enable the pinned column by default on WP worlds
 		if (!bPinnedColumnActive)
 		{
 			TogglePinnedColumn();
@@ -209,14 +232,22 @@ FText FActorBrowsingMode::GetStatusText() const
 		return FText();
 	}
 
-	const int32 TotalActorCount = ApplicableActors.Num();
-	const int32 SelectedActorCount = SceneOutliner->GetSelection().Num<FActorTreeItem>();
+	// The number of actors in the outliner before applying the text filter
+	const int32 TotalActorCount = ApplicableActors.Num() + ApplicableUnloadedActors.Num();
+	const int32 SelectedActorCount = SceneOutliner->GetSelection().Num<FActorTreeItem, FActorDescTreeItem>();
 
 	if (!SceneOutliner->IsTextFilterActive())
 	{
 		if (SelectedActorCount == 0) //-V547
 		{
-			return FText::Format(LOCTEXT("ShowingAllActorsFmt", "{0} actors"), FText::AsNumber(FilteredActorCount));
+			if (bRepresentingWorldPartitionedWorld)
+			{
+				return FText::Format(LOCTEXT("ShowingAllActorsFmt", "{0} actors ({1} loaded)"), FText::AsNumber(FilteredActorCount), FText::AsNumber(FilteredActorCount - FilteredUnloadedActorCount));
+			}
+			else
+			{
+				return FText::Format(LOCTEXT("ShowingAllActorsFmt", "{0} actors"), FText::AsNumber(FilteredActorCount));
+			}
 		}
 		else
 		{
@@ -560,17 +591,33 @@ void FActorBrowsingMode::OnItemAdded(FSceneOutlinerTreeItemPtr Item)
 			}
 		}
 	}
+	else if (Item->IsA<FActorDescTreeItem>())
+	{
+		if (!Item->Flags.bIsFilteredOut)
+		{
+			++FilteredActorCount;
+			++FilteredUnloadedActorCount;
+		}
+	}
 }
 
 void FActorBrowsingMode::OnItemRemoved(FSceneOutlinerTreeItemPtr Item)
 {
-	if (const FActorTreeItem* ActorItem = Item->CastTo<FActorTreeItem>())
+	if (Item->IsA<FActorTreeItem>())
 	{
-		if (!ActorItem->Flags.bIsFilteredOut)
+		if (!Item->Flags.bIsFilteredOut)
 		{
 			--FilteredActorCount;
 		}
 	}
+	else if (Item->IsA<FActorDescTreeItem>())
+	{
+		if (!Item->Flags.bIsFilteredOut)
+		{
+			--FilteredActorCount;
+			--FilteredUnloadedActorCount;
+		}
+	}
 }
 
 void FActorBrowsingMode::OnComponentsUpdated()
@@ -583,6 +630,11 @@ void FActorBrowsingMode::OnLevelActorDeleted(AActor* Actor)
 	ApplicableActors.Remove(Actor);
 }
 
+void FActorBrowsingMode::OnActorDescRemoved(FWorldPartitionActorDesc* InActorDesc)
+{
+	ApplicableUnloadedActors.Remove(InActorDesc);
+}
+
 void FActorBrowsingMode::OnItemSelectionChanged(FSceneOutlinerTreeItemPtr TreeItem, ESelectInfo::Type SelectionType, const FSceneOutlinerItemSelection& Selection)
 {
 	TArray<AActor*> SelectedActors = Selection.GetData<AActor*>(SceneOutliner::FActorSelector());
@@ -731,10 +783,14 @@ void FActorBrowsingMode::OnFilterTextCommited(FSceneOutlinerItemSelection& Selec
 
 void FActorBrowsingMode::OnItemPassesFilters(const ISceneOutlinerTreeItem& Item)
 {
-	if (const FActorTreeItem* ActorItem = Item.CastTo<FActorTreeItem>())
+	if (const FActorTreeItem* const ActorItem = Item.CastTo<FActorTreeItem>())
 	{
 		ApplicableActors.Add(ActorItem->Actor);
 	}
+	else if (const FActorDescTreeItem* const ActorDescItem = Item.CastTo<FActorDescTreeItem>(); ActorDescItem && ActorDescItem->IsValid())
+	{
+		ApplicableUnloadedActors.Add(ActorDescItem->ActorDescHandle.GetActorDesc());
+	}
 }
 
 FReply FActorBrowsingMode::OnKeyDown(const FKeyEvent& InKeyEvent)
@@ -1636,4 +1692,4 @@ void FActorBrowsingMode::TogglePinnedColumn()
 	bPinnedColumnActive = !bPinnedColumnActive;
 }
 
-#undef LOCTEXT_NAMESPACE
+#undef LOCTEXT_NAMESPACE
\ No newline at end of file
diff --git a/Engine/Source/Editor/SceneOutliner/Public/ActorBrowsingMode.h b/Engine/Source/Editor/SceneOutliner/Public/ActorBrowsingMode.h
--- a/Engine/Source/Editor/SceneOutliner/Public/ActorBrowsingMode.h
+++ b/Engine/Source/Editor/SceneOutliner/Public/ActorBrowsingMode.h
@@ -61,6 +61,9 @@ public:
 	void OnComponentsUpdated();
 	/** Called by the engine when an actor is deleted */
 	void OnLevelActorDeleted(AActor* Actor);
+	
+	/** Called when an actor desc is removed */
+	void OnActorDescRemoved(FWorldPartitionActorDesc* InActorDesc);
 
 	/** Called by engine when edit cut actors begins */
 	void OnEditCutActorsBegin();
@@ -112,11 +115,16 @@ private:
 	void TogglePinnedColumn();
 	bool IsPinnedColumnActive() const { return bPinnedColumnActive; }
 private:
-	/** Number of actors which have passed through the filters */
+	/** Number of actors (including unloaded) which have passed through the filters */
 	uint32 FilteredActorCount = 0;
+	/** Number of unloaded actors which have passed through all the filters */
+	uint32 FilteredUnloadedActorCount = 0;
+	/** List of unloaded actors which passed through the regular filters and may or may not have passed the search filter */
+	TSet<const FWorldPartitionActorDesc*> ApplicableUnloadedActors;
 	/** List of actors which passed the regular filters and may or may not have passed the search filter */
 	TSet<TWeakObjectPtr<AActor>> ApplicableActors;
 
+	bool bRepresentingWorldPartitionedWorld = false;
 	bool bActorSCCStatusColumnActive = false;
 	bool bPinnedColumnActive = false;
 };
diff --git a/Engine/Source/Editor/SceneOutliner/Public/SSceneOutliner.h b/Engine/Source/Editor/SceneOutliner/Public/SSceneOutliner.h
--- a/Engine/Source/Editor/SceneOutliner/Public/SSceneOutliner.h
+++ b/Engine/Source/Editor/SceneOutliner/Public/SSceneOutliner.h
@@ -117,8 +117,8 @@ struct FSceneOutlinerItemSelection
 		return SelectedItems.Num();
 	}
 
-	/** Returns the number of items of a specific type in the selection */
-	template <typename TreeType>
+	/** Returns the number of items of the specified types in the selection */
+	template <typename ...TreeTypes>
 	uint32 Num() const
 	{
 		uint32 Result = 0;
@@ -126,7 +126,7 @@ struct FSceneOutlinerItemSelection
 		{
 			if (const auto ItemPtr = Item.Pin())
 			{
-				if (ItemPtr->IsA<TreeType>())
+				if ((ItemPtr->IsA<TreeTypes>() || ...))
 				{
 					++Result;
 				}
