Actor's name is now also shown in SceneOutliner for unloaded actors (when selecting ID Name column).

diff --git a/Engine/Source/Editor/SceneOutliner/Private/SceneOutlinerActorInfoColumn.cpp b/Engine/Source/Editor/SceneOutliner/Private/SceneOutlinerActorInfoColumn.cpp
--- a/Engine/Source/Editor/SceneOutliner/Private/SceneOutlinerActorInfoColumn.cpp
+++ b/Engine/Source/Editor/SceneOutliner/Private/SceneOutlinerActorInfoColumn.cpp
@@ -167,6 +167,8 @@ struct FGetInfo
 			{
 			case SceneOutliner::ECustomColumnMode::Class:
 				return ActorDesc->GetActorClass()->GetName();
+			case SceneOutliner::ECustomColumnMode::InternalName:
+				return ActorDesc->GetActorName().ToString();
 			default:
 				return FString();
 			}
