SceneOutliner: Fixed crash when attempting to expand items which haven't been fully removed yet.

diff --git a/Engine/Source/Editor/SceneOutliner/Private/SSceneOutliner.cpp b/Engine/Source/Editor/SceneOutliner/Private/SSceneOutliner.cpp
--- a/Engine/Source/Editor/SceneOutliner/Private/SSceneOutliner.cpp
+++ b/Engine/Source/Editor/SceneOutliner/Private/SSceneOutliner.cpp
@@ -1580,7 +1580,7 @@ void SSceneOutliner::OnItemExpansionChanged(FSceneOutlinerTreeItemPtr TreeItem,
 	for (auto WeakChild : TreeItem->GetChildren())
 	{
 		auto Child = WeakChild.Pin();
-		if (Child->Flags.bIsExpanded)
+		if (Child.IsValid() && Child->Flags.bIsExpanded)
 		{
 			OutlinerTreeView->SetItemExpansion(Child, true);
 		}
diff --git a/Engine/Source/Editor/SceneOutliner/Public/ISceneOutlinerTreeItem.h b/Engine/Source/Editor/SceneOutliner/Public/ISceneOutlinerTreeItem.h
--- a/Engine/Source/Editor/SceneOutliner/Public/ISceneOutlinerTreeItem.h
+++ b/Engine/Source/Editor/SceneOutliner/Public/ISceneOutlinerTreeItem.h
@@ -87,7 +87,7 @@ public:
 		}
 	}
 
-	/** Get this item's children, if any. Although we store as weak pointers, they are guaranteed to be valid. */
+	/** Get this item's children, if any. */
 	FORCEINLINE const TSet<TWeakPtr<ISceneOutlinerTreeItem>>& GetChildren() const
 	{
 		return Children;
