World Outliner: Fixed issue with deleted actors showing up as unloaded actors.

diff --git a/Engine/Source/Editor/SceneOutliner/Private/ActorHierarchy.cpp b/Engine/Source/Editor/SceneOutliner/Private/ActorHierarchy.cpp
--- a/Engine/Source/Editor/SceneOutliner/Private/ActorHierarchy.cpp
+++ b/Engine/Source/Editor/SceneOutliner/Private/ActorHierarchy.cpp
@@ -291,7 +291,7 @@ void FActorHierarchy::CreateWorldChildren(UWorld* World, TArray<FSceneOutlinerTr
 		{
 			FWorldPartitionHelpers::ForEachActorDesc(WorldPartition, [this, WorldPartition, &OutItems](const FWorldPartitionActorDesc* ActorDesc)
                 {
-                    if (ActorDesc != nullptr && !ActorDesc->IsLoaded())
+                    if (ActorDesc != nullptr && !ActorDesc->IsLoaded(true))
                     {
                         if (const FSceneOutlinerTreeItemPtr ActorDescItem = Mode->CreateItemFor<FActorDescTreeItem>(FActorDescTreeItem(ActorDesc->GetGuid(), WorldPartition)))
                         {
@@ -557,7 +557,7 @@ void FActorHierarchy::OnLoadedActorRemoved(AActor& InActor)
 
 void FActorHierarchy::OnActorDescAdded(FWorldPartitionActorDesc* ActorDesc)
 {
-	if (bShowingUnloadedActors && ActorDesc && !ActorDesc->IsLoaded())
+	if (bShowingUnloadedActors && ActorDesc && !ActorDesc->IsLoaded(true))
 	{
 		if (UWorldPartition* WorldPartition = RepresentingWorld->GetWorldPartition())
 		{
diff --git a/Engine/Source/Runtime/Engine/Private/WorldPartition/WorldPartitionActorDesc.cpp b/Engine/Source/Runtime/Engine/Private/WorldPartition/WorldPartitionActorDesc.cpp
--- a/Engine/Source/Runtime/Engine/Private/WorldPartition/WorldPartitionActorDesc.cpp
+++ b/Engine/Source/Runtime/Engine/Private/WorldPartition/WorldPartitionActorDesc.cpp
@@ -220,7 +220,7 @@ UHLODLayer* FWorldPartitionActorDesc::GetHLODLayer() const
 	return HLODLayer.IsNone() ? nullptr : Cast<UHLODLayer>(FSoftObjectPath(HLODLayer).TryLoad());
 }
 
-bool FWorldPartitionActorDesc::IsLoaded() const
+bool FWorldPartitionActorDesc::IsLoaded(bool bEvenIfPendingKill) const
 {
 #if WITH_DEV_AUTOMATION_TESTS
 	if (GIsAutomationTesting)
@@ -229,7 +229,7 @@ bool FWorldPartitionActorDesc::IsLoaded() const
 	}
 #endif
 
-	return ActorPtr.IsValid();
+	return ActorPtr.IsValid(bEvenIfPendingKill);
 }
 
 AActor* FWorldPartitionActorDesc::GetActor(bool bEvenIfPendingKill, bool bEvenIfUnreachable) const
diff --git a/Engine/Source/Runtime/Engine/Public/WorldPartition/WorldPartitionActorDesc.h b/Engine/Source/Runtime/Engine/Public/WorldPartition/WorldPartitionActorDesc.h
--- a/Engine/Source/Runtime/Engine/Public/WorldPartition/WorldPartitionActorDesc.h
+++ b/Engine/Source/Runtime/Engine/Public/WorldPartition/WorldPartitionActorDesc.h
@@ -128,7 +128,7 @@ public:
 
 	FString ToString() const;
 
-	bool IsLoaded() const;
+	bool IsLoaded(bool bEvenIfPendingKill=false) const;
 	AActor* GetActor(bool bEvenIfPendingKill=true, bool bEvenIfUnreachable=false) const;
 	AActor* Load() const;
 	virtual void Unload();
