Add Short Package Name Column

diff --git a/Engine/Source/Editor/SceneOutliner/Private/SceneOutlinerActorInfoColumn.cpp b/Engine/Source/Editor/SceneOutliner/Private/SceneOutlinerActorInfoColumn.cpp
--- a/Engine/Source/Editor/SceneOutliner/Private/SceneOutlinerActorInfoColumn.cpp
+++ b/Engine/Source/Editor/SceneOutliner/Private/SceneOutlinerActorInfoColumn.cpp
@@ -54,6 +54,9 @@ struct FGetInfo
 			case SceneOutliner::ECustomColumnMode::InternalName:
 				return Actor->GetFName().ToString();
 
+			case SceneOutliner::ECustomColumnMode::PackageShortName:
+				return FPackageName::GetShortName(Actor->GetPackage()->GetName());
+
 			case SceneOutliner::ECustomColumnMode::UncachedLights:
 				return FString::Printf(TEXT("%7d"), Actor->GetNumUncachedStaticLightingInteractions());
 
@@ -169,6 +172,9 @@ struct FGetInfo
 				return ActorDesc->GetActorClass()->GetName();
 			case SceneOutliner::ECustomColumnMode::InternalName:
 				return ActorDesc->GetActorName().ToString();
+			case SceneOutliner::ECustomColumnMode::PackageShortName:
+				return FPackageName::GetShortName(ActorDesc->GetActorPackage());
+			}
 			default:
 				return FString();
 			}
@@ -431,6 +437,10 @@ FText FActorInfoColumn::MakeComboText( const SceneOutliner::ECustomColumnMode::T
 		ModeName = LOCTEXT("CustomColumnMode_InternalName", "ID Name");
 		break;
 
+	case SceneOutliner::ECustomColumnMode::PackageShortName:
+		ModeName = LOCTEXT("CustomColumnMode_PackageShortName", "Package Short Name");
+		break;
+
 	case SceneOutliner::ECustomColumnMode::UncachedLights:
 		ModeName = LOCTEXT("CustomColumnMode_UncachedLights", "# Uncached Lights");
 		break;
@@ -482,6 +492,10 @@ FText FActorInfoColumn::MakeComboToolTipText( const SceneOutliner::ECustomColumn
 		ToolTipText = LOCTEXT("CustomColumnModeToolTip_InternalName", "Shows the internal name of the actor (for diagnostics)");
 		break;
 
+	case SceneOutliner::ECustomColumnMode::PackageShortName:
+		ToolTipText = LOCTEXT("CustomColumnModeToolTip_PackageShortName", "Shows the short name of the actor's package (for diagnostics)");
+		break;
+
 	case SceneOutliner::ECustomColumnMode::UncachedLights:
 		ToolTipText = LOCTEXT("CustomColumnModeToolTip_UncachedLights", "Shows the number of uncached static lights (missing in lightmap)");
 		break;
diff --git a/Engine/Source/Editor/SceneOutliner/Private/SceneOutlinerActorInfoColumn.h b/Engine/Source/Editor/SceneOutliner/Private/SceneOutlinerActorInfoColumn.h
--- a/Engine/Source/Editor/SceneOutliner/Private/SceneOutlinerActorInfoColumn.h
+++ b/Engine/Source/Editor/SceneOutliner/Private/SceneOutlinerActorInfoColumn.h
@@ -45,6 +45,9 @@ namespace SceneOutliner
 			/** Actor's internal name (FName) */
 			InternalName,
 
+			/** Actor's package short name */
+			PackageShortName,
+
 			/** Actor's number of uncached lights */
 			UncachedLights,
 
