Fix build

diff --git a/Engine/Source/Editor/SceneOutliner/Private/SceneOutlinerActorInfoColumn.cpp b/Engine/Source/Editor/SceneOutliner/Private/SceneOutlinerActorInfoColumn.cpp
--- a/Engine/Source/Editor/SceneOutliner/Private/SceneOutlinerActorInfoColumn.cpp
+++ b/Engine/Source/Editor/SceneOutliner/Private/SceneOutlinerActorInfoColumn.cpp
@@ -174,7 +174,6 @@ struct FGetInfo
 				return ActorDesc->GetActorName().ToString();
 			case SceneOutliner::ECustomColumnMode::PackageShortName:
 				return FPackageName::GetShortName(ActorDesc->GetActorPackage());
-			}
 			default:
 				return FString();
 			}
