SceneOutliner: Pinning selects actors

diff --git a/Engine/Source/Editor/SceneOutliner/Private/ActorBrowsingMode.cpp b/Engine/Source/Editor/SceneOutliner/Private/ActorBrowsingMode.cpp
--- a/Engine/Source/Editor/SceneOutliner/Private/ActorBrowsingMode.cpp
+++ b/Engine/Source/Editor/SceneOutliner/Private/ActorBrowsingMode.cpp
@@ -1505,21 +1505,37 @@ void FActorBrowsingMode::SelectFoldersDescendants(const TArray<FFolderTreeItem*>
 
 void FActorBrowsingMode::PinItem(const FSceneOutlinerTreeItemPtr& InItem)
 {
+	AActor* PinnedActor = nullptr;
 	if (UWorldPartition* const WorldPartition = RepresentingWorld->GetWorldPartition())
 	{
 		if (const FActorDescTreeItem* const ActorDescTreeItem = InItem->CastTo<FActorDescTreeItem>())
 		{
-			WorldPartition->PinActor(ActorDescTreeItem->GetGuid());
+			PinnedActor = WorldPartition->PinActor(ActorDescTreeItem->GetGuid());
 		}
 		else if (const FActorTreeItem* const ActorTreeItem = InItem->CastTo<FActorTreeItem>())
 		{
 			if (ActorTreeItem->Actor.IsValid())
 			{
-				WorldPartition->PinActor(ActorTreeItem->Actor->GetActorGuid());
+				PinnedActor = WorldPartition->PinActor(ActorTreeItem->Actor->GetActorGuid());
 			}
 		}
 	}
 
+	// Check if we need to start a batch selection
+	bool bIsBatchSelectOwner = false;
+	if (PinnedActor)
+	{
+		const bool bIsBatchSelecting = GEditor->GetSelectedActors()->IsBatchSelecting();
+		if (!bIsBatchSelecting)
+		{
+			bIsBatchSelectOwner = true;
+			GEditor->GetSelectedActors()->BeginBatchSelectOperation();
+			GEditor->SelectNone(/*bNoteSelectionChange=*/false, /*bDeselectBSPSurfs=*/true);
+		}
+
+		GEditor->SelectActor(PinnedActor, /*bInSelected=*/true, /*bNotify=*/false);
+	}
+
 	// Recursively pin all children.
 	for (const auto& Child : InItem->GetChildren())
 	{
@@ -1528,6 +1544,12 @@ void FActorBrowsingMode::PinItem(const FSceneOutlinerTreeItemPtr& InItem)
 			PinItem(Child.Pin());
 	    }
 	}
+
+	// End batch selection if needed
+	if (bIsBatchSelectOwner)
+	{
+		GEditor->GetSelectedActors()->EndBatchSelectOperation(/*bNotify=*/true);
+	}
 }
 
 void FActorBrowsingMode::UnpinItem(const FSceneOutlinerTreeItemPtr& InItem)
@@ -1560,12 +1582,19 @@ void FActorBrowsingMode::UnpinItem(const FSceneOutlinerTreeItemPtr& InItem)
 void FActorBrowsingMode::PinSelectedItems()
 {
 	const FSceneOutlinerItemSelection Selection = SceneOutliner->GetSelection();
-
-	Selection.ForEachItem([this](const FSceneOutlinerTreeItemPtr& TreeItem)
+	if (Selection.Num())
 	{
-		PinItem(TreeItem);
-		return true;
-	});
+		GEditor->GetSelectedActors()->BeginBatchSelectOperation();
+		GEditor->SelectNone(/*bNoteSelectionChange=*/false, /*bDeselectBSPSurfs=*/true);
+
+		Selection.ForEachItem([this](const FSceneOutlinerTreeItemPtr& TreeItem)
+		{
+			PinItem(TreeItem);
+			return true;
+		});
+
+		GEditor->GetSelectedActors()->EndBatchSelectOperation(/*bNotify=*/true);
+	}
 }
 
 void FActorBrowsingMode::UnpinSelectedItems()
diff --git a/Engine/Source/Runtime/Engine/Private/WorldPartition/ActorDescContainer.cpp b/Engine/Source/Runtime/Engine/Private/WorldPartition/ActorDescContainer.cpp
--- a/Engine/Source/Runtime/Engine/Private/WorldPartition/ActorDescContainer.cpp
+++ b/Engine/Source/Runtime/Engine/Private/WorldPartition/ActorDescContainer.cpp
@@ -253,11 +253,11 @@ void UActorDescContainer::RemoveActor(const FGuid& ActorGuid)
 	}
 }
 
-void UActorDescContainer::PinActor(const FGuid& ActorGuid)
+AActor* UActorDescContainer::PinActor(const FGuid& ActorGuid)
 {
-	if (PinnedActors.Contains(ActorGuid))
+	if (FWorldPartitionReference* PinnedActor = PinnedActors.Find(ActorGuid))
 	{
-		return;
+		return (*PinnedActor).IsValid() ? (*PinnedActor)->GetActor() : nullptr;
 	}
 
 	TFunction<void(const FGuid&, TMap<FGuid, FWorldPartitionReference>&)> AddReferences = [this, &AddReferences](const FGuid& ActorGuid, TMap<FGuid, FWorldPartitionReference>& ReferenceMap)
@@ -283,7 +283,7 @@ void UActorDescContainer::PinActor(const FGuid& ActorGuid)
 	{
 		if (const FWorldPartitionActorDesc* const ActorDesc = ActorDescPtr->Get())
 		{
-			PinnedActors.Emplace(ActorGuid, ActorDescPtr);
+			FWorldPartitionReference& PinnedActor = PinnedActors.Emplace(ActorGuid, ActorDescPtr);
 
 			// If the pinned actor has references, we must also create references to those to ensure they are added
 			TMap<FGuid, FWorldPartitionReference>& References = PinnedActorRefs.Emplace(ActorGuid);
@@ -292,8 +292,12 @@ void UActorDescContainer::PinActor(const FGuid& ActorGuid)
 			{
 				AddReferences(ReferencedActor, References);
 			}
+
+			return PinnedActor.IsValid() ? PinnedActor->GetActor() : nullptr;
 		}
 	}
+
+	return nullptr;
 }
 
 void UActorDescContainer::UnpinActor(const FGuid& ActorGuid)
diff --git a/Engine/Source/Runtime/Engine/Public/WorldPartition/ActorDescContainer.h b/Engine/Source/Runtime/Engine/Public/WorldPartition/ActorDescContainer.h
--- a/Engine/Source/Runtime/Engine/Public/WorldPartition/ActorDescContainer.h
+++ b/Engine/Source/Runtime/Engine/Public/WorldPartition/ActorDescContainer.h
@@ -41,7 +41,7 @@ public:
 	/** Removes an actor desc without the need to load a package */
 	virtual void RemoveActor(const FGuid& ActorGuid);
 
-	void PinActor(const FGuid& ActorGuid);
+	AActor* PinActor(const FGuid& ActorGuid);
 	void UnpinActor(const FGuid& ActorGuid);
 	bool IsActorPinned(const FGuid& ActorGuid) const { return PinnedActors.Contains(ActorGuid); }
 public:
