Update SceneOutliner to support that no world is available.

diff --git a/Engine/Source/Editor/SceneOutliner/Private/ActorHierarchy.cpp b/Engine/Source/Editor/SceneOutliner/Private/ActorHierarchy.cpp
--- a/Engine/Source/Editor/SceneOutliner/Private/ActorHierarchy.cpp
+++ b/Engine/Source/Editor/SceneOutliner/Private/ActorHierarchy.cpp
@@ -323,18 +323,17 @@ void FActorHierarchy::CreateWorldChildren(UWorld* World, TArray<FSceneOutlinerTr
 
 void FActorHierarchy::CreateItems(TArray<FSceneOutlinerTreeItemPtr>& OutItems) const
 {
-	check(RepresentingWorld.IsValid());
-
-	if (FSceneOutlinerTreeItemPtr WorldItem = Mode->CreateItemFor<FWorldTreeItem>(RepresentingWorld))
+	if (RepresentingWorld.IsValid())
 	{
-		OutItems.Add(WorldItem);
+		UWorld* RepresentingWorldPtr = RepresentingWorld.Get();
+		check(RepresentingWorldPtr);
+		if (FSceneOutlinerTreeItemPtr WorldItem = Mode->CreateItemFor<FWorldTreeItem>(RepresentingWorldPtr))
+		{
+			OutItems.Add(WorldItem);
+		}
+		// Create world children regardless of if a world item was created
+		CreateWorldChildren(RepresentingWorldPtr, OutItems);
 	}
-
-	UWorld* RepresentingWorldPtr = RepresentingWorld.Get();
-	check(RepresentingWorldPtr);
-
-	// Create world children regardless of if a world item was created
-	CreateWorldChildren(RepresentingWorldPtr, OutItems);
 }
 
 void FActorHierarchy::CreateChildren(const FSceneOutlinerTreeItemPtr& Item, TArray<FSceneOutlinerTreeItemPtr>& OutChildren) const
