SceneOutliner: Fix processing of Move operation with item not in TreeItemMap

diff --git a/Engine/Source/Editor/SceneOutliner/Private/SSceneOutliner.cpp b/Engine/Source/Editor/SceneOutliner/Private/SSceneOutliner.cpp
--- a/Engine/Source/Editor/SceneOutliner/Private/SSceneOutliner.cpp
+++ b/Engine/Source/Editor/SceneOutliner/Private/SSceneOutliner.cpp
@@ -599,7 +599,7 @@ void SSceneOutliner::OnItemMoved(const FSceneOutlinerTreeItemRef& Item)
 		// This will potentially remove any non-matching, empty parents as well
 		RemoveItemFromTree(Item);
 	}
-	else
+	else if(TreeItemMap.Contains(Item->GetID()))
 	{
 		// The item still matches the filters (or has children that do)
 		// When an item has been asked to move, it will still reside under its old parent
