Fix Unpinning Actor crash: Make sure to remove unpinned actors from Selection

diff --git a/Engine/Source/Editor/SceneOutliner/Private/ActorBrowsingMode.cpp b/Engine/Source/Editor/SceneOutliner/Private/ActorBrowsingMode.cpp
--- a/Engine/Source/Editor/SceneOutliner/Private/ActorBrowsingMode.cpp
+++ b/Engine/Source/Editor/SceneOutliner/Private/ActorBrowsingMode.cpp
@@ -1556,6 +1556,13 @@ void FActorBrowsingMode::PinItem(const FSceneOutlinerTreeItemPtr& InItem)
 
 void FActorBrowsingMode::UnpinItem(const FSceneOutlinerTreeItemPtr& InItem)
 {
+	// Check if we need to start a batch selection
+	bool bIsBatchSelectOwner = !GEditor->GetSelectedActors()->IsBatchSelecting();
+	if (bIsBatchSelectOwner)
+	{
+		GEditor->GetSelectedActors()->BeginBatchSelectOperation();
+	}
+	
 	// Recursively unpin all children
 	for (const auto& Child : InItem->GetChildren())
 	{
@@ -1573,12 +1580,18 @@ void FActorBrowsingMode::UnpinItem(const FSceneOutlinerTreeItemPtr& InItem)
 		}
 		else if (const FActorTreeItem* const ActorTreeItem = InItem->CastTo<FActorTreeItem>())
 		{
-			if (ActorTreeItem->Actor.IsValid())
+			if (AActor* PinnedActor = ActorTreeItem->Actor.Get())
 			{
-				WorldPartition->UnpinActor(ActorTreeItem->Actor->GetActorGuid());
+				GEditor->SelectActor(PinnedActor, /*bInSelected=*/false, /*bNotify=*/false);
+				WorldPartition->UnpinActor(PinnedActor->GetActorGuid());
 			}
 		}
 	}
+
+	if (bIsBatchSelectOwner)
+	{
+		GEditor->GetSelectedActors()->EndBatchSelectOperation(/*bNotify=*/true);
+	}
 }
 
 void FActorBrowsingMode::PinSelectedItems()
@@ -1602,12 +1615,18 @@ void FActorBrowsingMode::PinSelectedItems()
 void FActorBrowsingMode::UnpinSelectedItems()
 {
 	const FSceneOutlinerItemSelection Selection = SceneOutliner->GetSelection();
-
-	Selection.ForEachItem([this](const FSceneOutlinerTreeItemPtr& TreeItem)
+	if(Selection.Num())
 	{
-		UnpinItem(TreeItem);
-		return true;
-	});
+		GEditor->GetSelectedActors()->BeginBatchSelectOperation();
+
+		Selection.ForEachItem([this](const FSceneOutlinerTreeItemPtr& TreeItem)
+		{
+			UnpinItem(TreeItem);
+			return true;
+		});
+
+		GEditor->GetSelectedActors()->EndBatchSelectOperation(/*bNotify=*/true);
+	}
 }
 
 FCreateSceneOutlinerMode FActorBrowsingMode::CreateFolderPickerMode() const
