Added the ability to show/hide columns in the Scene Outliner. Refactored FActorInfoColumn to separate out the modes into separate columns so they can exist simultaneously

#okfforgithub public

diff --git a/Engine/Source/Editor/SceneOutliner/Private/ActorBrowsingMode.cpp b/Engine/Source/Editor/SceneOutliner/Private/ActorBrowsingMode.cpp
--- a/Engine/Source/Editor/SceneOutliner/Private/ActorBrowsingMode.cpp
+++ b/Engine/Source/Editor/SceneOutliner/Private/ActorBrowsingMode.cpp
@@ -211,18 +211,10 @@ void FActorBrowsingMode::Rebuild()
 	{
 		UWorldPartition* const WorldPartition = RepresentingWorld->GetWorldPartition();
 		WorldPartition->OnActorDescRemovedEvent.AddRaw(this, &FActorBrowsingMode::OnActorDescRemoved);
-
-		// Enable the pinned column by default on WP worlds
-		if (!bPinnedColumnActive)
-		{
-			TogglePinnedColumn();
-		}
-	}
-	// Disable it by default on non-WP worlds
-	else if (bPinnedColumnActive)
-	{
-		TogglePinnedColumn();
 	}
+
+	// Enable the pinned column by default on WP worlds or Disable it by default on non-WP worlds
+	SceneOutliner->SetColumnVisibility(FSceneOutlinerBuiltInColumnTypes::Pinned(), bRepresentingWorldPartitionedWorld);
 }
 
 FText FActorBrowsingMode::GetStatusText() const 
@@ -286,39 +278,6 @@ FSlateColor FActorBrowsingMode::GetStatusTextColor() const
 
 void FActorBrowsingMode::CreateViewContent(FMenuBuilder& MenuBuilder)
 {
-	MenuBuilder.BeginSection("AssetThumbnails", LOCTEXT("ShowColumnHeading", "Columns"));
-	{
-		// For now hard code this column in.
-		// #todo_Outliner: refactor all info columns out of ActorInfoColumn into toggleable entries of this menu.
-		// could be done with a similar interface to FSceneOutlinerFilterInfo
-		MenuBuilder.AddMenuEntry(
-			LOCTEXT("SourceControlColumnName", "Source Control"),
-			LOCTEXT("SourceControlColumnTooltip", ""),
-			FSlateIcon(),
-			FUIAction(
-				FExecuteAction::CreateRaw(this, &FActorBrowsingMode::ToggleActorSCCStatusColumn),
-				FCanExecuteAction(),
-				FIsActionChecked::CreateRaw(this, &FActorBrowsingMode::IsActorSCCStatusColumnActive)
-			),
-			NAME_None,
-			EUserInterfaceActionType::ToggleButton
-		);
-
-		MenuBuilder.AddMenuEntry(
-			LOCTEXT("PinnedColumnName", "Pinned Column"),
-			LOCTEXT("PinnedColumnToolip", "Displays the pinned state of items"),
-			FSlateIcon(),
-			FUIAction(
-				FExecuteAction::CreateRaw(this, &FActorBrowsingMode::TogglePinnedColumn),
-				FCanExecuteAction(),
-				FIsActionChecked::CreateRaw(this, &FActorBrowsingMode::IsPinnedColumnActive)
-			),
-			NAME_None,
-			EUserInterfaceActionType::ToggleButton
-			);
-	}
-	MenuBuilder.EndSection();
-
 	MenuBuilder.BeginSection("AssetThumbnails", LOCTEXT("ShowWorldHeading", "World"));
 	{
 		MenuBuilder.AddSubMenu(
@@ -1715,31 +1674,4 @@ void FActorBrowsingMode::OnDeleteActorsEnd()
 {
 	SceneOutliner->DeleteFoldersEnd();
 }
-
-void FActorBrowsingMode::ToggleActorSCCStatusColumn()
-{
-	if (bActorSCCStatusColumnActive)
-	{
-		SceneOutliner->RemoveColumn(FSceneOutlinerBuiltInColumnTypes::SourceControl());
-	}
-	else
-	{
-		SceneOutliner->AddColumn(FSceneOutlinerBuiltInColumnTypes::SourceControl(), FSceneOutlinerColumnInfo(ESceneOutlinerColumnVisibility::Visible, 30));
-	}
-	bActorSCCStatusColumnActive = !bActorSCCStatusColumnActive;
-}
-
-void FActorBrowsingMode::TogglePinnedColumn()
-{
-	if (bPinnedColumnActive)
-	{
-		SceneOutliner->RemoveColumn(FSceneOutlinerBuiltInColumnTypes::Pinned());
-	}
-	else
-	{
-		SceneOutliner->AddColumn(FSceneOutlinerBuiltInColumnTypes::Pinned(), FSceneOutlinerColumnInfo(ESceneOutlinerColumnVisibility::Visible, 5));
-	}
-	bPinnedColumnActive = !bPinnedColumnActive;
-}
-
 #undef LOCTEXT_NAMESPACE
\ No newline at end of file
diff --git a/Engine/Source/Editor/SceneOutliner/Private/SSceneOutliner.cpp b/Engine/Source/Editor/SceneOutliner/Private/SSceneOutliner.cpp
--- a/Engine/Source/Editor/SceneOutliner/Private/SSceneOutliner.cpp
+++ b/Engine/Source/Editor/SceneOutliner/Private/SSceneOutliner.cpp
@@ -101,7 +101,9 @@ void SSceneOutliner::Construct(const FArguments& InArgs, const FSceneOutlinerIni
 	HeaderRowWidget =
 		SNew( SHeaderRow )
 			// Only show the list header if the user configured the outliner for that
-			.Visibility(InInitOptions.bShowHeaderRow ? EVisibility::Visible : EVisibility::Collapsed);
+			.Visibility(InInitOptions.bShowHeaderRow ? EVisibility::Visible : EVisibility::Collapsed)
+			.CanSelectGeneratedColumn(true)
+			.OnHiddenColumnsListChanged(this, &SSceneOutliner::HandleHiddenColumnsChanged);
 
 	SetupColumns(*HeaderRowWidget);
 
@@ -265,6 +267,11 @@ void SSceneOutliner::Construct(const FArguments& InArgs, const FSceneOutlinerIni
 	FCoreUObjectDelegates::OnPackageReloaded.AddRaw(this, &SSceneOutliner::OnAssetReloaded);
 }
 
+void SSceneOutliner::HandleHiddenColumnsChanged()
+{
+	//Todo: Save?
+}
+
 void SSceneOutliner::SetupColumns(SHeaderRow& HeaderRow)
 {
 	FSceneOutlinerModule& SceneOutlinerModule = FModuleManager::LoadModuleChecked<FSceneOutlinerModule>("SceneOutliner");
@@ -288,9 +295,11 @@ void SSceneOutliner::SetupColumns(SHeaderRow& HeaderRow)
 
 	for (const FName& ID : SortedIDs)
 	{
+		bool bIsVisible = true;
+
 		if (SharedData->ColumnMap[ID].Visibility == ESceneOutlinerColumnVisibility::Invisible)
 		{
-			continue;
+			bIsVisible = false;
 		}
 
 		TSharedPtr<ISceneOutlinerColumn> Column;
@@ -306,22 +315,28 @@ void SSceneOutliner::SetupColumns(SHeaderRow& HeaderRow)
 
 		if (ensure(Column.IsValid()))
 		{
-			check(Column->GetColumnID() == ID);
-			Columns.Add(Column->GetColumnID(), Column);
+			Columns.Add(ID, Column);
 
 			auto ColumnArgs = Column->ConstructHeaderRowColumn();
 
 			if (Column->SupportsSorting())
 			{
 				ColumnArgs
-					.SortMode(this, &SSceneOutliner::GetColumnSortMode, Column->GetColumnID())
+					.SortMode(this, &SSceneOutliner::GetColumnSortMode, ID)
 					.OnSort(this, &SSceneOutliner::OnColumnSortModeChanged);
 			}
+			
+			ColumnArgs.DefaultLabel(FText::FromName(ID));
+
+			if (!SharedData->ColumnMap[ID].bCanBeHidden)
+			{
+				ColumnArgs.ShouldGenerateWidget(true);
+			}
 
 			HeaderRow.AddColumn(ColumnArgs);
+			HeaderRowWidget->SetShowGeneratedColumn(ID, bIsVisible);
 		}
 	}
-
 	Columns.Shrink();
 	bNeedsColumRefresh = false;
 }
@@ -914,6 +929,14 @@ void SSceneOutliner::RemoveColumn(FName ColumId)
 	}
 }
 
+void SSceneOutliner::SetColumnVisibility(FName ColumnId, bool bIsVisible)
+{
+	if (Columns.Contains(ColumnId))
+	{
+		HeaderRowWidget->SetShowGeneratedColumn(ColumnId, bIsVisible);
+	}
+}
+
 TArray<FName> SSceneOutliner::GetColumnIds() const
 {
 	TArray<FName> ColumnsName;
diff --git a/Engine/Source/Editor/SceneOutliner/Private/SceneOutlinerActorInfoColumn.cpp b/Engine/Source/Editor/SceneOutliner/Private/SceneOutlinerActorInfoColumn.cpp
--- a/Engine/Source/Editor/SceneOutliner/Private/SceneOutlinerActorInfoColumn.cpp
+++ b/Engine/Source/Editor/SceneOutliner/Private/SceneOutlinerActorInfoColumn.cpp
@@ -26,9 +26,6 @@
 
 struct FGetInfo
 {
-	SceneOutliner::ECustomColumnMode::Type CurrentMode;
-
-	FGetInfo(SceneOutliner::ECustomColumnMode::Type InCurrentMode) : CurrentMode(InCurrentMode) {}
 
 	FString operator()(const ISceneOutlinerTreeItem& Item) const
 	{
@@ -40,122 +37,19 @@ struct FGetInfo
 				return FString();
 			}
 
-			switch (CurrentMode)
-			{
-			case SceneOutliner::ECustomColumnMode::Class:
-				return Actor->GetClass()->GetName();
-
-			case SceneOutliner::ECustomColumnMode::Level:
-				return FPackageName::GetShortName(Actor->GetOutermostObject()->GetName());
-
-			case SceneOutliner::ECustomColumnMode::Socket:
-				return Actor->GetAttachParentSocketName().ToString();
-
-			case SceneOutliner::ECustomColumnMode::InternalName:
-				return Actor->GetFName().ToString();
-
-			case SceneOutliner::ECustomColumnMode::PackageShortName:
-				return FPackageName::GetShortName(Actor->GetPackage()->GetName());
-
-			case SceneOutliner::ECustomColumnMode::UncachedLights:
-				return FString::Printf(TEXT("%7d"), Actor->GetNumUncachedStaticLightingInteractions());
-
-			case SceneOutliner::ECustomColumnMode::Layer:
-			{
-				FString Result;
-				for (const auto& Layer : Actor->Layers)
-				{
-					if (Result.Len())
-					{
-						Result += TEXT(", ");
-					}
-
-					Result += Layer.ToString();
-				}
-				return Result;
-			}
-
-			case SceneOutliner::ECustomColumnMode::DataLayer:
-			{
-				TStringBuilder<128> Builder;
-				for (const UDataLayer* DataLayer : Actor->GetDataLayerObjects())
-				{
-					if (Builder.Len())
-					{
-						Builder += TEXT(", ");
-					}
-					Builder += DataLayer->GetDataLayerLabel().ToString();
-				}
-				return Builder.ToString();
-			}
-
-		case SceneOutliner::ECustomColumnMode::Mobility:
-			{
-				FString Result;
-				USceneComponent* RootComponent = Actor->GetRootComponent();
-				if (RootComponent)
-				{
-					if( RootComponent->Mobility == EComponentMobility::Static )
-					{
-						Result = FString(TEXT("Static"));
-					}
-					if (RootComponent->Mobility == EComponentMobility::Stationary)
-					{
-						Result = FString(TEXT("Stationary"));
-					}
-					else if (RootComponent->Mobility == EComponentMobility::Movable)
-					{
-						Result = FString(TEXT("Movable"));
-					}
-				}
-				return Result;
-			}
-
-			default:
-				return FString();
-			}
+			return Actor->GetClass()->GetName();
 		}
 		else if (Item.IsA<FFolderTreeItem>())
 		{
-			switch (CurrentMode)
-			{
-			case SceneOutliner::ECustomColumnMode::Class:
-				return LOCTEXT("FolderTypeName", "Folder").ToString();
-
-			default:
-				return FString();
-			}
+			return LOCTEXT("FolderTypeName", "Folder").ToString();
 		}
 		else if (Item.IsA<FWorldTreeItem>())
 		{
-			switch (CurrentMode)
-			{
-			case SceneOutliner::ECustomColumnMode::Class:
-				return LOCTEXT("WorldTypeName", "World").ToString();
-
-			default:
-				return FString();
-			}
+			return LOCTEXT("WorldTypeName", "World").ToString();
 		}
 		else if (const FComponentTreeItem* ComponentItem = Item.CastTo<FComponentTreeItem>())
 		{
-			UActorComponent* Component = ComponentItem->Component.Get();
-			if (!Component)
-			{
-				return FString();
-			}
-
-			switch (CurrentMode)
-			{
-			case SceneOutliner::ECustomColumnMode::Class:
-				return LOCTEXT("ComponentTypeName", "Component").ToString();
-
-			case SceneOutliner::ECustomColumnMode::InternalName:
-				return Component->GetFName().ToString();
-
-			default:
-				return FString();
-			}
+			return LOCTEXT("ComponentTypeName", "Component").ToString();
 		}
 		else if (const FActorDescTreeItem* ActorDescItem = Item.CastTo<FActorDescTreeItem>())
 		{
@@ -165,46 +59,33 @@ struct FGetInfo
 			{
 				return FString();
 			}
-			
-			switch (CurrentMode)
-			{
-			case SceneOutliner::ECustomColumnMode::Class:
-				return ActorDesc->GetActorClass()->GetName();
-			case SceneOutliner::ECustomColumnMode::InternalName:
-				return ActorDesc->GetActorName().ToString();
-			case SceneOutliner::ECustomColumnMode::PackageShortName:
-				return FPackageName::GetShortName(ActorDesc->GetActorPackage());
-			default:
-				return FString();
-			}
-		}
 
+			return ActorDesc->GetActorClass()->GetName();
+		}
 
 		return FString();
 	}
 };
 
-FActorInfoColumn::FActorInfoColumn( ISceneOutliner& Outliner, SceneOutliner::ECustomColumnMode::Type InDefaultCustomColumnMode )
-	: CurrentMode( InDefaultCustomColumnMode )
-	, SceneOutlinerWeak( StaticCastSharedRef<ISceneOutliner>(Outliner.AsShared()) )
+FTypeInfoColumn::FTypeInfoColumn( ISceneOutliner& Outliner)
+	: SceneOutlinerWeak( StaticCastSharedRef<ISceneOutliner>(Outliner.AsShared()) )
 {
 
 }
 
 
-FName FActorInfoColumn::GetColumnID()
+FName FTypeInfoColumn::GetColumnID()
 {
 	return GetID();
 }
 
 
-SHeaderRow::FColumn::FArguments FActorInfoColumn::ConstructHeaderRowColumn()
+SHeaderRow::FColumn::FArguments FTypeInfoColumn::ConstructHeaderRowColumn()
 {
 	/** Customizable actor data column */
 	return SHeaderRow::Column( GetColumnID() )
 		.FillWidth(2)
 		.HeaderComboVisibility(EHeaderComboVisibility::OnHover)
-		.OnGetMenuContent(this, &FActorInfoColumn::GenerateActorTypeMenu)
 		.HeaderContent()
 		[
 			SNew( SHorizontalBox )
@@ -213,37 +94,12 @@ SHeaderRow::FColumn::FArguments FActorInfoColumn::ConstructHeaderRowColumn()
 			.VAlign(VAlign_Center)
 			[
 				SNew(STextBlock)
-				.Text( this, &FActorInfoColumn::GetSelectedMode )
+				.Text( this, &FTypeInfoColumn::GetSelectedMode )
 			]
 		];
 }
 
-TSharedRef<SWidget> FActorInfoColumn::GenerateActorTypeMenu() 
-{
-	FMenuBuilder MenuBuilder(true, nullptr);
-
-	for(SceneOutliner::ECustomColumnMode::Type CurMode : TEnumRange<SceneOutliner::ECustomColumnMode::Type>() )
-	{
-		FText CurModeText = MakeComboText(CurMode);
-		MenuBuilder.AddMenuEntry(
-			CurModeText,
-			MakeComboToolTipText( CurMode ),
-			FSlateIcon(), 
-			FUIAction(
-				FExecuteAction::CreateSP(this, &FActorInfoColumn::OnModeChanged, CurMode),
-				FCanExecuteAction(),
-				FIsActionChecked::CreateSP(this, &FActorInfoColumn::IsCurrentMode, CurMode)
-			),
-			NAME_None, 
-			EUserInterfaceActionType::RadioButton
-		);
-	}
-
-	return MenuBuilder.MakeWidget();
-
-}
-
-const TSharedRef< SWidget > FActorInfoColumn::ConstructRowWidget( FSceneOutlinerTreeItemRef TreeItem, const STableRow<FSceneOutlinerTreeItemPtr>& Row )
+const TSharedRef< SWidget > FTypeInfoColumn::ConstructRowWidget( FSceneOutlinerTreeItemRef TreeItem, const STableRow<FSceneOutlinerTreeItemPtr>& Row )
 {
 	auto SceneOutliner = SceneOutlinerWeak.Pin();
 	check(SceneOutliner.IsValid());
@@ -251,7 +107,7 @@ const TSharedRef< SWidget > FActorInfoColumn::ConstructRowWidget( FSceneOutliner
 	TSharedRef<SHorizontalBox> HorizontalBox = SNew(SHorizontalBox);
 
 	TSharedRef<STextBlock> MainText = SNew( STextBlock )
-		.Text( this, &FActorInfoColumn::GetTextForItem, TWeakPtr<ISceneOutlinerTreeItem>(TreeItem) ) 
+		.Text( this, &FTypeInfoColumn::GetTextForItem, TWeakPtr<ISceneOutlinerTreeItem>(TreeItem) )
 		.HighlightText( SceneOutliner->GetFilterHighlightText() )
 		.ColorAndOpacity( FSlateColor::UseSubduedForeground() );
 
@@ -267,8 +123,8 @@ const TSharedRef< SWidget > FActorInfoColumn::ConstructRowWidget( FSceneOutliner
 	if (Hyperlink.IsValid())
 	{
 		// If we got a hyperlink, disable hide default text, and show the hyperlink
-		MainText->SetVisibility(TAttribute<EVisibility>::Create(TAttribute<EVisibility>::FGetter::CreateSP(this, &FActorInfoColumn::GetColumnDataVisibility, false)));
-		Hyperlink->SetVisibility(TAttribute<EVisibility>::Create(TAttribute<EVisibility>::FGetter::CreateSP(this, &FActorInfoColumn::GetColumnDataVisibility, true)));
+		MainText->SetVisibility(TAttribute<EVisibility>::Create(TAttribute<EVisibility>::FGetter::CreateSP(this, &FTypeInfoColumn::GetColumnDataVisibility, false)));
+		Hyperlink->SetVisibility(TAttribute<EVisibility>::Create(TAttribute<EVisibility>::FGetter::CreateSP(this, &FTypeInfoColumn::GetColumnDataVisibility, true)));
 
 		HorizontalBox->AddSlot()
 		.VAlign(VAlign_Center)
@@ -291,7 +147,7 @@ const TSharedRef< SWidget > FActorInfoColumn::ConstructRowWidget( FSceneOutliner
 	return HorizontalBox;
 }
 
-TSharedPtr<SWidget> FActorInfoColumn::ConstructClassHyperlink( ISceneOutlinerTreeItem& TreeItem )
+TSharedPtr<SWidget> FTypeInfoColumn::ConstructClassHyperlink( ISceneOutlinerTreeItem& TreeItem )
 {
 	if (const FActorTreeItem* ActorItem = TreeItem.CastTo<FActorTreeItem>())
 	{
@@ -327,20 +183,10 @@ TSharedPtr<SWidget> FActorInfoColumn::ConstructClassHyperlink( ISceneOutlinerTre
 	return nullptr;
 }
 
-void FActorInfoColumn::PopulateSearchStrings( const ISceneOutlinerTreeItem& Item, TArray< FString >& OutSearchStrings ) const
+void FTypeInfoColumn::PopulateSearchStrings( const ISceneOutlinerTreeItem& Item, TArray< FString >& OutSearchStrings ) const
 {
 	{
-		FString String = FGetInfo(CurrentMode)(Item);
-		if (String.Len())
-		{
-			OutSearchStrings.Add(String);
-		}
-	}
-
-	// We always add the class
-	if (CurrentMode != SceneOutliner::ECustomColumnMode::Class)
-	{
-		FString String = FGetInfo(SceneOutliner::ECustomColumnMode::Class)(Item);
+		FString String = FGetInfo()(Item);
 		if (String.Len())
 		{
 			OutSearchStrings.Add(String);
@@ -348,181 +194,46 @@ void FActorInfoColumn::PopulateSearchStrings( const ISceneOutlinerTreeItem& Item
 	}
 }
 
-bool FActorInfoColumn::SupportsSorting() const
+bool FTypeInfoColumn::SupportsSorting() const
 {
-	return CurrentMode != SceneOutliner::ECustomColumnMode::None;
+	return true;
 }
 
-void FActorInfoColumn::SortItems(TArray<FSceneOutlinerTreeItemPtr>& RootItems, const EColumnSortMode::Type SortMode) const
+void FTypeInfoColumn::SortItems(TArray<FSceneOutlinerTreeItemPtr>& RootItems, const EColumnSortMode::Type SortMode) const
 {
 	FSceneOutlinerSortHelper<FString>()
-		.Primary(FGetInfo(CurrentMode), SortMode)
+		.Primary(FGetInfo(), SortMode)
 		.Sort(RootItems);
 }
 
-void FActorInfoColumn::OnModeChanged( const SceneOutliner::ECustomColumnMode::Type NewSelection )
-{
-	CurrentMode = NewSelection;
-
-	// Refresh and refilter the list
-	SceneOutlinerWeak.Pin()->Refresh();
-	FSlateApplication::Get().DismissAllMenus();
-}
-
-bool FActorInfoColumn::IsCurrentMode( const SceneOutliner::ECustomColumnMode::Type Mode ) const 
-{
-	return CurrentMode == Mode;
-}
-
-EVisibility FActorInfoColumn::GetColumnDataVisibility( bool bIsClassHyperlink ) const
+EVisibility FTypeInfoColumn::GetColumnDataVisibility( bool bIsClassHyperlink ) const
 {
-	if ( CurrentMode == SceneOutliner::ECustomColumnMode::Class )
-	{
-		return bIsClassHyperlink ? EVisibility::Visible : EVisibility::Collapsed;
-	}
-	else
-	{
-		return bIsClassHyperlink ? EVisibility::Collapsed : EVisibility::Visible;
-	}
+	return bIsClassHyperlink ? EVisibility::Visible : EVisibility::Collapsed;
 }
 
-FText FActorInfoColumn::GetTextForItem( TWeakPtr<ISceneOutlinerTreeItem> TreeItem ) const
+FText FTypeInfoColumn::GetTextForItem( TWeakPtr<ISceneOutlinerTreeItem> TreeItem ) const
 {
 	auto Item = TreeItem.Pin();
-	return Item.IsValid() ? FText::FromString(FGetInfo(CurrentMode)(*Item)) : FText::GetEmpty();
+	return Item.IsValid() ? FText::FromString(FGetInfo()(*Item)) : FText::GetEmpty();
 }
 
-FText FActorInfoColumn::GetSelectedMode() const
+FText FTypeInfoColumn::GetSelectedMode() const
 {
-	if (CurrentMode == SceneOutliner::ECustomColumnMode::None)
-	{
-		return FText();
-	}
-
-	return MakeComboText(CurrentMode);
+	return MakeComboText();
 }
 
-FText FActorInfoColumn::MakeComboText( const SceneOutliner::ECustomColumnMode::Type& Mode ) const
+FText FTypeInfoColumn::MakeComboText() const
 {
-	FText ModeName;
-
-	switch( Mode )
-	{
-	case SceneOutliner::ECustomColumnMode::None:
-		ModeName = LOCTEXT("CustomColumnMode_None", "None");
-		break;
-
-	case SceneOutliner::ECustomColumnMode::Class:
-		ModeName = LOCTEXT("CustomColumnMode_Class", "Type");
-		break;
-
-	case SceneOutliner::ECustomColumnMode::Level:
-		ModeName = LOCTEXT("CustomColumnMode_Level", "Level");
-		break;
-
-	case SceneOutliner::ECustomColumnMode::Layer:
-		ModeName = LOCTEXT("CustomColumnMode_Layer", "Layer");
-		break;
-
-	case SceneOutliner::ECustomColumnMode::DataLayer:
-		ModeName = LOCTEXT("CustomColumnMode_DataLayer", "Data Layer");
-		break;
-
-	case SceneOutliner::ECustomColumnMode::Socket:
-		ModeName = LOCTEXT("CustomColumnMode_Socket", "Socket");
-		break;
-
-	case SceneOutliner::ECustomColumnMode::InternalName:
-		ModeName = LOCTEXT("CustomColumnMode_InternalName", "ID Name");
-		break;
-
-	case SceneOutliner::ECustomColumnMode::PackageShortName:
-		ModeName = LOCTEXT("CustomColumnMode_PackageShortName", "Package Short Name");
-		break;
-
-	case SceneOutliner::ECustomColumnMode::UncachedLights:
-		ModeName = LOCTEXT("CustomColumnMode_UncachedLights", "# Uncached Lights");
-		break;
-
-	case SceneOutliner::ECustomColumnMode::Mobility:
-		ModeName = LOCTEXT("CustomColumnMode_Mobility", "Mobility");
-		break;
-
-	default:
-		ensure(0);
-		break;
-	}
+	FText ModeName = LOCTEXT("CustomColumnMode_Class", "Type");
 
 	return ModeName;
 }
 
 
-FText FActorInfoColumn::MakeComboToolTipText( const SceneOutliner::ECustomColumnMode::Type& Mode ) const
+FText FTypeInfoColumn::MakeComboToolTipText() const
 {
-	FText ToolTipText;
-
-	switch( Mode )
-	{
-	case SceneOutliner::ECustomColumnMode::None:
-		ToolTipText = LOCTEXT("CustomColumnModeToolTip_None", "Hides all extra actor info");
-		break;
-
-	case SceneOutliner::ECustomColumnMode::Class:
-		ToolTipText = LOCTEXT("CustomColumnModeToolTip_Class", "Displays the name of each actor's type");
-		break;
-
-	case SceneOutliner::ECustomColumnMode::Level:
-		ToolTipText = LOCTEXT("CustomColumnModeToolTip_Level", "Displays the level each actor is in, and allows you to search by level name");
-		break;
-
-	case SceneOutliner::ECustomColumnMode::Layer:
-		ToolTipText = LOCTEXT("CustomColumnModeToolTip_Layer", "Displays the layer each actor is in, and allows you to search by layer name");
-		break;
-
-	case SceneOutliner::ECustomColumnMode::DataLayer:
-		ToolTipText = LOCTEXT("CustomColumnModeToolTip_DataLayer", "Displays the data layers each actor is in, and allows you to search by data layer label");
-		break;
-
-	case SceneOutliner::ECustomColumnMode::Socket:
-		ToolTipText = LOCTEXT("CustomColumnModeToolTip_Socket", "Shows the socket the actor is attached to, and allows you to search by socket name");
-		break;
-
-	case SceneOutliner::ECustomColumnMode::InternalName:
-		ToolTipText = LOCTEXT("CustomColumnModeToolTip_InternalName", "Shows the internal name of the actor (for diagnostics)");
-		break;
-
-	case SceneOutliner::ECustomColumnMode::PackageShortName:
-		ToolTipText = LOCTEXT("CustomColumnModeToolTip_PackageShortName", "Shows the short name of the actor's package (for diagnostics)");
-		break;
-
-	case SceneOutliner::ECustomColumnMode::UncachedLights:
-		ToolTipText = LOCTEXT("CustomColumnModeToolTip_UncachedLights", "Shows the number of uncached static lights (missing in lightmap)");
-		break;
-
-	case SceneOutliner::ECustomColumnMode::Mobility:
-		ToolTipText = LOCTEXT("CustomColumnModeToolTip_Mobility", "Shows the mobility of each actor");
-		break;
-
-
-	default:
-		ensure(0);
-		break;
-	}
-
+	FText ToolTipText = LOCTEXT("CustomColumnModeToolTip_Class", "Displays the name of each actor's type");
 	return ToolTipText;
 }
 
-
-TSharedRef< ITableRow > FActorInfoColumn::MakeComboButtonItemWidget( TSharedPtr< SceneOutliner::ECustomColumnMode::Type > Mode, const TSharedRef<STableViewBase>& Owner )
-{
-	return
-		SNew( STableRow< TSharedPtr<SceneOutliner::ECustomColumnMode::Type> >, Owner )
-		.Style(FAppStyle::Get(), "ComboBox.Row")
-		[
-			SNew( STextBlock )
-			.Text( MakeComboText( *Mode ) )
-			.ToolTipText( MakeComboToolTipText( *Mode ) )
-		];
-}
-
 #undef LOCTEXT_NAMESPACE
diff --git a/Engine/Source/Editor/SceneOutliner/Private/SceneOutlinerActorInfoColumn.h b/Engine/Source/Editor/SceneOutliner/Private/SceneOutlinerActorInfoColumn.h
--- a/Engine/Source/Editor/SceneOutliner/Private/SceneOutlinerActorInfoColumn.h
+++ b/Engine/Source/Editor/SceneOutliner/Private/SceneOutlinerActorInfoColumn.h
@@ -14,55 +14,12 @@
 
 class ISceneOutliner;
 
-namespace SceneOutliner
-{
-	/** Types of actor data we can display in a 'custom' tree column */
-	namespace ECustomColumnMode
-	{
-		enum Type
-		{
-			/** Empty column -- doesn't display anything */
-			None = 0,
-
-			/** Class */
-			Class,
-
-			/** Mobility */
-			Mobility,
-
-			/** Level */
-			Level,
-
-			/** Layer */
-			Layer,
-
-			/** Data Layer */
-			DataLayer,
-
-			/** The socket the actor is attached to. */
-			Socket,
-
-			/** Actor's internal name (FName) */
-			InternalName,
-
-			/** Actor's package short name */
-			PackageShortName,
-
-			/** Actor's number of uncached lights */
-			UncachedLights,
-
-			// ---
-
-			/** Number of options */
-			Count
-		};
-	}
-}
-
 /**
- * A custom column for the SceneOutliner which is capable of displaying a variety of Actor details
+ * A custom column for the SceneOutliner which is used to display the Class/Type of Actors
+   Note: Used to have support for displaying a variety of details related to Actors, that functionality
+   has been moved FTextInfoColumn (See SceneOutlinerModule::CreateActorInfoColumns)
  */
-class FActorInfoColumn : public ISceneOutlinerColumn
+class FTypeInfoColumn : public ISceneOutlinerColumn
 {
 
 public:
@@ -70,9 +27,9 @@ public:
 	/**
 	 *	Constructor
 	 */
-	FActorInfoColumn( ISceneOutliner& Outliner, SceneOutliner::ECustomColumnMode::Type InDefaultCustomColumnMode = SceneOutliner::ECustomColumnMode::Class );
+	FTypeInfoColumn( ISceneOutliner& Outliner);
 
-	virtual ~FActorInfoColumn() {}
+	virtual ~FTypeInfoColumn() {}
 
 	static FName GetID() { return FSceneOutlinerBuiltInColumnTypes::ActorInfo(); }
 
@@ -100,32 +57,16 @@ private:
 
 	TSharedPtr<SWidget> ConstructClassHyperlink( ISceneOutlinerTreeItem& TreeItem );
 
-	TSharedRef<SWidget> GenerateActorTypeMenu();
-
-	void OnModeChanged( const SceneOutliner::ECustomColumnMode::Type NewSelection );
-
-	bool IsCurrentMode( const SceneOutliner::ECustomColumnMode::Type Mode) const;
-
 	EVisibility GetColumnDataVisibility( bool bIsClassHyperlink ) const;
 
-	FText MakeComboText( const SceneOutliner::ECustomColumnMode::Type& Mode ) const;
-
-	FText MakeComboToolTipText( const SceneOutliner::ECustomColumnMode::Type& Mode ) const;
+	FText MakeComboText( ) const;
 
-	TSharedRef< ITableRow > MakeComboButtonItemWidget( TSharedPtr< SceneOutliner::ECustomColumnMode::Type > Mode, const TSharedRef<STableViewBase> & );
+	FText MakeComboToolTipText( ) const;
 
 	FText GetSelectedMode() const;
 
 private:
 
-	/** 
-	 * Current custom column mode.  This is used for displaying a bit of extra data about the actors, as well as
-	 * allowing the user to search by additional criteria 
-	 */
-	mutable SceneOutliner::ECustomColumnMode::Type CurrentMode;
-
 	/** Weak reference to the outliner widget that owns our list */
 	TWeakPtr< ISceneOutliner > SceneOutlinerWeak;
-};
-
-ENUM_RANGE_BY_COUNT(SceneOutliner::ECustomColumnMode::Type, SceneOutliner::ECustomColumnMode::Count)
+};
\ No newline at end of file
diff --git a/Engine/Source/Editor/SceneOutliner/Private/SceneOutlinerModule.cpp b/Engine/Source/Editor/SceneOutliner/Private/SceneOutlinerModule.cpp
--- a/Engine/Source/Editor/SceneOutliner/Private/SceneOutlinerModule.cpp
+++ b/Engine/Source/Editor/SceneOutliner/Private/SceneOutlinerModule.cpp
@@ -11,23 +11,29 @@
 #include "SceneOutlinerItemLabelColumn.h"
 #include "SceneOutlinerActorSCCColumn.h"
 #include "SceneOutlinerPinnedColumn.h"
+#include "SceneOutlinerTextInfoColumn.h"
 #include "SceneOutlinerPublicTypes.h"
 
 #include "ActorPickingMode.h"
 #include "ActorBrowsingMode.h"
 #include "ActorTreeItem.h"
 #include "ComponentTreeItem.h"
+#include "ActorDescTreeItem.h"
+#include "FolderTreeItem.h"
+#include "WorldTreeItem.h"
+#include "WorldPartition/DataLayer/DataLayer.h"
+#include "WorldPartition/DataLayer/WorldDataLayers.h"
 
 /* FSceneOutlinerModule interface
  *****************************************************************************/
 
 void FSceneOutlinerModule::StartupModule()
 {
-	RegisterDefaultColumnType< FSceneOutlinerItemLabelColumn >(FSceneOutlinerColumnInfo(ESceneOutlinerColumnVisibility::Visible, 10));
+	RegisterDefaultColumnType< FSceneOutlinerItemLabelColumn >(FSceneOutlinerColumnInfo(ESceneOutlinerColumnVisibility::Visible, 10, FCreateSceneOutlinerColumn(), false));
 
 	// Register builtin column types which are not active by default
 	RegisterColumnType<FSceneOutlinerGutter>();
-	RegisterColumnType<FActorInfoColumn>();
+	RegisterColumnType<FTypeInfoColumn>();
 	RegisterColumnType<FSceneOutlinerActorSCCColumn>();
 	RegisterColumnType<FSceneOutlinerPinnedColumn>();
 }
@@ -37,7 +43,7 @@ void FSceneOutlinerModule::ShutdownModule()
 {
 	UnRegisterColumnType<FSceneOutlinerGutter>();
 	UnRegisterColumnType<FSceneOutlinerItemLabelColumn>();
-	UnRegisterColumnType<FActorInfoColumn>();
+	UnRegisterColumnType<FTypeInfoColumn>();
 	UnRegisterColumnType<FSceneOutlinerActorSCCColumn>();
 	UnRegisterColumnType<FSceneOutlinerPinnedColumn>();
 }
@@ -79,8 +85,10 @@ TSharedRef<ISceneOutliner> FSceneOutlinerModule::CreateActorPicker(const FSceneO
 	InitOptions.ModeFactory = ModeFactory;
 	if (InitOptions.ColumnMap.Num() == 0)
 	{
-		InitOptions.ColumnMap.Add(FSceneOutlinerBuiltInColumnTypes::Label(), FSceneOutlinerColumnInfo(ESceneOutlinerColumnVisibility::Visible, 0));
+		InitOptions.ColumnMap.Add(FSceneOutlinerBuiltInColumnTypes::Label(), FSceneOutlinerColumnInfo(ESceneOutlinerColumnVisibility::Visible, 0, FCreateSceneOutlinerColumn(), false));
 		InitOptions.ColumnMap.Add(FSceneOutlinerBuiltInColumnTypes::ActorInfo(), FSceneOutlinerColumnInfo(ESceneOutlinerColumnVisibility::Visible, 10));
+		InitOptions.ColumnMap.Add(FSceneOutlinerBuiltInColumnTypes::Pinned(), FSceneOutlinerColumnInfo(ESceneOutlinerColumnVisibility::Invisible, 5));
+		CreateActorInfoColumns(InitOptions);
 	}
 	return CreateSceneOutliner(InitOptions);
 }
@@ -114,8 +122,9 @@ TSharedRef<ISceneOutliner> FSceneOutlinerModule::CreateComponentPicker(const FSc
 	InitOptions.ModeFactory = ModeFactory;
 	if (InitOptions.ColumnMap.Num() == 0)
 	{
-		InitOptions.ColumnMap.Add(FSceneOutlinerBuiltInColumnTypes::Label(), FSceneOutlinerColumnInfo(ESceneOutlinerColumnVisibility::Visible, 0));
+		InitOptions.ColumnMap.Add(FSceneOutlinerBuiltInColumnTypes::Label(), FSceneOutlinerColumnInfo(ESceneOutlinerColumnVisibility::Visible, 0, FCreateSceneOutlinerColumn(), false));
 		InitOptions.ColumnMap.Add(FSceneOutlinerBuiltInColumnTypes::ActorInfo(), FSceneOutlinerColumnInfo(ESceneOutlinerColumnVisibility::Visible, 10));
+		CreateActorInfoColumns(InitOptions);
 	}
 	return CreateSceneOutliner(InitOptions);
 }
@@ -134,8 +143,219 @@ TSharedRef< ISceneOutliner > FSceneOutlinerModule::CreateActorBrowser(const FSce
 		InitOptions.UseDefaultColumns();
 		InitOptions.ColumnMap.Add(FSceneOutlinerBuiltInColumnTypes::Gutter(), FSceneOutlinerColumnInfo(ESceneOutlinerColumnVisibility::Visible, 0));
 		InitOptions.ColumnMap.Add(FSceneOutlinerBuiltInColumnTypes::ActorInfo(), FSceneOutlinerColumnInfo(ESceneOutlinerColumnVisibility::Visible, 20));
+		InitOptions.ColumnMap.Add(FSceneOutlinerBuiltInColumnTypes::SourceControl(), FSceneOutlinerColumnInfo(ESceneOutlinerColumnVisibility::Invisible, 30));
+		InitOptions.ColumnMap.Add(FSceneOutlinerBuiltInColumnTypes::Pinned(), FSceneOutlinerColumnInfo(ESceneOutlinerColumnVisibility::Invisible, 5));
+		CreateActorInfoColumns(InitOptions);
 	}
 	return CreateSceneOutliner(InitOptions);
 }
 
+void FSceneOutlinerModule::CreateActorInfoColumns(FSceneOutlinerInitializationOptions& InInitOptions) const
+{
+	FGetTextForItem MobilityInfoText = FGetTextForItem::CreateLambda([](const ISceneOutlinerTreeItem& Item) -> FString
+	{
+		if (const FActorTreeItem* ActorItem = Item.CastTo<FActorTreeItem>())
+		{
+			AActor* Actor = ActorItem->Actor.Get();
+			if (!Actor)
+			{
+				return FString();
+			}
+
+			FString Result;
+			USceneComponent* RootComponent = Actor->GetRootComponent();
+			if (RootComponent)
+			{
+				if (RootComponent->Mobility == EComponentMobility::Static)
+				{
+					Result = FString(TEXT("Static"));
+				}
+				if (RootComponent->Mobility == EComponentMobility::Stationary)
+				{
+					Result = FString(TEXT("Stationary"));
+				}
+				else if (RootComponent->Mobility == EComponentMobility::Movable)
+				{
+					Result = FString(TEXT("Movable"));
+				}
+			}
+			return Result;
+		}
+		return FString();
+	});
+
+	FGetTextForItem LevelInfoText = FGetTextForItem::CreateLambda([](const ISceneOutlinerTreeItem& Item) -> FString
+	{
+		if (const FActorTreeItem* ActorItem = Item.CastTo<FActorTreeItem>())
+		{
+			AActor* Actor = ActorItem->Actor.Get();
+			if (!Actor)
+			{
+				return FString();
+			}
+
+			return FPackageName::GetShortName(Actor->GetOutermostObject()->GetName());
+		}
+		return FString();
+	});
+
+	FGetTextForItem LayerInfoText = FGetTextForItem::CreateLambda([](const ISceneOutlinerTreeItem& Item) -> FString
+	{
+		FString Result;
+		
+		if (const FActorTreeItem* ActorItem = Item.CastTo<FActorTreeItem>())
+		{
+			AActor* Actor = ActorItem->Actor.Get();
+			if (!Actor)
+			{
+				return FString();
+			}
+
+			for (const auto& Layer : Actor->Layers)
+			{
+				if (Result.Len())
+				{
+					Result += TEXT(", ");
+				}
+
+				Result += Layer.ToString();
+			}
+		}
+
+		return Result;
+	});
+
+	FGetTextForItem DataLayerInfoText = FGetTextForItem::CreateLambda([](const ISceneOutlinerTreeItem& Item) -> FString
+	{
+		TStringBuilder<128> Builder;
+
+		if (const FActorTreeItem* ActorItem = Item.CastTo<FActorTreeItem>())
+		{
+			AActor* Actor = ActorItem->Actor.Get();
+
+			if (!Actor)
+			{
+				return FString();
+			}
+
+			for (const UDataLayer* DataLayer : Actor->GetDataLayerObjects())
+			{
+				if (Builder.Len())
+				{
+					Builder += TEXT(", ");
+				}
+				Builder += DataLayer->GetDataLayerLabel().ToString();
+			}
+		}
+		
+		return Builder.ToString();
+	});
+
+	FGetTextForItem SocketInfoText = FGetTextForItem::CreateLambda([](const ISceneOutlinerTreeItem& Item) -> FString
+	{
+		if (const FActorTreeItem* ActorItem = Item.CastTo<FActorTreeItem>())
+		{
+			AActor* Actor = ActorItem->Actor.Get();
+
+			if (!Actor)
+			{
+				return FString();
+			}
+
+			return Actor->GetAttachParentSocketName().ToString();
+		}
+		
+		return FString();
+	});
+
+	FGetTextForItem InternalNameInfoText = FGetTextForItem::CreateLambda([](const ISceneOutlinerTreeItem& Item) -> FString
+	{
+		if (const FActorTreeItem* ActorItem = Item.CastTo<FActorTreeItem>())
+		{
+			AActor* Actor = ActorItem->Actor.Get();
+			if (!Actor)
+			{
+				return FString();
+			}
+
+			return Actor->GetFName().ToString();
+		}
+		else if (const FComponentTreeItem* ComponentItem = Item.CastTo<FComponentTreeItem>())
+		{
+			UActorComponent* Component = ComponentItem->Component.Get();
+			if (!Component)
+			{
+				return FString();
+			}
+
+			return Component->GetFName().ToString();
+
+		}
+		else if (const FActorDescTreeItem* ActorDescItem = Item.CastTo<FActorDescTreeItem>())
+		{
+			const FWorldPartitionActorDesc* ActorDesc = ActorDescItem->ActorDescHandle.GetActorDesc();
+
+			if (!ActorDesc)
+			{
+				return FString();
+			}
+
+			return ActorDesc->GetActorName().ToString();
+		}
+
+		return FString();
+	});
+
+	FGetTextForItem PackageShortNameInfoText = FGetTextForItem::CreateLambda([](const ISceneOutlinerTreeItem& Item) -> FString
+	{
+		if (const FActorTreeItem* ActorItem = Item.CastTo<FActorTreeItem>())
+		{
+			AActor* Actor = ActorItem->Actor.Get();
+			if (!Actor)
+			{
+				return FString();
+			}
+
+			return FPackageName::GetShortName(Actor->GetPackage()->GetName());
+		}
+		else if (const FActorDescTreeItem* ActorDescItem = Item.CastTo<FActorDescTreeItem>())
+		{
+			const FWorldPartitionActorDesc* ActorDesc = ActorDescItem->ActorDescHandle.GetActorDesc();
+
+			if (!ActorDesc)
+			{
+				return FString();
+			}
+
+			return FPackageName::GetShortName(ActorDesc->GetActorPackage());
+		}
+
+		return FString();
+	});
+
+	FGetTextForItem UncachedLightsInfoText = FGetTextForItem::CreateLambda([](const ISceneOutlinerTreeItem& Item) -> FString
+	{
+		if (const FActorTreeItem* ActorItem = Item.CastTo<FActorTreeItem>())
+		{
+			AActor* Actor = ActorItem->Actor.Get();
+			if (!Actor)
+			{
+				return FString();
+			}
+
+			return FString::Printf(TEXT("%7d"), Actor->GetNumUncachedStaticLightingInteractions());
+		}
+		return FString();
+	});
+
+	InInitOptions.ColumnMap.Add("Mobility", FSceneOutlinerColumnInfo(ESceneOutlinerColumnVisibility::Invisible, 20, FCreateSceneOutlinerColumn::CreateStatic(&FTextInfoColumn::CreateTextInfoColumn, FName("Mobility"), MobilityInfoText)));
+	InInitOptions.ColumnMap.Add("Layer", FSceneOutlinerColumnInfo(ESceneOutlinerColumnVisibility::Invisible, 20, FCreateSceneOutlinerColumn::CreateStatic(&FTextInfoColumn::CreateTextInfoColumn, FName("Layer"), LayerInfoText)));
+	InInitOptions.ColumnMap.Add("Data Layer", FSceneOutlinerColumnInfo(ESceneOutlinerColumnVisibility::Invisible, 20, FCreateSceneOutlinerColumn::CreateStatic(&FTextInfoColumn::CreateTextInfoColumn, FName("Data Layer"), DataLayerInfoText)));
+	InInitOptions.ColumnMap.Add("Socket", FSceneOutlinerColumnInfo(ESceneOutlinerColumnVisibility::Invisible, 20, FCreateSceneOutlinerColumn::CreateStatic(&FTextInfoColumn::CreateTextInfoColumn, FName("Socket"), SocketInfoText)));
+	InInitOptions.ColumnMap.Add("ID Name", FSceneOutlinerColumnInfo(ESceneOutlinerColumnVisibility::Invisible, 20, FCreateSceneOutlinerColumn::CreateStatic(&FTextInfoColumn::CreateTextInfoColumn, FName("ID Name"), InternalNameInfoText)));
+	InInitOptions.ColumnMap.Add("Package Short Name", FSceneOutlinerColumnInfo(ESceneOutlinerColumnVisibility::Invisible, 20, FCreateSceneOutlinerColumn::CreateStatic(&FTextInfoColumn::CreateTextInfoColumn, FName("Package Short Name"), PackageShortNameInfoText)));
+	InInitOptions.ColumnMap.Add("# Uncached Lights", FSceneOutlinerColumnInfo(ESceneOutlinerColumnVisibility::Invisible, 20, FCreateSceneOutlinerColumn::CreateStatic(&FTextInfoColumn::CreateTextInfoColumn, FName("# Uncached Lights"), UncachedLightsInfoText)));
+
+}
+
 IMPLEMENT_MODULE(FSceneOutlinerModule, SceneOutliner);
diff --git a/Engine/Source/Editor/SceneOutliner/Private/SceneOutlinerTextInfoColumn.cpp b/Engine/Source/Editor/SceneOutliner/Private/SceneOutlinerTextInfoColumn.cpp
new file mode 100644
--- /dev/null
+++ b/Engine/Source/Editor/SceneOutliner/Private/SceneOutlinerTextInfoColumn.cpp
@@ -0,0 +1,88 @@
+// Copyright Epic Games, Inc. All Rights Reserved.
+
+#include "SceneOutlinerTextInfoColumn.h"
+#include "ISceneOutliner.h"
+#include "SceneOutlinerPublicTypes.h"
+#include "SortHelper.h"
+#include "Widgets/Views/SHeaderRow.h"
+
+#define LOCTEXT_NAMESPACE "SceneOutlinerTextInfoColumn"
+
+
+TSharedRef<ISceneOutlinerColumn> FTextInfoColumn::CreateTextInfoColumn(ISceneOutliner& Outliner, const FName InColumnName, const FGetTextForItem InGetTextForItem)
+{
+	return TSharedRef<ISceneOutlinerColumn>(MakeShareable(new FTextInfoColumn(Outliner, InColumnName, InGetTextForItem)));
+}
+
+FTextInfoColumn::FTextInfoColumn(ISceneOutliner& Outliner, const FName InColumnName, const FGetTextForItem& InGetTextForItem):
+	ColumnName(InColumnName),
+	SceneOutlinerWeak(StaticCastSharedRef<ISceneOutliner>(Outliner.AsShared())),
+	GetTextForItem(InGetTextForItem)
+{
+
+}
+
+FName FTextInfoColumn::GetColumnID()
+{
+	return ColumnName;
+}
+
+SHeaderRow::FColumn::FArguments FTextInfoColumn::ConstructHeaderRowColumn()
+{
+	return SHeaderRow::Column(ColumnName)
+		.FillWidth(2)
+		.HeaderComboVisibility(EHeaderComboVisibility::OnHover);
+}
+
+const TSharedRef< SWidget > FTextInfoColumn::ConstructRowWidget(FSceneOutlinerTreeItemRef TreeItem, const STableRow<FSceneOutlinerTreeItemPtr>& Row)
+{
+	auto SceneOutliner = SceneOutlinerWeak.Pin();
+	check(SceneOutliner.IsValid());
+
+	TSharedRef<SHorizontalBox> HorizontalBox = SNew(SHorizontalBox);
+
+	TSharedRef<STextBlock> MainText = SNew(STextBlock)
+		.Text(this, &FTextInfoColumn::GetInfoForItem, TWeakPtr<ISceneOutlinerTreeItem>(TreeItem))
+		.HighlightText(SceneOutliner->GetFilterHighlightText())
+		.ColorAndOpacity(FSlateColor::UseSubduedForeground());
+
+	HorizontalBox->AddSlot()
+		.AutoWidth()
+		.VAlign(VAlign_Center)
+		.Padding(8, 0, 0, 0)
+		[
+			MainText
+		];
+	
+	return HorizontalBox;
+}
+
+void FTextInfoColumn::PopulateSearchStrings(const ISceneOutlinerTreeItem& Item, TArray< FString >& OutSearchStrings) const
+{
+	FString String = GetTextForItem.Execute(Item);
+	if (String.Len())
+	{
+		OutSearchStrings.Add(String);
+	}
+}
+
+bool FTextInfoColumn::SupportsSorting() const
+{
+	return true; // Text Columns are always sortable
+}
+
+void FTextInfoColumn::SortItems(TArray<FSceneOutlinerTreeItemPtr>& RootItems, const EColumnSortMode::Type SortMode) const
+{
+	FSceneOutlinerSortHelper<FString>()
+		.Primary([this](const ISceneOutlinerTreeItem& Item) { return GetTextForItem.Execute(Item); }, SortMode)
+		.Sort(RootItems);
+}
+
+FText FTextInfoColumn::GetInfoForItem(TWeakPtr<ISceneOutlinerTreeItem> TreeItem) const
+{
+	auto Item = TreeItem.Pin();
+	return Item.IsValid() ? FText::FromString(GetTextForItem.Execute(*Item)) : FText::GetEmpty();
+}
+
+
+#undef LOCTEXT_NAMESPACE
diff --git a/Engine/Source/Editor/SceneOutliner/Private/SceneOutlinerTextInfoColumn.h b/Engine/Source/Editor/SceneOutliner/Private/SceneOutlinerTextInfoColumn.h
new file mode 100644
--- /dev/null
+++ b/Engine/Source/Editor/SceneOutliner/Private/SceneOutlinerTextInfoColumn.h
@@ -0,0 +1,70 @@
+// Copyright Epic Games, Inc. All Rights Reserved.
+
+#pragma once
+
+#include "CoreMinimal.h"
+#include "Layout/Visibility.h"
+#include "Widgets/SWidget.h"
+#include "Widgets/Views/STableViewBase.h"
+#include "Widgets/Views/SHeaderRow.h"
+#include "Widgets/Views/STableRow.h"
+#include "SceneOutlinerPublicTypes.h"
+#include "ISceneOutlinerColumn.h"
+#include "SceneOutlinerActorInfoColumn.h"
+#include "Misc/CoreDelegates.h"
+
+class ISceneOutliner;
+
+/** A delegate used to factory a new TextInfo Column */
+DECLARE_DELEGATE_RetVal_OneParam(FString, FGetTextForItem, const ISceneOutlinerTreeItem&);
+
+
+/**
+* A custom column for the SceneOutliner which is capable of displaying any text based information related to items
+* Note: Only has support for text, any other requirements should have a custom column made
+*/
+class FTextInfoColumn : public ISceneOutlinerColumn
+{
+
+	public:
+
+		/**
+		 *	Constructor
+		 */
+		FTextInfoColumn(ISceneOutliner& Outliner, const FName InColumnName, const FGetTextForItem& InGetTextForItem);
+
+		virtual ~FTextInfoColumn() {}
+
+		// Factory function to create a TextInfoColumn (This column is not registered by default because of it's requirements and therefore requires a factory!)
+		static TSharedRef<ISceneOutlinerColumn> CreateTextInfoColumn(ISceneOutliner& Outliner, const FName InColumnName, const FGetTextForItem InGetTextForItem);
+
+		//////////////////////////////////////////////////////////////////////////
+		// Begin ISceneOutlinerColumn Implementation
+
+		virtual FName GetColumnID() override;
+
+		virtual SHeaderRow::FColumn::FArguments ConstructHeaderRowColumn() override;
+
+		virtual const TSharedRef< SWidget > ConstructRowWidget(FSceneOutlinerTreeItemRef TreeItem, const STableRow<FSceneOutlinerTreeItemPtr>& Row) override;
+
+		virtual void PopulateSearchStrings(const ISceneOutlinerTreeItem& Item, TArray< FString >& OutSearchStrings) const override;
+
+		virtual bool SupportsSorting() const override;
+
+		virtual void SortItems(TArray<FSceneOutlinerTreeItemPtr>& RootItems, const EColumnSortMode::Type SortMode) const override;
+
+		// End ISceneOutlinerColumn Implementation
+		//////////////////////////////////////////////////////////////////////////
+
+	private:
+
+		FText GetInfoForItem(TWeakPtr<ISceneOutlinerTreeItem> TreeItem) const;
+
+		FName ColumnName;
+
+		/** Weak reference to the outliner widget that owns our list */
+		TWeakPtr< ISceneOutliner > SceneOutlinerWeak;
+
+		/* Delegate Called to Get Display Text for an item in the column*/
+		FGetTextForItem GetTextForItem;
+};
diff --git a/Engine/Source/Editor/SceneOutliner/Public/ActorBrowsingMode.h b/Engine/Source/Editor/SceneOutliner/Public/ActorBrowsingMode.h
--- a/Engine/Source/Editor/SceneOutliner/Public/ActorBrowsingMode.h
+++ b/Engine/Source/Editor/SceneOutliner/Public/ActorBrowsingMode.h
@@ -109,11 +109,6 @@ private:
 	static TSharedRef<FSceneOutlinerFilter> CreateHideLevelInstancesFilter();
 	static TSharedRef<FSceneOutlinerFilter> CreateHideUnloadedActorsFilter();
 	
-	void ToggleActorSCCStatusColumn();
-	bool IsActorSCCStatusColumnActive() const { return bActorSCCStatusColumnActive; }
-
-	void TogglePinnedColumn();
-	bool IsPinnedColumnActive() const { return bPinnedColumnActive; }
 private:
 	/** Number of actors (including unloaded) which have passed through the filters */
 	uint32 FilteredActorCount = 0;
@@ -125,6 +120,4 @@ private:
 	TSet<TWeakObjectPtr<AActor>> ApplicableActors;
 
 	bool bRepresentingWorldPartitionedWorld = false;
-	bool bActorSCCStatusColumnActive = false;
-	bool bPinnedColumnActive = false;
 };
diff --git a/Engine/Source/Editor/SceneOutliner/Public/SSceneOutliner.h b/Engine/Source/Editor/SceneOutliner/Public/SSceneOutliner.h
--- a/Engine/Source/Editor/SceneOutliner/Public/SSceneOutliner.h
+++ b/Engine/Source/Editor/SceneOutliner/Public/SSceneOutliner.h
@@ -313,6 +313,8 @@ public:
 		*/
 	virtual void RemoveColumn(FName ColumId) override;
 
+	void SetColumnVisibility(FName ColumnId, bool bIsVisible);
+
 	/** Return the name/Id of the columns of the scene outliner */
 	virtual TArray<FName> GetColumnIds() const override;
 
@@ -543,9 +545,13 @@ private:
 	/** Map of columns that are shown on this outliner. */
 	TMap<FName, TSharedPtr<ISceneOutlinerColumn>> Columns;
 
+	TArray<FName> HiddenColumnsList;
+
 	/** Set up the columns required for this outliner */
 	void SetupColumns(SHeaderRow& HeaderRow);
 
+	void HandleHiddenColumnsChanged();
+
 	/** Refresh the scene outliner for when a colum was added or removed */
 	void RefreshColums();
 
diff --git a/Engine/Source/Editor/SceneOutliner/Public/SceneOutlinerModule.h b/Engine/Source/Editor/SceneOutliner/Public/SceneOutlinerModule.h
--- a/Engine/Source/Editor/SceneOutliner/Public/SceneOutlinerModule.h
+++ b/Engine/Source/Editor/SceneOutliner/Public/SceneOutlinerModule.h
@@ -9,6 +9,7 @@
 
 class ICustomSceneOutliner;
 class ISceneOutlinerColumn;
+struct FSceneOutlinerInitializationOptions;
 
 /** Delegate used with the Scene Outliner in 'actor picking' mode.  You'll bind a delegate when the
 	outliner widget is created, which will be fired off when an actor is selected in the list */
@@ -114,6 +115,9 @@ private:
 	/** Map of column type name -> factory delegate */
 	TMap< FName, FCreateSceneOutlinerColumn > ColumnMap;
 
+	// Function to recreate what originally was FActorInfoColumn as separate columns using FTextInfoColumn
+	void CreateActorInfoColumns(FSceneOutlinerInitializationOptions& InInitOptions) const;
+
 public:
 
 	// IModuleInterface interface
diff --git a/Engine/Source/Editor/SceneOutliner/Public/SceneOutlinerPublicTypes.h b/Engine/Source/Editor/SceneOutliner/Public/SceneOutlinerPublicTypes.h
--- a/Engine/Source/Editor/SceneOutliner/Public/SceneOutlinerPublicTypes.h
+++ b/Engine/Source/Editor/SceneOutliner/Public/SceneOutlinerPublicTypes.h
@@ -23,27 +23,27 @@ struct FSceneOutlinerBuiltInColumnTypes
 	/** The gutter column */
 	static const FName& Gutter()
 	{
-		static FName Gutter("Gutter");
+		static FName Gutter("Visibility"); // Renamed "Gutter" to "Visibility" so the purpose is more obvious in editor menus
 		return Gutter;
 	}
 
 	/** The item label column */
 	static const FName& Label()
 	{
-		static FName Label("ItemLabel");
+		static FName Label("Item Label");
 		return Label;
 	}
 
 	/** Generic actor info column */
 	static FName& ActorInfo()
 	{
-		static FName ActorInfo("ActorInfo");
+		static FName ActorInfo("Type"); // Renamed "Actor Info" to "Type" since it has been refactored to only show type information
 		return ActorInfo;
 	}
 
 	static FName& SourceControl()
 	{
-		static FName SourceControl("SourceControl");
+		static FName SourceControl("Source Control");
 		return SourceControl;
 	}
 
@@ -67,19 +67,19 @@ enum class ESceneOutlinerColumnVisibility : uint8
 /** Column information for the scene outliner */
 struct FSceneOutlinerColumnInfo
 {
-	FSceneOutlinerColumnInfo(ESceneOutlinerColumnVisibility InVisibility, int32 InPriorityIndex, const FCreateSceneOutlinerColumn& InFactory = FCreateSceneOutlinerColumn())
-		: Visibility(InVisibility), PriorityIndex(InPriorityIndex), Factory(InFactory)
+	FSceneOutlinerColumnInfo(ESceneOutlinerColumnVisibility InVisibility, int32 InPriorityIndex, const FCreateSceneOutlinerColumn& InFactory = FCreateSceneOutlinerColumn(), bool inCanBeHidden = true)
+		: Visibility(InVisibility), PriorityIndex(InPriorityIndex), bCanBeHidden(inCanBeHidden), Factory(InFactory)
 	{}
 
 	FSceneOutlinerColumnInfo() {}
 
 	FSceneOutlinerColumnInfo(const FSceneOutlinerColumnInfo& InColumnInfo)
-		: Visibility(InColumnInfo.Visibility), PriorityIndex(InColumnInfo.PriorityIndex), Factory(InColumnInfo.Factory)
+		: Visibility(InColumnInfo.Visibility), PriorityIndex(InColumnInfo.PriorityIndex), bCanBeHidden(InColumnInfo.bCanBeHidden), Factory(InColumnInfo.Factory)
 	{}
 
 	ESceneOutlinerColumnVisibility 	Visibility;
 	uint8				PriorityIndex;
-
+	bool bCanBeHidden;
 	FCreateSceneOutlinerColumn	Factory;
 };
 
