Fix crash when refreshing world outliner.

diff --git a/Engine/Source/Editor/SceneOutliner/Private/ActorMode.cpp b/Engine/Source/Editor/SceneOutliner/Private/ActorMode.cpp
--- a/Engine/Source/Editor/SceneOutliner/Private/ActorMode.cpp
+++ b/Engine/Source/Editor/SceneOutliner/Private/ActorMode.cpp
@@ -283,17 +283,17 @@ bool FActorMode::IsActorDisplayable(const SSceneOutliner* SceneOutliner, const A
 	static const FName SequencerActorTag(TEXT("SequencerActor"));
 
 	return Actor &&
-		!SceneOutliner->GetSharedData().bOnlyShowFolders && 												// Don't show actors if we're only showing folders
-		Actor->IsEditable() &&															// Only show actors that are allowed to be selected and drawn in editor
+		!SceneOutliner->GetSharedData().bOnlyShowFolders && 									// Don't show actors if we're only showing folders
+		Actor->IsEditable() &&																	// Only show actors that are allowed to be selected and drawn in editor
 		Actor->IsListedInSceneOutliner() &&
-		((Actor->GetWorld()->IsPlayInEditor() || !Actor->HasAnyFlags(RF_Transient)) ||
-		(SceneOutliner->GetSharedData().bShowTransient && Actor->HasAnyFlags(RF_Transient)) ||			// Don't show transient actors in non-play worlds
-			(Actor->ActorHasTag(SequencerActorTag))) &&
-		!Actor->IsTemplate() &&															// Should never happen, but we never want CDOs displayed
-		!FActorEditorUtils::IsABuilderBrush(Actor) &&									// Don't show the builder brush
-		!Actor->IsA(AWorldSettings::StaticClass()) &&									// Don't show the WorldSettings actor, even though it is technically editable
-		!Actor->IsPendingKill() &&														// We don't want to show actors that are about to go away
-		FLevelUtils::IsLevelVisible(Actor->GetLevel());									// Only show Actors whose level is visible
+		(((Actor->GetWorld() && Actor->GetWorld()->IsPlayInEditor()) || !Actor->HasAnyFlags(RF_Transient)) ||
+		(SceneOutliner->GetSharedData().bShowTransient && Actor->HasAnyFlags(RF_Transient)) ||	// Don't show transient actors in non-play worlds
+		(Actor->ActorHasTag(SequencerActorTag))) &&
+		!Actor->IsTemplate() &&																	// Should never happen, but we never want CDOs displayed
+		!FActorEditorUtils::IsABuilderBrush(Actor) &&											// Don't show the builder brush
+		!Actor->IsA(AWorldSettings::StaticClass()) &&											// Don't show the WorldSettings actor, even though it is technically editable
+		!Actor->IsPendingKill() &&																// We don't want to show actors that are about to go away
+		FLevelUtils::IsLevelVisible(Actor->GetLevel());											// Only show Actors whose level is visible
 }
 
 void FActorMode::OnFilterTextChanged(const FText& InFilterText)
