SceneOutliner: fixed issue where compiler a blueprint could reset the label for instances of that class

diff --git a/Engine/Source/Editor/SceneOutliner/Private/ActorModeInteractive.cpp b/Engine/Source/Editor/SceneOutliner/Private/ActorModeInteractive.cpp
--- a/Engine/Source/Editor/SceneOutliner/Private/ActorModeInteractive.cpp
+++ b/Engine/Source/Editor/SceneOutliner/Private/ActorModeInteractive.cpp
@@ -15,6 +15,7 @@ FActorModeInteractive::FActorModeInteractive(const FActorModeParams& Params)
 	FEditorDelegates::NewCurrentLevel.AddRaw(this, &FActorModeInteractive::OnNewCurrentLevel);
 
 	FCoreDelegates::OnActorLabelChanged.AddRaw(this, &FActorModeInteractive::OnActorLabelChanged);
+	FCoreUObjectDelegates::OnObjectsReplaced.AddRaw(this, &FActorModeInteractive::OnObjectsReplaced);
 	FCoreUObjectDelegates::PostLoadMapWithWorld.AddRaw(this, &FActorModeInteractive::OnPostLoadMapWithWorld);
 	GEngine->OnLevelActorRequestRename().AddRaw(this, &FActorModeInteractive::OnLevelActorRequestsRename);
 }
@@ -28,6 +29,7 @@ FActorModeInteractive::~FActorModeInteractive()
 	FEditorDelegates::NewCurrentLevel.RemoveAll(this);
 
 	FCoreDelegates::OnActorLabelChanged.RemoveAll(this);
+	FCoreUObjectDelegates::OnObjectsReplaced.RemoveAll(this);
 	FCoreUObjectDelegates::PostLoadMapWithWorld.RemoveAll(this);
 	GEngine->OnLevelActorRequestRename().RemoveAll(this);
 }
@@ -110,3 +112,18 @@ void FActorModeInteractive::OnActorLabelChanged(AActor* ChangedActor)
 		}
 	}
 }
+
+void FActorModeInteractive::OnObjectsReplaced(const TMap<UObject*, UObject*>& ReplacementMap)
+{
+	for (const TPair<UObject*, UObject*>& Pair : ReplacementMap)
+	{
+		AActor* Actor = Cast<AActor>(Pair.Value);
+		if (Actor && RepresentingWorld.Get() == Actor->GetWorld() && IsActorDisplayable(Actor))
+		{
+			if (FSceneOutlinerTreeItemPtr Item = CreateItemFor<FActorTreeItem>(Actor, true))
+			{
+				SceneOutliner->OnItemLabelChanged(Item);
+			}
+		}
+	}
+}
diff --git a/Engine/Source/Editor/SceneOutliner/Public/ActorModeInteractive.h b/Engine/Source/Editor/SceneOutliner/Public/ActorModeInteractive.h
--- a/Engine/Source/Editor/SceneOutliner/Public/ActorModeInteractive.h
+++ b/Engine/Source/Editor/SceneOutliner/Public/ActorModeInteractive.h
@@ -19,6 +19,7 @@ private:
 
 	void OnLevelSelectionChanged(UObject* Obj);
 	void OnActorLabelChanged(AActor* ChangedActor);
+	void OnObjectsReplaced(const TMap<UObject*, UObject*>& ReplacementMap);
 	void OnLevelActorRequestsRename(const AActor* Actor);
 	void OnPostLoadMapWithWorld(UWorld* World);
 };
