SceneOutliner: Pinning Actors scrolls them into view

diff --git a/Engine/Source/Editor/SceneOutliner/Private/ActorBrowsingMode.cpp b/Engine/Source/Editor/SceneOutliner/Private/ActorBrowsingMode.cpp
--- a/Engine/Source/Editor/SceneOutliner/Private/ActorBrowsingMode.cpp
+++ b/Engine/Source/Editor/SceneOutliner/Private/ActorBrowsingMode.cpp
@@ -1492,6 +1492,7 @@ void FActorBrowsingMode::PinItem(const FSceneOutlinerTreeItemPtr& InItem)
 			bIsBatchSelectOwner = true;
 			GEditor->GetSelectedActors()->BeginBatchSelectOperation();
 			GEditor->SelectNone(/*bNoteSelectionChange=*/false, /*bDeselectBSPSurfs=*/true);
+			SceneOutliner->OnItemAdded(PinnedActor, SceneOutliner::ENewItemAction::ScrollIntoView);
 		}
 
 		GEditor->SelectActor(PinnedActor, /*bInSelected=*/true, /*bNotify=*/false);
@@ -1568,6 +1569,11 @@ void FActorBrowsingMode::PinSelectedItems()
 		});
 
 		GEditor->GetSelectedActors()->EndBatchSelectOperation(/*bNotify=*/true);
+
+		if (AActor* Actor = GEditor->GetSelectedActors()->GetTop<AActor>())
+		{
+			SceneOutliner->OnItemAdded(Actor, SceneOutliner::ENewItemAction::ScrollIntoView);
+		}
 	}
 }
 
