Fixed mouse click on "Pin Actor" in Scene outliner that enters FSlateThrottleManager's Responsive mode without ever leaving it (similar to SSceneOutlinerTreeRow::OnMouseButtonDown)

diff --git a/Engine/Source/Editor/SceneOutliner/Private/SceneOutlinerPinnedColumn.cpp b/Engine/Source/Editor/SceneOutliner/Private/SceneOutlinerPinnedColumn.cpp
--- a/Engine/Source/Editor/SceneOutliner/Private/SceneOutlinerPinnedColumn.cpp
+++ b/Engine/Source/Editor/SceneOutliner/Private/SceneOutlinerPinnedColumn.cpp
@@ -133,7 +133,8 @@ private:
 	{
 		if (MouseEvent.GetEffectingButton() == EKeys::LeftMouseButton)
 		{
-			return HandleClick();
+			FReply Reply = HandleClick();
+			return Reply.PreventThrottling();
 		}
 		
 		return FReply::Unhandled();
