WorldPartition: ActorHierarchy / DataLayerHierarchy listen to IWorldPartitionEditorModule::OnWorldPartitionCreated event to handle a full refresh (SaveAs prep work)

diff --git a/Engine/Source/Editor/DataLayerEditor/Private/DataLayer/DataLayerHierarchy.cpp b/Engine/Source/Editor/DataLayerEditor/Private/DataLayer/DataLayerHierarchy.cpp
--- a/Engine/Source/Editor/DataLayerEditor/Private/DataLayer/DataLayerHierarchy.cpp
+++ b/Engine/Source/Editor/DataLayerEditor/Private/DataLayer/DataLayerHierarchy.cpp
@@ -13,6 +13,7 @@
 #include "EngineUtils.h"
 #include "Editor.h"
 #include "WorldPartition/WorldPartitionHelpers.h"
+#include "WorldPartition/IWorldPartitionEditorModule.h"
 
 TUniquePtr<FDataLayerHierarchy> FDataLayerHierarchy::Create(FDataLayerMode* Mode, const TWeakObjectPtr<UWorld>& World)
 {
@@ -34,6 +35,9 @@ FDataLayerHierarchy::FDataLayerHierarchy(FDataLayerMode* Mode, const TWeakObject
 		GEngine->OnLevelActorListChanged().AddRaw(this, &FDataLayerHierarchy::OnLevelActorListChanged);
 	}
 
+	IWorldPartitionEditorModule& WorldPartitionEditorModule = FModuleManager::LoadModuleChecked<IWorldPartitionEditorModule>("WorldPartitionEditor");
+	WorldPartitionEditorModule.OnWorldPartitionCreated().AddRaw(this, &FDataLayerHierarchy::OnWorldPartitionCreated);
+
 	if (World.IsValid())
 	{
 		if (World->PersistentLevel)
@@ -69,6 +73,9 @@ FDataLayerHierarchy::~FDataLayerHierarchy()
 		GEngine->OnLevelActorListChanged().RemoveAll(this);
 	}
 
+	IWorldPartitionEditorModule& WorldPartitionEditorModule = FModuleManager::LoadModuleChecked<IWorldPartitionEditorModule>("WorldPartitionEditor");
+	WorldPartitionEditorModule.OnWorldPartitionCreated().RemoveAll(this);
+
 	if (RepresentingWorld.IsValid())
 	{
 		if (RepresentingWorld->PersistentLevel)
@@ -198,6 +205,14 @@ FSceneOutlinerTreeItemPtr FDataLayerHierarchy::CreateParentItem(const FSceneOutl
 	return nullptr;
 }
 
+void FDataLayerHierarchy::OnWorldPartitionCreated(UWorld* InWorld)
+{
+	if (RepresentingWorld.Get() == InWorld)
+	{
+		FullRefreshEvent();
+	}
+}
+
 void FDataLayerHierarchy::OnLevelActorsAdded(const TArray<AActor*>& InActors)
 {
 	if (UWorld* CurrentWorld = RepresentingWorld.Get())
diff --git a/Engine/Source/Editor/DataLayerEditor/Private/DataLayer/DataLayerHierarchy.h b/Engine/Source/Editor/DataLayerEditor/Private/DataLayer/DataLayerHierarchy.h
--- a/Engine/Source/Editor/DataLayerEditor/Private/DataLayer/DataLayerHierarchy.h
+++ b/Engine/Source/Editor/DataLayerEditor/Private/DataLayer/DataLayerHierarchy.h
@@ -28,6 +28,7 @@ private:
 	FDataLayerHierarchy(const FDataLayerHierarchy&) = delete;
 	FDataLayerHierarchy& operator=(const FDataLayerHierarchy&) = delete;
 
+	void OnWorldPartitionCreated(UWorld* InWorld);
 	void OnLevelActorsAdded(const TArray<AActor*>& InActors);
 	void OnLevelActorsRemoved(const TArray<AActor*>& InActors);
 	void OnLevelActorAdded(AActor* InActor);
diff --git a/Engine/Source/Editor/SceneOutliner/Private/ActorHierarchy.cpp b/Engine/Source/Editor/SceneOutliner/Private/ActorHierarchy.cpp
--- a/Engine/Source/Editor/SceneOutliner/Private/ActorHierarchy.cpp
+++ b/Engine/Source/Editor/SceneOutliner/Private/ActorHierarchy.cpp
@@ -20,6 +20,7 @@
 #include "LevelInstance/LevelInstanceSubsystem.h"
 #include "WorldPartition/WorldPartitionHelpers.h"
 #include "WorldPartition/WorldPartitionActorDesc.h"
+#include "WorldPartition/IWorldPartitionEditorModule.h"
 
 TUniquePtr<FActorHierarchy> FActorHierarchy::Create(ISceneOutlinerMode* Mode, const TWeakObjectPtr<UWorld>& World)
 {
@@ -32,6 +33,9 @@ TUniquePtr<FActorHierarchy> FActorHierarchy::Create(ISceneOutlinerMode* Mode, co
 	GEngine->OnLevelActorFolderChanged().AddRaw(Hierarchy, &FActorHierarchy::OnLevelActorFolderChanged);
 	GEngine->OnLevelActorListChanged().AddRaw(Hierarchy, &FActorHierarchy::OnLevelActorListChanged);
 
+	IWorldPartitionEditorModule& WorldPartitionEditorModule = FModuleManager::LoadModuleChecked<IWorldPartitionEditorModule>("WorldPartitionEditor");
+	WorldPartitionEditorModule.OnWorldPartitionCreated().AddRaw(Hierarchy, &FActorHierarchy::OnWorldPartitionCreated);
+
 	if (World.IsValid())
 	{
 		if (World->PersistentLevel)
@@ -76,6 +80,9 @@ FActorHierarchy::~FActorHierarchy()
 		GEngine->OnLevelActorListChanged().RemoveAll(this);
 	}
 
+	IWorldPartitionEditorModule& WorldPartitionEditorModule = FModuleManager::LoadModuleChecked<IWorldPartitionEditorModule>("WorldPartitionEditor");
+	WorldPartitionEditorModule.OnWorldPartitionCreated().RemoveAll(this);
+
 	if (RepresentingWorld.IsValid())
 	{
 		if (RepresentingWorld->PersistentLevel)
@@ -497,6 +504,14 @@ void FActorHierarchy::FullRefreshEvent()
 	HierarchyChangedEvent.Broadcast(EventData);
 }
 
+void FActorHierarchy::OnWorldPartitionCreated(UWorld* InWorld)
+{
+	if (RepresentingWorld.Get() == InWorld)
+	{
+		FullRefreshEvent();
+	}
+}
+
 void FActorHierarchy::OnLevelActorAdded(AActor* InActor)
 {
 	if (InActor != nullptr && RepresentingWorld.Get() == InActor->GetWorld())
diff --git a/Engine/Source/Editor/SceneOutliner/Public/ActorHierarchy.h b/Engine/Source/Editor/SceneOutliner/Public/ActorHierarchy.h
--- a/Engine/Source/Editor/SceneOutliner/Public/ActorHierarchy.h
+++ b/Engine/Source/Editor/SceneOutliner/Public/ActorHierarchy.h
@@ -35,6 +35,7 @@ private:
 	void CreateComponentItems(const AActor* Actor, TArray<FSceneOutlinerTreeItemPtr>& OutItems) const;
 private:
 	// Update the hierarchy when actor or world changing events occur
+	void OnWorldPartitionCreated(UWorld* InWorld);
 
 	void OnLevelActorAdded(AActor* InActor);
 	void OnLevelActorDeleted(AActor* InActor);
