Add localization support for Outliner column names

diff --git a/Engine/Source/Editor/SceneOutliner/Private/SSceneOutliner.cpp b/Engine/Source/Editor/SceneOutliner/Private/SSceneOutliner.cpp
--- a/Engine/Source/Editor/SceneOutliner/Private/SSceneOutliner.cpp
+++ b/Engine/Source/Editor/SceneOutliner/Private/SSceneOutliner.cpp
@@ -325,9 +325,21 @@ void SSceneOutliner::SetupColumns(SHeaderRow& HeaderRow)
 					.SortMode(this, &SSceneOutliner::GetColumnSortMode, ID)
 					.OnSort(this, &SSceneOutliner::OnColumnSortModeChanged);
 			}
-			
-			ColumnArgs.DefaultLabel(FText::FromName(ID));
 
+			if (SharedData->ColumnMap[ID].ColumnLabel.IsSet())
+			{
+				ColumnArgs.DefaultLabel(SharedData->ColumnMap[ID].ColumnLabel);
+			}
+			else
+			{
+				if (HeaderRow.GetVisibility() == EVisibility::Visible)
+				{
+					UE_LOG(LogSceneOutliner, Error, TEXT("Outliner Column %s does not have a localizable name, please specify one to FSceneOutlinerColumnInfo"), *ID.ToString());
+				}
+				
+				ColumnArgs.DefaultLabel(FText::FromName(ID));
+			}
+			
 			if (!SharedData->ColumnMap[ID].bCanBeHidden)
 			{
 				ColumnArgs.ShouldGenerateWidget(true);
diff --git a/Engine/Source/Editor/SceneOutliner/Private/SceneOutlinerModule.cpp b/Engine/Source/Editor/SceneOutliner/Private/SceneOutlinerModule.cpp
--- a/Engine/Source/Editor/SceneOutliner/Private/SceneOutlinerModule.cpp
+++ b/Engine/Source/Editor/SceneOutliner/Private/SceneOutlinerModule.cpp
@@ -24,12 +24,14 @@
 #include "WorldPartition/DataLayer/DataLayer.h"
 #include "WorldPartition/DataLayer/WorldDataLayers.h"
 
+#define LOCTEXT_NAMESPACE "SceneOutlinerModule"
+
 /* FSceneOutlinerModule interface
  *****************************************************************************/
 
 void FSceneOutlinerModule::StartupModule()
 {
-	RegisterDefaultColumnType< FSceneOutlinerItemLabelColumn >(FSceneOutlinerColumnInfo(ESceneOutlinerColumnVisibility::Visible, 10, FCreateSceneOutlinerColumn(), false));
+	RegisterDefaultColumnType< FSceneOutlinerItemLabelColumn >(FSceneOutlinerColumnInfo(ESceneOutlinerColumnVisibility::Visible, 10, FCreateSceneOutlinerColumn(), false, TOptional<float>(), FSceneOutlinerBuiltInColumnTypes::Label_Localized()));
 
 	// Register builtin column types which are not active by default
 	RegisterColumnType<FSceneOutlinerGutter>();
@@ -85,9 +87,9 @@ TSharedRef<ISceneOutliner> FSceneOutlinerModule::CreateActorPicker(const FSceneO
 	InitOptions.ModeFactory = ModeFactory;
 	if (InitOptions.ColumnMap.Num() == 0)
 	{
-		InitOptions.ColumnMap.Add(FSceneOutlinerBuiltInColumnTypes::Label(), FSceneOutlinerColumnInfo(ESceneOutlinerColumnVisibility::Visible, 0, FCreateSceneOutlinerColumn(), false));
-		InitOptions.ColumnMap.Add(FSceneOutlinerBuiltInColumnTypes::ActorInfo(), FSceneOutlinerColumnInfo(ESceneOutlinerColumnVisibility::Visible, 10));
-		InitOptions.ColumnMap.Add(FSceneOutlinerBuiltInColumnTypes::Pinned(), FSceneOutlinerColumnInfo(ESceneOutlinerColumnVisibility::Invisible, 5));
+		InitOptions.ColumnMap.Add(FSceneOutlinerBuiltInColumnTypes::Label(), FSceneOutlinerColumnInfo(ESceneOutlinerColumnVisibility::Visible, 0, FCreateSceneOutlinerColumn(), false, TOptional<float>(), FSceneOutlinerBuiltInColumnTypes::Label_Localized()));
+		InitOptions.ColumnMap.Add(FSceneOutlinerBuiltInColumnTypes::ActorInfo(), FSceneOutlinerColumnInfo(ESceneOutlinerColumnVisibility::Visible, 10, FCreateSceneOutlinerColumn(), true, TOptional<float>(), FSceneOutlinerBuiltInColumnTypes::ActorInfo_Localized()));
+		InitOptions.ColumnMap.Add(FSceneOutlinerBuiltInColumnTypes::Pinned(), FSceneOutlinerColumnInfo(ESceneOutlinerColumnVisibility::Invisible, 5, FCreateSceneOutlinerColumn(), true, TOptional<float>(), FSceneOutlinerBuiltInColumnTypes::Pinned_Localized()));
 		CreateActorInfoColumns(InitOptions);
 	}
 	return CreateSceneOutliner(InitOptions);
@@ -122,8 +124,8 @@ TSharedRef<ISceneOutliner> FSceneOutlinerModule::CreateComponentPicker(const FSc
 	InitOptions.ModeFactory = ModeFactory;
 	if (InitOptions.ColumnMap.Num() == 0)
 	{
-		InitOptions.ColumnMap.Add(FSceneOutlinerBuiltInColumnTypes::Label(), FSceneOutlinerColumnInfo(ESceneOutlinerColumnVisibility::Visible, 0, FCreateSceneOutlinerColumn(), false));
-		InitOptions.ColumnMap.Add(FSceneOutlinerBuiltInColumnTypes::ActorInfo(), FSceneOutlinerColumnInfo(ESceneOutlinerColumnVisibility::Visible, 10));
+		InitOptions.ColumnMap.Add(FSceneOutlinerBuiltInColumnTypes::Label(), FSceneOutlinerColumnInfo(ESceneOutlinerColumnVisibility::Visible, 0, FCreateSceneOutlinerColumn(), false, TOptional<float>(), FSceneOutlinerBuiltInColumnTypes::Label_Localized()));
+		InitOptions.ColumnMap.Add(FSceneOutlinerBuiltInColumnTypes::ActorInfo(), FSceneOutlinerColumnInfo(ESceneOutlinerColumnVisibility::Visible, 10, FCreateSceneOutlinerColumn(), true, TOptional<float>(), FSceneOutlinerBuiltInColumnTypes::ActorInfo_Localized()));
 		CreateActorInfoColumns(InitOptions);
 	}
 	return CreateSceneOutliner(InitOptions);
@@ -141,10 +143,10 @@ TSharedRef< ISceneOutliner > FSceneOutlinerModule::CreateActorBrowser(const FSce
 	if (InitOptions.ColumnMap.Num() == 0)
 	{
 		InitOptions.UseDefaultColumns();
-		InitOptions.ColumnMap.Add(FSceneOutlinerBuiltInColumnTypes::Gutter(), FSceneOutlinerColumnInfo(ESceneOutlinerColumnVisibility::Visible, 0));
-		InitOptions.ColumnMap.Add(FSceneOutlinerBuiltInColumnTypes::ActorInfo(), FSceneOutlinerColumnInfo(ESceneOutlinerColumnVisibility::Visible, 20));
-		InitOptions.ColumnMap.Add(FSceneOutlinerBuiltInColumnTypes::SourceControl(), FSceneOutlinerColumnInfo(ESceneOutlinerColumnVisibility::Invisible, 30));
-		InitOptions.ColumnMap.Add(FSceneOutlinerBuiltInColumnTypes::Pinned(), FSceneOutlinerColumnInfo(ESceneOutlinerColumnVisibility::Invisible, 5));
+		InitOptions.ColumnMap.Add(FSceneOutlinerBuiltInColumnTypes::Gutter(), FSceneOutlinerColumnInfo(ESceneOutlinerColumnVisibility::Visible, 0, FCreateSceneOutlinerColumn(), true, TOptional<float>(), FSceneOutlinerBuiltInColumnTypes::Gutter_Localized()));
+		InitOptions.ColumnMap.Add(FSceneOutlinerBuiltInColumnTypes::ActorInfo(), FSceneOutlinerColumnInfo(ESceneOutlinerColumnVisibility::Visible, 20, FCreateSceneOutlinerColumn(), true, TOptional<float>(), FSceneOutlinerBuiltInColumnTypes::ActorInfo_Localized()));
+		InitOptions.ColumnMap.Add(FSceneOutlinerBuiltInColumnTypes::SourceControl(), FSceneOutlinerColumnInfo(ESceneOutlinerColumnVisibility::Invisible, 30, FCreateSceneOutlinerColumn(), true, TOptional<float>(), FSceneOutlinerBuiltInColumnTypes::SourceControl_Localized()));
+		InitOptions.ColumnMap.Add(FSceneOutlinerBuiltInColumnTypes::Pinned(), FSceneOutlinerColumnInfo(ESceneOutlinerColumnVisibility::Invisible, 5, FCreateSceneOutlinerColumn(), true, TOptional<float>(), FSceneOutlinerBuiltInColumnTypes::Pinned_Localized()));
 		CreateActorInfoColumns(InitOptions);
 	}
 	return CreateSceneOutliner(InitOptions);
@@ -231,6 +233,7 @@ void FSceneOutlinerModule::CreateActorInfoColumns(FSceneOutlinerInitializationOp
 
 		TArray<const UDataLayer*> DataLayerObjects;
 
+
 		if (const FActorTreeItem* ActorItem = Item.CastTo<FActorTreeItem>())
 		{
 			AActor* Actor = ActorItem->Actor.Get();
@@ -359,14 +362,31 @@ void FSceneOutlinerModule::CreateActorInfoColumns(FSceneOutlinerInitializationOp
 		return FString();
 	});
 
-	InInitOptions.ColumnMap.Add("Mobility", FSceneOutlinerColumnInfo(ESceneOutlinerColumnVisibility::Invisible, 20, FCreateSceneOutlinerColumn::CreateStatic(&FTextInfoColumn::CreateTextInfoColumn, FName("Mobility"), MobilityInfoText)));
-	InInitOptions.ColumnMap.Add("Layer", FSceneOutlinerColumnInfo(ESceneOutlinerColumnVisibility::Invisible, 20, FCreateSceneOutlinerColumn::CreateStatic(&FTextInfoColumn::CreateTextInfoColumn, FName("Layer"), LayerInfoText)));
-	InInitOptions.ColumnMap.Add("Data Layer", FSceneOutlinerColumnInfo(ESceneOutlinerColumnVisibility::Invisible, 20, FCreateSceneOutlinerColumn::CreateStatic(&FTextInfoColumn::CreateTextInfoColumn, FName("Data Layer"), DataLayerInfoText)));
-	InInitOptions.ColumnMap.Add("Socket", FSceneOutlinerColumnInfo(ESceneOutlinerColumnVisibility::Invisible, 20, FCreateSceneOutlinerColumn::CreateStatic(&FTextInfoColumn::CreateTextInfoColumn, FName("Socket"), SocketInfoText)));
-	InInitOptions.ColumnMap.Add("ID Name", FSceneOutlinerColumnInfo(ESceneOutlinerColumnVisibility::Invisible, 20, FCreateSceneOutlinerColumn::CreateStatic(&FTextInfoColumn::CreateTextInfoColumn, FName("ID Name"), InternalNameInfoText)));
-	InInitOptions.ColumnMap.Add("Package Short Name", FSceneOutlinerColumnInfo(ESceneOutlinerColumnVisibility::Invisible, 20, FCreateSceneOutlinerColumn::CreateStatic(&FTextInfoColumn::CreateTextInfoColumn, FName("Package Short Name"), PackageShortNameInfoText)));
-	InInitOptions.ColumnMap.Add("# Uncached Lights", FSceneOutlinerColumnInfo(ESceneOutlinerColumnVisibility::Invisible, 20, FCreateSceneOutlinerColumn::CreateStatic(&FTextInfoColumn::CreateTextInfoColumn, FName("# Uncached Lights"), UncachedLightsInfoText)));
-
+	auto AddTextInfoColumn = [&InInitOptions](FName ColumnID, TAttribute<FText> ColumnName, FGetTextForItem ColumnInfo)
+	{
+		InInitOptions.ColumnMap.Add(
+			ColumnID,
+			FSceneOutlinerColumnInfo(
+				ESceneOutlinerColumnVisibility::Invisible,
+				20,
+				FCreateSceneOutlinerColumn::CreateStatic(
+					&FTextInfoColumn::CreateTextInfoColumn,
+					ColumnID,
+					ColumnInfo),
+				true,
+				TOptional<float>(),
+				ColumnName));
+	};
+
+	AddTextInfoColumn("Mobility", LOCTEXT("SceneOutlinerMobilityColumn", "Mobility"), MobilityInfoText);
+	AddTextInfoColumn("Level", LOCTEXT("SceneOutlinerLevelColumn", "Level"), LevelInfoText);
+	AddTextInfoColumn("Layer", LOCTEXT("SceneOutlinerLayerColumn", "Layer"), LayerInfoText);
+	AddTextInfoColumn("Data Layer", LOCTEXT("SceneOutlinerDataLayerColumn", "Data Layer"), DataLayerInfoText);
+	AddTextInfoColumn("ID Name", LOCTEXT("SceneOutlinerIDColumn", "ID Name"), InternalNameInfoText);
+	AddTextInfoColumn("Package Short Name", LOCTEXT("SceneOutlinerPackageShortNameColumn", "Package Short Name"), PackageShortNameInfoText);
+	AddTextInfoColumn("Uncached Lights", LOCTEXT("SceneOutlinerUncachedLightsColumn", "# Uncached Lights"), UncachedLightsInfoText);
 }
 
 IMPLEMENT_MODULE(FSceneOutlinerModule, SceneOutliner);
+
+#undef LOCTEXT_NAMESPACE //SceneOutlinerModule
diff --git a/Engine/Source/Editor/SceneOutliner/Public/SceneOutlinerPublicTypes.h b/Engine/Source/Editor/SceneOutliner/Public/SceneOutlinerPublicTypes.h
--- a/Engine/Source/Editor/SceneOutliner/Public/SceneOutlinerPublicTypes.h
+++ b/Engine/Source/Editor/SceneOutliner/Public/SceneOutlinerPublicTypes.h
@@ -9,6 +9,8 @@
 
 #include "SceneOutlinerFilters.h"
 
+#define LOCTEXT_NAMESPACE "SceneOutlinerPublicTypes"
+
 class FExtender;
 struct FToolMenuContext;
 
@@ -27,6 +29,13 @@ struct FSceneOutlinerBuiltInColumnTypes
 		return Gutter;
 	}
 
+	/** Localizable FText name for the Gutter Column (pass into FSceneOutlinerColumnInfo) */
+	static const FText& Gutter_Localized()
+	{
+		static FText Gutter_Localized = LOCTEXT("VisibilityColumnName", "Visibility");
+		return Gutter_Localized;
+	}
+
 	/** The item label column */
 	static const FName& Label()
 	{
@@ -34,6 +43,13 @@ struct FSceneOutlinerBuiltInColumnTypes
 		return Label;
 	}
 
+	/** Localizable FText name for the Item Label Column (pass into FSceneOutlinerColumnInfo) */
+	static const FText& Label_Localized()
+	{
+		static FText Label_Localized = LOCTEXT("ItemLabelColumnName", "Item Label");
+		return Label_Localized;
+	}
+
 	/** Generic actor info column */
 	static FName& ActorInfo()
 	{
@@ -41,17 +57,38 @@ struct FSceneOutlinerBuiltInColumnTypes
 		return ActorInfo;
 	}
 
+	/** Localizable FText name for the Type Column (pass into FSceneOutlinerColumnInfo) */
+	static const FText& ActorInfo_Localized()
+	{
+		static FText ActorInfo_Localized = LOCTEXT("TypeColumnName", "Type");
+		return ActorInfo_Localized;
+	}
+
 	static FName& SourceControl()
 	{
 		static FName SourceControl("Source Control");
 		return SourceControl;
 	}
 
+	/** Localizable FText name for the Type Column (pass into FSceneOutlinerColumnInfo) */
+	static const FText& SourceControl_Localized()
+	{
+		static FText SourceControl_Localized = LOCTEXT("SourceControlColumnName", "Source Control");
+		return SourceControl_Localized;
+	}
+
 	static FName& Pinned()
 	{
 		static FName Pinned("Pinned");
 		return Pinned;
 	}
+
+	/** Localizable FText name for the Type Column (pass into FSceneOutlinerColumnInfo) */
+	static const FText& Pinned_Localized()
+	{
+		static FText Pinned_Localized = LOCTEXT("PinnedColumnName", "Pinned");
+		return Pinned_Localized;
+	}
 };
 
 /** Visibility enum for scene outliner columns */
@@ -67,15 +104,16 @@ enum class ESceneOutlinerColumnVisibility : uint8
 /** Column information for the scene outliner */
 struct FSceneOutlinerColumnInfo
 {
-	FSceneOutlinerColumnInfo(ESceneOutlinerColumnVisibility InVisibility, int32 InPriorityIndex, const FCreateSceneOutlinerColumn& InFactory = FCreateSceneOutlinerColumn(), bool inCanBeHidden = true, TOptional<float> InFillSize = TOptional<float>())
-		: Visibility(InVisibility), PriorityIndex(InPriorityIndex), bCanBeHidden(inCanBeHidden), Factory(InFactory), FillSize(InFillSize)
+	FSceneOutlinerColumnInfo(ESceneOutlinerColumnVisibility InVisibility, int32 InPriorityIndex, const FCreateSceneOutlinerColumn& InFactory = FCreateSceneOutlinerColumn(), bool inCanBeHidden = true, TOptional<float> InFillSize = TOptional<float>()
+		, TAttribute<FText> InColumnLabel = TAttribute<FText>())
+		: Visibility(InVisibility), PriorityIndex(InPriorityIndex), bCanBeHidden(inCanBeHidden), Factory(InFactory), FillSize(InFillSize), ColumnLabel(InColumnLabel)
 	{
 	}
 
 	FSceneOutlinerColumnInfo() {}
 
 	FSceneOutlinerColumnInfo(const FSceneOutlinerColumnInfo& InColumnInfo)
-		: Visibility(InColumnInfo.Visibility), PriorityIndex(InColumnInfo.PriorityIndex), bCanBeHidden(InColumnInfo.bCanBeHidden), Factory(InColumnInfo.Factory), FillSize(InColumnInfo.FillSize)
+		: Visibility(InColumnInfo.Visibility), PriorityIndex(InColumnInfo.PriorityIndex), bCanBeHidden(InColumnInfo.bCanBeHidden), Factory(InColumnInfo.Factory), FillSize(InColumnInfo.FillSize), ColumnLabel(InColumnInfo.ColumnLabel)
 	{}
 
 	ESceneOutlinerColumnVisibility 	Visibility;
@@ -83,6 +121,7 @@ struct FSceneOutlinerColumnInfo
 	bool bCanBeHidden;
 	FCreateSceneOutlinerColumn	Factory;
 	TOptional< float > FillSize;
+	TAttribute<FText> ColumnLabel; // Override for the column name used instead of ID if specified (use this if you want the column name to be localizable)
 };
 
 /** Settings for the scene outliner which can be quieried publicly */
@@ -175,3 +214,5 @@ struct SCENEOUTLINER_API FSceneOutlinerVisibilityCache
 
 	bool GetVisibility(const ISceneOutlinerTreeItem& Item) const;
 };
+
+#undef LOCTEXT_NAMESPACE
\ No newline at end of file
