SceneOutliner: fixed selection being reset when level change occured.

diff --git a/Engine/Source/Editor/SceneOutliner/Private/SSceneOutliner.cpp b/Engine/Source/Editor/SceneOutliner/Private/SSceneOutliner.cpp
--- a/Engine/Source/Editor/SceneOutliner/Private/SSceneOutliner.cpp
+++ b/Engine/Source/Editor/SceneOutliner/Private/SSceneOutliner.cpp
@@ -429,6 +429,7 @@ void SSceneOutliner::FullRefresh()
 	UE_LOG(LogSceneOutliner, Verbose, TEXT("Full Refresh"));
 	bDisableIntermediateSorting = true;
 	bFullRefresh = true;
+	RefreshSelection();
 	Refresh();
 }
 
