Change Log in Scene Outliner for non localized column names from Error to Info

diff --git a/Engine/Source/Editor/SceneOutliner/Private/SSceneOutliner.cpp b/Engine/Source/Editor/SceneOutliner/Private/SSceneOutliner.cpp
--- a/Engine/Source/Editor/SceneOutliner/Private/SSceneOutliner.cpp
+++ b/Engine/Source/Editor/SceneOutliner/Private/SSceneOutliner.cpp
@@ -334,7 +334,7 @@ void SSceneOutliner::SetupColumns(SHeaderRow& HeaderRow)
 			{
 				if (HeaderRow.GetVisibility() == EVisibility::Visible)
 				{
-					UE_LOG(LogSceneOutliner, Error, TEXT("Outliner Column %s does not have a localizable name, please specify one to FSceneOutlinerColumnInfo"), *ID.ToString());
+					UE_LOG(LogSceneOutliner, Log, TEXT("Outliner Column %s does not have a localizable name, please specify one to FSceneOutlinerColumnInfo"), *ID.ToString());
 				}
 				
 				ColumnArgs.DefaultLabel(FText::FromName(ID));
