Outliner column visibility is now saved across sessions

diff --git a/Engine/Source/Editor/LevelEditor/Private/SLevelEditor.cpp b/Engine/Source/Editor/LevelEditor/Private/SLevelEditor.cpp
--- a/Engine/Source/Editor/LevelEditor/Private/SLevelEditor.cpp
+++ b/Engine/Source/Editor/LevelEditor/Private/SLevelEditor.cpp
@@ -767,6 +767,7 @@ TSharedRef<SDockTab> SLevelEditor::SpawnLevelEditorTab( const FSpawnTabArgs& Arg
 
 
 		FText Label = NSLOCTEXT( "LevelEditor", "SceneOutlinerTabTitle", "Outliner" );
+		InitOptions.OutlinerIdentifier = "LevelEditorOutliner";
 
 		FSceneOutlinerModule& SceneOutlinerModule = FModuleManager::Get().LoadModuleChecked<FSceneOutlinerModule>("SceneOutliner");
 		TSharedRef<ISceneOutliner> SceneOutlinerRef = SceneOutlinerModule.CreateActorBrowser(
diff --git a/Engine/Source/Editor/SceneOutliner/Private/ActorBrowsingMode.cpp b/Engine/Source/Editor/SceneOutliner/Private/ActorBrowsingMode.cpp
--- a/Engine/Source/Editor/SceneOutliner/Private/ActorBrowsingMode.cpp
+++ b/Engine/Source/Editor/SceneOutliner/Private/ActorBrowsingMode.cpp
@@ -212,9 +212,6 @@ void FActorBrowsingMode::Rebuild()
 		UWorldPartition* const WorldPartition = RepresentingWorld->GetWorldPartition();
 		WorldPartition->OnActorDescRemovedEvent.AddRaw(this, &FActorBrowsingMode::OnActorDescRemoved);
 	}
-
-	// Enable the pinned column by default on WP worlds or Disable it by default on non-WP worlds
-	SceneOutliner->SetColumnVisibility(FSceneOutlinerBuiltInColumnTypes::Pinned(), bRepresentingWorldPartitionedWorld);
 }
 
 FText FActorBrowsingMode::GetStatusText() const 
diff --git a/Engine/Source/Editor/SceneOutliner/Private/SSceneOutliner.cpp b/Engine/Source/Editor/SceneOutliner/Private/SSceneOutliner.cpp
--- a/Engine/Source/Editor/SceneOutliner/Private/SSceneOutliner.cpp
+++ b/Engine/Source/Editor/SceneOutliner/Private/SSceneOutliner.cpp
@@ -36,6 +36,7 @@
 #include "ISceneOutlinerMode.h"
 #include "FolderTreeItem.h"
 #include "EditorFolderUtils.h"
+#include "SceneOutlinerConfig.h"
 
 #include "Framework/Notifications/NotificationManager.h"
 #include "Widgets/Notifications/SNotificationList.h"
@@ -75,10 +76,14 @@ void SSceneOutliner::Construct(const FArguments& InArgs, const FSceneOutlinerIni
 	bSelectionDirty = true;
 	SortOutlinerTimer = 0.0f;
 	bPendingFocusNextFrame = InInitOptions.bFocusSearchBoxWhenOpened;
-
+	OutlinerIdentifier = InInitOptions.OutlinerIdentifier;
+	
 	SortByColumn = FSceneOutlinerBuiltInColumnTypes::Label();
 	SortMode = EColumnSortMode::Ascending;
 
+	UOutlinerConfig* OutlinerConfig = GetMutableDefault<UOutlinerConfig>();
+	OutlinerConfig->LoadEditorConfig();
+
 	// @todo outliner: Should probably save this in layout!
 	// @todo outliner: Should save spacing for list view in layout
 
@@ -269,7 +274,21 @@ void SSceneOutliner::Construct(const FArguments& InArgs, const FSceneOutlinerIni
 
 void SSceneOutliner::HandleHiddenColumnsChanged()
 {
-	//Todo: Save?
+	const TArray<FName> HiddenColumns = HeaderRowWidget->GetHiddenColumnIds();
+	FSceneOutlinerConfig* OutlinerConfig = GetMutableConfig();
+	TMap<FName, bool> ColumnVisibilities;
+
+	if (OutlinerConfig != nullptr)
+	{
+		for (const TPair<FName, bool> Pair : Columns)
+		{
+			ColumnVisibilities.Add(Pair.Key, HiddenColumns.Find(Pair.Key) == INDEX_NONE);
+		}
+
+		OutlinerConfig->ColumnVisibilities = ColumnVisibilities;
+
+		SaveConfig();
+	}
 }
 
 void SSceneOutliner::SetupColumns(SHeaderRow& HeaderRow)
@@ -293,11 +312,26 @@ void SSceneOutliner::SetupColumns(SHeaderRow& HeaderRow)
 		return SharedData->ColumnMap[A].PriorityIndex < SharedData->ColumnMap[B].PriorityIndex;
 	});
 
+	TMap<FName, bool> ColumnVisibilities;
+	const FSceneOutlinerConfig* OutlinerConfig = GetConstConfig();
+
+	// Try to load visibility of columns from the config file
+	if (OutlinerConfig)
+	{
+		ColumnVisibilities = OutlinerConfig->ColumnVisibilities;
+	}
+
+
 	for (const FName& ID : SortedIDs)
 	{
 		bool bIsVisible = true;
 
-		if (SharedData->ColumnMap[ID].Visibility == ESceneOutlinerColumnVisibility::Invisible)
+		// If there is a config saved for this column, ignore the default visibility
+		if (bool *ColumnVisibility = ColumnVisibilities.Find(ID))
+		{
+			bIsVisible = *ColumnVisibility;
+		}
+		else if (SharedData->ColumnMap[ID].Visibility == ESceneOutlinerColumnVisibility::Invisible)
 		{
 			bIsVisible = false;
 		}
@@ -2067,4 +2101,33 @@ FSceneOutlinerTreeItemPtr SSceneOutliner::FindParent(const ISceneOutlinerTreeIte
 	return Parent;
 }
 
+
+struct FSceneOutlinerConfig* SSceneOutliner::GetMutableConfig()
+{
+	if (OutlinerIdentifier.IsNone())
+	{
+		return nullptr;
+	}
+
+	UOutlinerConfig* OutlinerConfig = GetMutableDefault<UOutlinerConfig>();
+	return &OutlinerConfig->Outliners.FindOrAdd(OutlinerIdentifier);
+}
+
+
+const FSceneOutlinerConfig* SSceneOutliner::GetConstConfig() const
+{
+	if (OutlinerIdentifier.IsNone())
+	{
+		return nullptr;
+	}
+
+	const UOutlinerConfig* OutlinerConfig = GetDefault<UOutlinerConfig>();
+	return OutlinerConfig->Outliners.Find(OutlinerIdentifier);
+}
+
+void SSceneOutliner::SaveConfig()
+{
+	GetMutableDefault<UOutlinerConfig>()->SaveEditorConfig();
+}
+
 #undef LOCTEXT_NAMESPACE
diff --git a/Engine/Source/Editor/SceneOutliner/Private/SceneOutlinerConfig.h b/Engine/Source/Editor/SceneOutliner/Private/SceneOutlinerConfig.h
new file mode 100644
--- /dev/null
+++ b/Engine/Source/Editor/SceneOutliner/Private/SceneOutlinerConfig.h
@@ -0,0 +1,33 @@
+// Copyright Epic Games, Inc. All Rights Reserved.
+
+#pragma once
+
+#include "Containers/Map.h"
+#include "EditorConfigBase.h"
+
+#include "SceneOutlinerConfig.generated.h"
+
+
+USTRUCT()
+struct FSceneOutlinerConfig
+{
+	GENERATED_BODY()
+
+public:
+	
+	/** Map to store the visibility of each column */
+	UPROPERTY()
+	TMap<FName, bool> ColumnVisibilities;
+};
+
+
+UCLASS(EditorConfig="Outliner")
+class UOutlinerConfig : public UEditorConfigBase
+{
+	GENERATED_BODY()
+	
+public:
+
+	UPROPERTY(meta=(EditorConfig))
+	TMap<FName, FSceneOutlinerConfig> Outliners;
+};
\ No newline at end of file
diff --git a/Engine/Source/Editor/SceneOutliner/Private/SceneOutlinerModule.cpp b/Engine/Source/Editor/SceneOutliner/Private/SceneOutlinerModule.cpp
--- a/Engine/Source/Editor/SceneOutliner/Private/SceneOutlinerModule.cpp
+++ b/Engine/Source/Editor/SceneOutliner/Private/SceneOutlinerModule.cpp
@@ -146,7 +146,7 @@ TSharedRef< ISceneOutliner > FSceneOutlinerModule::CreateActorBrowser(const FSce
 		InitOptions.ColumnMap.Add(FSceneOutlinerBuiltInColumnTypes::Gutter(), FSceneOutlinerColumnInfo(ESceneOutlinerColumnVisibility::Visible, 0, FCreateSceneOutlinerColumn(), true, TOptional<float>(), FSceneOutlinerBuiltInColumnTypes::Gutter_Localized()));
 		InitOptions.ColumnMap.Add(FSceneOutlinerBuiltInColumnTypes::ActorInfo(), FSceneOutlinerColumnInfo(ESceneOutlinerColumnVisibility::Visible, 20, FCreateSceneOutlinerColumn(), true, TOptional<float>(), FSceneOutlinerBuiltInColumnTypes::ActorInfo_Localized()));
 		InitOptions.ColumnMap.Add(FSceneOutlinerBuiltInColumnTypes::SourceControl(), FSceneOutlinerColumnInfo(ESceneOutlinerColumnVisibility::Invisible, 30, FCreateSceneOutlinerColumn(), true, TOptional<float>(), FSceneOutlinerBuiltInColumnTypes::SourceControl_Localized()));
-		InitOptions.ColumnMap.Add(FSceneOutlinerBuiltInColumnTypes::Pinned(), FSceneOutlinerColumnInfo(ESceneOutlinerColumnVisibility::Invisible, 5, FCreateSceneOutlinerColumn(), true, TOptional<float>(), FSceneOutlinerBuiltInColumnTypes::Pinned_Localized()));
+		InitOptions.ColumnMap.Add(FSceneOutlinerBuiltInColumnTypes::Pinned(), FSceneOutlinerColumnInfo(ESceneOutlinerColumnVisibility::Visible, 5, FCreateSceneOutlinerColumn(), true, TOptional<float>(), FSceneOutlinerBuiltInColumnTypes::Pinned_Localized()));
 		CreateActorInfoColumns(InitOptions);
 	}
 	return CreateSceneOutliner(InitOptions);
diff --git a/Engine/Source/Editor/SceneOutliner/Public/SSceneOutliner.h b/Engine/Source/Editor/SceneOutliner/Public/SSceneOutliner.h
--- a/Engine/Source/Editor/SceneOutliner/Public/SSceneOutliner.h
+++ b/Engine/Source/Editor/SceneOutliner/Public/SSceneOutliner.h
@@ -545,8 +545,6 @@ private:
 	/** Map of columns that are shown on this outliner. */
 	TMap<FName, TSharedPtr<ISceneOutlinerColumn>> Columns;
 
-	TArray<FName> HiddenColumnsList;
-
 	/** Set up the columns required for this outliner */
 	void SetupColumns(SHeaderRow& HeaderRow);
 
@@ -830,6 +828,9 @@ private:
 	/** Currently selected sorting mode */
 	EColumnSortMode::Type SortMode;
 
+	/** Identifier for this outliner (Set through FSceneOutlinerInitializationOptions)*/
+	FName OutlinerIdentifier;
+
 	/** Handles column sorting mode change */
 	void OnColumnSortModeChanged( const EColumnSortPriority::Type SortPriority, const FName& ColumnId, const EColumnSortMode::Type InSortMode );
 
@@ -840,4 +841,20 @@ private:
 
 	/** Handler for recursively expanding/collapsing items */
 	void SetItemExpansionRecursive(FSceneOutlinerTreeItemPtr Model, bool bInExpansionState);
+
+	/**
+	 * Get a mutable version of the outliner config for setting values.
+	 * @returns		The outliner config for this outliner.
+	 * @note		If FSceneOutlinerInitializationOptions.ViewIdentifier is not set, it is not possible to store settings for this outliner.
+	 */
+	struct FSceneOutlinerConfig* GetMutableConfig();
+
+	/**
+	 * Get a const version of the outliner config for getting values.
+	 * @returns		The outliner config for this outliner.
+	 * @note		If FSceneOutlinerInitializationOptions.ViewIdentifier is not set, it is not possible to retrieve settings for this outliner.
+	 */
+	const FSceneOutlinerConfig* GetConstConfig() const;
+
+	void SaveConfig();
 };
diff --git a/Engine/Source/Editor/SceneOutliner/Public/SceneOutlinerPublicTypes.h b/Engine/Source/Editor/SceneOutliner/Public/SceneOutlinerPublicTypes.h
--- a/Engine/Source/Editor/SceneOutliner/Public/SceneOutlinerPublicTypes.h
+++ b/Engine/Source/Editor/SceneOutliner/Public/SceneOutlinerPublicTypes.h
@@ -182,6 +182,9 @@ struct FSceneOutlinerInitializationOptions : FSharedSceneOutlinerData
 
 	FCreateSceneOutlinerMode ModeFactory;
 
+	/** Identifier for this outliner; NAME_None if this view is anonymous (Needs to be specified to save visibility of columns in EditorConfig)*/
+	FName OutlinerIdentifier;
+
 public:
 
 	/** Constructor */
@@ -192,6 +195,7 @@ public:
 		, bShowCreateNewFolder( true )
 		, Filters( new FSceneOutlinerFilters )
 		, ModeFactory()
+		, OutlinerIdentifier(NAME_None)
 	{}
 };
 
diff --git a/Engine/Source/Editor/SceneOutliner/SceneOutliner.Build.cs b/Engine/Source/Editor/SceneOutliner/SceneOutliner.Build.cs
--- a/Engine/Source/Editor/SceneOutliner/SceneOutliner.Build.cs
+++ b/Engine/Source/Editor/SceneOutliner/SceneOutliner.Build.cs
@@ -24,8 +24,9 @@ public class SceneOutliner : ModuleRules
 				"EditorFramework",
 				"UnrealEd",
 				"ToolMenus",
-				"SourceControl"
-            }
+				"SourceControl",
+				"EditorConfig"
+			}
 		);
 	}
 }
