Fix a bug in SceneOutliner that would sometimes cause all actors in the tree to be selected when you tried to select a filtered out search

diff --git a/Engine/Source/Editor/SceneOutliner/Private/SSceneOutliner.cpp b/Engine/Source/Editor/SceneOutliner/Private/SSceneOutliner.cpp
--- a/Engine/Source/Editor/SceneOutliner/Private/SSceneOutliner.cpp
+++ b/Engine/Source/Editor/SceneOutliner/Private/SSceneOutliner.cpp
@@ -1821,6 +1821,7 @@ void SSceneOutliner::OnFilterTextCommitted( const FText& InFilterText, ETextComm
 			// Gather all of the items that match the filter text
 			for (auto& Pair : TreeItemMap)
 			{
+				Pair.Value->Flags.bIsFilteredOut = !SearchBoxFilter->PassesFilter(*Pair.Value);
 				if (!Pair.Value->Flags.bIsFilteredOut)
 				{
 					Selection.Add(Pair.Value);
