Fix Static Analysis Error in Scene Outliner

diff --git a/Engine/Source/Editor/SceneOutliner/Private/SSceneOutliner.cpp b/Engine/Source/Editor/SceneOutliner/Private/SSceneOutliner.cpp
--- a/Engine/Source/Editor/SceneOutliner/Private/SSceneOutliner.cpp
+++ b/Engine/Source/Editor/SceneOutliner/Private/SSceneOutliner.cpp
@@ -280,7 +280,7 @@ void SSceneOutliner::HandleHiddenColumnsChanged()
 
 	if (OutlinerConfig != nullptr)
 	{
-		for (const TPair<FName, bool> Pair : Columns)
+		for (const TPair<FName, TSharedPtr<ISceneOutlinerColumn>>& Pair : Columns)
 		{
 			ColumnVisibilities.Add(Pair.Key, HiddenColumns.Find(Pair.Key) == INDEX_NONE);
 		}
