Add Expand/Collapse all option to the Outliner Settings menu

diff --git a/Engine/Source/Editor/SceneOutliner/Private/SSceneOutliner.cpp b/Engine/Source/Editor/SceneOutliner/Private/SSceneOutliner.cpp
--- a/Engine/Source/Editor/SceneOutliner/Private/SSceneOutliner.cpp
+++ b/Engine/Source/Editor/SceneOutliner/Private/SSceneOutliner.cpp
@@ -431,6 +431,22 @@ TSharedRef<SWidget> SSceneOutliner::GetViewButtonContent(bool bShowFilters)
 	// Menu should stay open on selection if filters are not being shown
 	FMenuBuilder MenuBuilder(bShowFilters, NULL);
 
+	MenuBuilder.BeginSection("OutlinerSettings", LOCTEXT("HierarchyHeading", "Hierarchy"));
+	{
+		MenuBuilder.AddMenuEntry(
+			LOCTEXT("ExpandAll", "Expand All"),
+			LOCTEXT("ExpandAllToolTip", "Expand All Items in the Hierarchy"),
+			FSlateIcon(),
+			FUIAction(FExecuteAction::CreateSP(this, &SSceneOutliner::ExpandAll)));
+
+		MenuBuilder.AddMenuEntry(
+			LOCTEXT("CollapseAll", "Collapse All"),
+			LOCTEXT("CollapseAllToolTip", "Collapse All Items in the Hierarchy"),
+			FSlateIcon(),
+			FUIAction(FExecuteAction::CreateSP(this, &SSceneOutliner::CollapseAll)));
+	}
+	MenuBuilder.EndSection();
+
 	if (bShowFilters)
 	{
 		MenuBuilder.BeginSection("AssetThumbnails", LOCTEXT("ShowHeading", "Show"));
