Fix missing pinned icon on Outliner

diff --git a/Engine/Source/Editor/SceneOutliner/Private/SceneOutlinerPinnedColumn.cpp b/Engine/Source/Editor/SceneOutliner/Private/SceneOutlinerPinnedColumn.cpp
--- a/Engine/Source/Editor/SceneOutliner/Private/SceneOutlinerPinnedColumn.cpp
+++ b/Engine/Source/Editor/SceneOutliner/Private/SceneOutlinerPinnedColumn.cpp
@@ -73,8 +73,8 @@ public:
 		static const FName NAME_PinnedBrush = TEXT("SceneOutliner.PinnedIcon");
 		static const FName NAME_UnpinnedBrush = TEXT("SceneOutliner.UnpinnedIcon");
 
-		PinnedBrush = FEditorStyle::GetBrush(NAME_PinnedBrush);
-		UnpinnedBrush = FEditorStyle::GetBrush(NAME_UnpinnedBrush);
+		PinnedBrush = FAppStyle::Get().GetBrush(NAME_PinnedBrush);
+		UnpinnedBrush = FAppStyle::Get().GetBrush(NAME_UnpinnedBrush);
 	}
 
 private:
@@ -221,7 +221,7 @@ bool FSceneOutlinerPinnedColumn::IsItemPinned(const ISceneOutlinerTreeItem& Item
 
 const FSlateBrush* FSceneOutlinerPinnedColumn::GetHeaderIcon() const
 {
-	return FEditorStyle::GetBrush("SceneOutliner.PinnedHighlighIcon");
+	return FAppStyle::Get().GetBrush("SceneOutliner.UnpinnedIcon");
 }
 
 #undef LOCTEXT_NAMESPACE
\ No newline at end of file
