LevelInstance: Fix Item Foreground color when editing child

diff --git a/Engine/Source/Editor/SceneOutliner/Private/ActorTreeItem.cpp b/Engine/Source/Editor/SceneOutliner/Private/ActorTreeItem.cpp
--- a/Engine/Source/Editor/SceneOutliner/Private/ActorTreeItem.cpp
+++ b/Engine/Source/Editor/SceneOutliner/Private/ActorTreeItem.cpp
@@ -253,6 +253,15 @@ private:
 	{
 		AActor* Actor = ActorPtr.Get();
 
+		// Color LevelInstances differently if they are being edited
+		if (const ALevelInstance* LevelInstanceActor = Cast<ALevelInstance>(Actor))
+		{
+			if (LevelInstanceActor->IsEditing())
+			{
+				return FAppStyle::Get().GetSlateColor("Colors.AccentGreen");
+			}
+		}
+
 		auto TreeItem = TreeItemPtr.Pin();
 		if (auto BaseColor = FSceneOutlinerCommonLabelData::GetForegroundColor(*TreeItem))
 		{
@@ -287,15 +296,6 @@ private:
 			return FSceneOutlinerCommonLabelData::DarkColor;
 		}
 
-		// Color LevelInstances differently if they are being edited
-		if (const ALevelInstance* LevelInstanceActor = Cast<ALevelInstance>(Actor))
-		{
-			if (LevelInstanceActor->IsEditing() && !LevelInstanceActor->IsSelected())
-			{
-				return FAppStyle::Get().GetSlateColor("Colors.AccentGreen");
-			}
-		}
-
 		return FSlateColor::UseForeground();
 	}
 
