Add Tooltip for the pin icon (header) in SceneOutliner
#trivial

diff --git a/Engine/Source/Editor/SceneOutliner/Private/SceneOutlinerPinnedColumn.cpp b/Engine/Source/Editor/SceneOutliner/Private/SceneOutlinerPinnedColumn.cpp
--- a/Engine/Source/Editor/SceneOutliner/Private/SceneOutlinerPinnedColumn.cpp
+++ b/Engine/Source/Editor/SceneOutliner/Private/SceneOutlinerPinnedColumn.cpp
@@ -186,7 +186,7 @@ SHeaderRow::FColumn::FArguments FSceneOutlinerPinnedColumn::ConstructHeaderRowCo
 		.VAlignHeader(VAlign_Center)
 		.HAlignCell(HAlign_Center)
 		.VAlignCell(VAlign_Center)
-		.DefaultTooltip(FText::FromName(GetColumnID()))
+		.DefaultTooltip(LOCTEXT("SceneOutlinerPinnedColumnToolTip", "Pinned - Force loaded in editor"))
 		[
 			SNew(SImage)
 			.ColorAndOpacity(FSlateColor::UseForeground())
