SceneOutliner: Remove Pinned & Visibility toggles when unsupported

diff --git a/Engine/Source/Editor/SceneOutliner/Private/ActorTreeItem.cpp b/Engine/Source/Editor/SceneOutliner/Private/ActorTreeItem.cpp
--- a/Engine/Source/Editor/SceneOutliner/Private/ActorTreeItem.cpp
+++ b/Engine/Source/Editor/SceneOutliner/Private/ActorTreeItem.cpp
@@ -20,6 +20,7 @@
 #include "LevelInstance/LevelInstanceActor.h"
 #include "WorldPartition/WorldPartition.h"
 #include "ToolMenu.h"
+#include "Engine/Level.h"
 
 #define LOCTEXT_NAMESPACE "SceneOutliner_ActorTreeItem"
 
@@ -383,6 +384,28 @@ TSharedRef<SWidget> FActorTreeItem::GenerateLabelWidget(ISceneOutliner& Outliner
 	return SNew(SActorTreeLabel, *this, Outliner, InRow);
 }
 
+bool FActorTreeItem::ShouldShowPinnedState() const
+{
+	if (const AActor* ActorPtr = Actor.Get())
+	{
+		const ULevel* Level = ActorPtr->GetLevel();
+		return !!Level->GetWorldPartition();
+	}
+
+	return false;
+}
+
+bool FActorTreeItem::ShouldShowVisibilityState() const
+{
+	if (const AActor* ActorPtr = Actor.Get())
+	{
+		const ULevel* Level = ActorPtr->GetLevel();
+		return Level->IsPersistentLevel() || !Level->IsInstancedLevel();
+	}
+
+	return false;
+}
+
 void FActorTreeItem::OnVisibilityChanged(const bool bNewVisibility)
 {
 	// Save the actor to the transaction buffer to support undo/redo, but do
diff --git a/Engine/Source/Editor/SceneOutliner/Public/ActorDescTreeItem.h b/Engine/Source/Editor/SceneOutliner/Public/ActorDescTreeItem.h
--- a/Engine/Source/Editor/SceneOutliner/Public/ActorDescTreeItem.h
+++ b/Engine/Source/Editor/SceneOutliner/Public/ActorDescTreeItem.h
@@ -85,6 +85,7 @@ public:
 	virtual bool HasVisibilityInfo() const override { return true; }
 	virtual bool GetVisibility() const override;
 	virtual bool ShouldShowPinnedState() const override { return true; }
+	virtual bool ShouldShowVisibilityState() const override { return false; }
 	virtual bool HasPinnedStateInfo() const override { return true; }
 	virtual bool GetPinnedState() const override;
 	/* End ISceneOutlinerTreeItem Implementation */
diff --git a/Engine/Source/Editor/SceneOutliner/Public/ActorTreeItem.h b/Engine/Source/Editor/SceneOutliner/Public/ActorTreeItem.h
--- a/Engine/Source/Editor/SceneOutliner/Public/ActorTreeItem.h
+++ b/Engine/Source/Editor/SceneOutliner/Public/ActorTreeItem.h
@@ -46,7 +46,8 @@ public:
 	virtual void OnVisibilityChanged(const bool bNewVisibility) override;
 	virtual bool HasVisibilityInfo() const override { return true; }
 	virtual bool GetVisibility() const override;
-	virtual bool ShouldShowPinnedState() const override { return true; }
+	virtual bool ShouldShowPinnedState() const override;
+	virtual bool ShouldShowVisibilityState() const override;
 	virtual bool HasPinnedStateInfo() const override { return true; };
 	virtual bool GetPinnedState() const override;
 	virtual void OnLabelChanged() override;
