SceneOutliner: Hide Pin on folders when not in WP World & Defautl root (persistent level)

diff --git a/Engine/Source/Editor/SceneOutliner/Private/ActorFolderTreeItem.cpp b/Engine/Source/Editor/SceneOutliner/Private/ActorFolderTreeItem.cpp
--- a/Engine/Source/Editor/SceneOutliner/Private/ActorFolderTreeItem.cpp
+++ b/Engine/Source/Editor/SceneOutliner/Private/ActorFolderTreeItem.cpp
@@ -12,6 +12,7 @@
 #include "ScopedTransaction.h"
 #include "ToolMenus.h"
 #include "ActorDescTreeItem.h"
+#include "Engine/World.h"
 
 #define LOCTEXT_NAMESPACE "SceneOutliner_ActorFolderTreeItem"
 
@@ -278,4 +279,9 @@ TSharedRef<SWidget> FActorFolderTreeItem::GenerateLabelWidget(ISceneOutliner& Ou
 	return SNew(SActorFolderTreeLabel, *this, Outliner, InRow);
 }
 
+bool FActorFolderTreeItem::ShouldShowPinnedState() const
+{
+	return GetRootObject() == FFolder::GetDefaultRootObject() && World.IsValid() && World->IsPartitionedWorld();
+}
+
 #undef LOCTEXT_NAMESPACE
diff --git a/Engine/Source/Editor/SceneOutliner/Public/ActorFolderTreeItem.h b/Engine/Source/Editor/SceneOutliner/Public/ActorFolderTreeItem.h
--- a/Engine/Source/Editor/SceneOutliner/Public/ActorFolderTreeItem.h
+++ b/Engine/Source/Editor/SceneOutliner/Public/ActorFolderTreeItem.h
@@ -21,6 +21,7 @@ public:
 	virtual void OnExpansionChanged() override;
 	virtual void Delete(const FFolder& InNewParentFolder) override;
 	virtual TSharedRef<SWidget> GenerateLabelWidget(ISceneOutliner& Outliner, const STableRow<FSceneOutlinerTreeItemPtr>& InRow) override;
+	virtual bool ShouldShowPinnedState() const override;
 	/* End FFolderTreeItem Implementation */
 		
 	/* Begin FFolderTreeItem Implementation */
