SceneOutliner: Fix bad LayerId increment.

diff --git a/Engine/Source/Editor/SceneOutliner/Private/SOutlinerTreeView.cpp b/Engine/Source/Editor/SceneOutliner/Private/SOutlinerTreeView.cpp
--- a/Engine/Source/Editor/SceneOutliner/Private/SOutlinerTreeView.cpp
+++ b/Engine/Source/Editor/SceneOutliner/Private/SOutlinerTreeView.cpp
@@ -395,7 +395,7 @@ int32 SSceneOutlinerTreeRow::OnPaint( const FPaintArgs& Args, const FGeometry& A
 			HighlightTargetColorAndOpacity );							// Color
 	}
 
-	return LayerId + TextLayer;
+	return StartLayer + TextLayer;
 }
 
 #undef LOCTEXT_NAMESPACE
