SceneOutliner: refactor FindParent/CreateParentItem into one method (shared logic)

diff --git a/Engine/Source/Editor/DataLayerEditor/Private/DataLayer/DataLayerActorTreeItem.h b/Engine/Source/Editor/DataLayerEditor/Private/DataLayer/DataLayerActorTreeItem.h
--- a/Engine/Source/Editor/DataLayerEditor/Private/DataLayer/DataLayerActorTreeItem.h
+++ b/Engine/Source/Editor/DataLayerEditor/Private/DataLayer/DataLayerActorTreeItem.h
@@ -37,9 +37,8 @@ public:
 	{
 	}
 
-	const UDataLayer* GetDataLayer() const { return DataLayer.Get(); }
-	UDataLayer* GetDataLayer() { return DataLayer.Get(); }
-
+	UDataLayer* GetDataLayer() const { return DataLayer.Get(); }
+	
 	const AActor* GetActor() const { return Actor.Get(); }
 	AActor* GetActor() { return Actor.Get(); }
 
diff --git a/Engine/Source/Editor/DataLayerEditor/Private/DataLayer/DataLayerHierarchy.cpp b/Engine/Source/Editor/DataLayerEditor/Private/DataLayer/DataLayerHierarchy.cpp
--- a/Engine/Source/Editor/DataLayerEditor/Private/DataLayer/DataLayerHierarchy.cpp
+++ b/Engine/Source/Editor/DataLayerEditor/Private/DataLayer/DataLayerHierarchy.cpp
@@ -202,70 +202,53 @@ void FDataLayerHierarchy::CreateItems(TArray<FSceneOutlinerTreeItemPtr>& OutItem
 	}
 }
 
-FSceneOutlinerTreeItemPtr FDataLayerHierarchy::FindParent(const ISceneOutlinerTreeItem& Item, const TMap<FSceneOutlinerTreeItemID, FSceneOutlinerTreeItemPtr>& Items) const
+FSceneOutlinerTreeItemPtr FDataLayerHierarchy::FindOrCreateParentItem(const ISceneOutlinerTreeItem& Item, const TMap<FSceneOutlinerTreeItemID, FSceneOutlinerTreeItemPtr>& Items, bool bCreate)
 {
 	if (const FDataLayerTreeItem* DataLayerTreeItem = Item.CastTo<FDataLayerTreeItem>())
 	{
-		if (const UDataLayer* DataLayer = DataLayerTreeItem->GetDataLayer())
+		if (UDataLayer* DataLayer = DataLayerTreeItem->GetDataLayer())
 		{
-			if (const UDataLayer* ParentDataLayer = DataLayer->GetParent())
+			if (UDataLayer* ParentDataLayer = DataLayer->GetParent())
 			{
 				if (const FSceneOutlinerTreeItemPtr* ParentItem = Items.Find(ParentDataLayer))
 				{
 					return *ParentItem;
 				}
+				else if (bCreate)
+				{
+					return CreateDataLayerTreeItem(ParentDataLayer, true);
+				}
 			}
 		}
 	}
 	else if (const FDataLayerActorTreeItem* DataLayerActorTreeItem = Item.CastTo<FDataLayerActorTreeItem>())
 	{
-		if (const UDataLayer* DataLayer = DataLayerActorTreeItem->GetDataLayer())
+		if (UDataLayer* DataLayer = DataLayerActorTreeItem->GetDataLayer())
 		{
 			if (const FSceneOutlinerTreeItemPtr* ParentItem = Items.Find(DataLayer))
 			{
 				return *ParentItem;
 			}
+			else if (bCreate)
+			{
+				return CreateDataLayerTreeItem(DataLayer, true);
+			}
 		}
 	}
 	else if (const FDataLayerActorDescTreeItem* DataLayerActorDescTreeItem = Item.CastTo<FDataLayerActorDescTreeItem>())
 	{
-		if (const UDataLayer* DataLayer = DataLayerActorDescTreeItem->GetDataLayer())
+		if (UDataLayer* DataLayer = DataLayerActorDescTreeItem->GetDataLayer())
 		{
 			if (const FSceneOutlinerTreeItemPtr* ParentItem = Items.Find(DataLayer))
 			{
 				return *ParentItem;
 			}
-		}
-	}
-	return nullptr;
-}
-
-FSceneOutlinerTreeItemPtr FDataLayerHierarchy::CreateParentItem(const FSceneOutlinerTreeItemPtr& Item) const
-{
-	if (FDataLayerTreeItem* DataLayerTreeItem = Item->CastTo<FDataLayerTreeItem>())
-	{
-		if (UDataLayer* DataLayer = DataLayerTreeItem->GetDataLayer())
-		{
-			if (UDataLayer* ParentDataLayer = DataLayer->GetParent())
+			else if (bCreate)
 			{
-				return CreateDataLayerTreeItem(ParentDataLayer, true);
+				return CreateDataLayerTreeItem(DataLayer, true);
 			}
 		}
 	}
-	else if (FDataLayerActorTreeItem* DataLayerActorTreeItem = Item->CastTo<FDataLayerActorTreeItem>())
-	{
-		if (UDataLayer* DataLayer = DataLayerActorTreeItem->GetDataLayer())
-		{
-			return CreateDataLayerTreeItem(DataLayer, true);
-		}
-	}
-	else if (FDataLayerActorDescTreeItem* DataLayerActorDescTreeItem = Item->CastTo<FDataLayerActorDescTreeItem>())
-	{
-		if (UDataLayer* DataLayer = DataLayerActorDescTreeItem->GetDataLayer())
-		{
-			return CreateDataLayerTreeItem(DataLayer, true);
-		}
-	}
 	return nullptr;
 }
 
diff --git a/Engine/Source/Editor/DataLayerEditor/Private/DataLayer/DataLayerHierarchy.h b/Engine/Source/Editor/DataLayerEditor/Private/DataLayer/DataLayerHierarchy.h
--- a/Engine/Source/Editor/DataLayerEditor/Private/DataLayer/DataLayerHierarchy.h
+++ b/Engine/Source/Editor/DataLayerEditor/Private/DataLayer/DataLayerHierarchy.h
@@ -17,10 +17,9 @@ class FDataLayerHierarchy : public ISceneOutlinerHierarchy
 public:
 	virtual ~FDataLayerHierarchy();
 	static TUniquePtr<FDataLayerHierarchy> Create(FDataLayerMode* Mode, const TWeakObjectPtr<UWorld>& World);
-	virtual FSceneOutlinerTreeItemPtr FindParent(const ISceneOutlinerTreeItem& Item, const TMap<FSceneOutlinerTreeItemID, FSceneOutlinerTreeItemPtr>& Items) const override;
 	virtual void CreateItems(TArray<FSceneOutlinerTreeItemPtr>& OutItems) const override;
 	virtual void CreateChildren(const FSceneOutlinerTreeItemPtr& Item, TArray<FSceneOutlinerTreeItemPtr>& OutChildren) const override {}
-	virtual FSceneOutlinerTreeItemPtr CreateParentItem(const FSceneOutlinerTreeItemPtr& Item) const override;
+	virtual FSceneOutlinerTreeItemPtr FindOrCreateParentItem(const ISceneOutlinerTreeItem& Item, const TMap<FSceneOutlinerTreeItemID, FSceneOutlinerTreeItemPtr>& Items, bool bCreate = false) override;
 	void SetShowEditorDataLayers(bool bInShowEditorDataLayers) { bShowEditorDataLayers = bInShowEditorDataLayers; }
 	void SetShowRuntimeDataLayers(bool bInShowRuntimeDataLayers) { bShowRuntimeDataLayers = bInShowRuntimeDataLayers; }
 	void SetShowDataLayerActors(bool bInShowDataLayerActors) { bShowDataLayerActors = bInShowDataLayerActors; }
diff --git a/Engine/Source/Editor/DataLayerEditor/Private/DataLayer/DataLayersActorDescTreeItem.h b/Engine/Source/Editor/DataLayerEditor/Private/DataLayer/DataLayersActorDescTreeItem.h
--- a/Engine/Source/Editor/DataLayerEditor/Private/DataLayer/DataLayersActorDescTreeItem.h
+++ b/Engine/Source/Editor/DataLayerEditor/Private/DataLayer/DataLayersActorDescTreeItem.h
@@ -41,9 +41,8 @@ public:
 	{
 	}
 
-	const UDataLayer* GetDataLayer() const { return DataLayer.Get(); }
-	UDataLayer* GetDataLayer() { return DataLayer.Get(); }
-
+	UDataLayer* GetDataLayer() const { return DataLayer.Get(); }
+	
 	static uint32 ComputeTreeItemID(const FGuid& InActorGuid, const UDataLayer* InDataLayer)
 	{
 		return HashCombine(GetTypeHash(InActorGuid), GetTypeHash(FObjectKey(InDataLayer)));
diff --git a/Engine/Source/Editor/SceneOutliner/Private/ActorFolderHierarchy.cpp b/Engine/Source/Editor/SceneOutliner/Private/ActorFolderHierarchy.cpp
--- a/Engine/Source/Editor/SceneOutliner/Private/ActorFolderHierarchy.cpp
+++ b/Engine/Source/Editor/SceneOutliner/Private/ActorFolderHierarchy.cpp
@@ -19,7 +19,7 @@ FActorFolderHierarchy::FActorFolderHierarchy(ISceneOutlinerMode* InMode, const T
 	check(Mode->ShouldShowFolders());
 }
 
-FSceneOutlinerTreeItemPtr FActorFolderHierarchy::FindParent(const ISceneOutlinerTreeItem& Item, const TMap<FSceneOutlinerTreeItemID, FSceneOutlinerTreeItemPtr>& Items) const
+FSceneOutlinerTreeItemPtr FActorFolderHierarchy::FindOrCreateParentItem(const ISceneOutlinerTreeItem& Item, const TMap<FSceneOutlinerTreeItemID, FSceneOutlinerTreeItemPtr>& Items, bool bCreate)
 {
 	if (Item.IsA<FWorldTreeItem>())
 	{
@@ -29,29 +29,53 @@ FSceneOutlinerTreeItemPtr FActorFolderHierarchy::FindParent(const ISceneOutliner
 	{
 		const FFolder ParentPath = ActorFolderItem->GetFolder().GetParent();
 
-		const FSceneOutlinerTreeItemPtr* ParentItem = nullptr;
 		// If the folder has no parent path
 		if (ParentPath.IsNone())
 		{
 			if (UObject* Object = ParentPath.GetRootObjectPtr())
 			{
-				ParentItem = Items.Find(Object);
+				if (const FSceneOutlinerTreeItemPtr* ParentItem = Items.Find(Object))
+				{
+					return *ParentItem;
+				}
+				else if (bCreate)
+				{
+					FFolder Folder = ActorFolderItem->GetFolder();
+					if (Folder.GetRootObject() == RootObject)
+					{
+						// If item belongs to a LevelInstance
+						if (ALevelInstance* RootLevelInstance = Cast<ALevelInstance>(Folder.GetRootObjectPtr()))
+						{
+							return Mode->CreateItemFor<FActorTreeItem>(RootLevelInstance, true);
+						}
+					}
+				}
 			}
 			else
 			{
-				ParentItem = Items.Find(ActorFolderItem->World.Get());
+				if (const FSceneOutlinerTreeItemPtr* ParentItem = Items.Find(ActorFolderItem->World.Get()))
+				{
+					return *ParentItem;
+				}
+				else if (bCreate)
+				{
+					// Parent World
+					UWorld* OwningWorld = ActorFolderItem->World.Get();
+					check(OwningWorld);
+					return Mode->CreateItemFor<FWorldTreeItem>(OwningWorld, true);
+				}
 			}
 		}
-		else
+		else if (const FSceneOutlinerTreeItemPtr* ParentItem = Items.Find(ParentPath))
 		{
-			ParentItem = Items.Find(ParentPath);
+			return *ParentItem;
 		}
-
-		if (ParentItem)
+		else if(bCreate)
 		{
-			return *ParentItem;
+			return Mode->CreateItemFor<FActorFolderTreeItem>(FActorFolderTreeItem(ParentPath, ActorFolderItem->World), true);
 		}
 	}
+
 	return nullptr;
 }
 
@@ -143,42 +167,3 @@ void FActorFolderHierarchy::CreateChildren(const FSceneOutlinerTreeItemPtr& Item
 	}
 }
 
-FSceneOutlinerTreeItemPtr FActorFolderHierarchy::CreateParentItem(const FSceneOutlinerTreeItemPtr& Item) const
-{
-	if (Item->IsA<FWorldTreeItem>())
-	{
-		return nullptr;
-	}
-	else if (const FActorFolderTreeItem* FolderTreeItem = Item->CastTo<FActorFolderTreeItem>())
-	{
-		FFolder Folder = FolderTreeItem->GetFolder();
-
-		// Parent Folder
-		const FFolder ParentFolder = Folder.GetParent();
-		if (!ParentFolder.IsNone())
-		{
-			return Mode->CreateItemFor<FActorFolderTreeItem>(FActorFolderTreeItem(ParentFolder, FolderTreeItem->World), true);
-		}
-
-		if (FFolder::HasRootObject(RootObject))
-		{
-			// Parent Object
-			if (Folder.GetRootObject() == RootObject)
-			{
-				// If item belongs to a LevelInstance
-				if (ALevelInstance* RootLevelInstance = Cast<ALevelInstance>(Folder.GetRootObjectPtr()))
-				{
-					return Mode->CreateItemFor<FActorTreeItem>(RootLevelInstance, true);
-				}
-			}
-		}
-		else
-		{
-			// Parent World
-			UWorld* OwningWorld = FolderTreeItem->World.Get();
-			check(OwningWorld);
-			return Mode->CreateItemFor<FWorldTreeItem>(OwningWorld, true);
-		}
-	}
-	return nullptr;
-}
diff --git a/Engine/Source/Editor/SceneOutliner/Private/ActorHierarchy.cpp b/Engine/Source/Editor/SceneOutliner/Private/ActorHierarchy.cpp
--- a/Engine/Source/Editor/SceneOutliner/Private/ActorHierarchy.cpp
+++ b/Engine/Source/Editor/SceneOutliner/Private/ActorHierarchy.cpp
@@ -128,7 +128,7 @@ static FName GetFolderPathFromActorDesc(UWorld* InWorld, const FWorldPartitionAc
 	return InActorDesc->GetFolderPath();
 }
 
-FSceneOutlinerTreeItemPtr FActorHierarchy::FindParent(const ISceneOutlinerTreeItem& Item, const TMap<FSceneOutlinerTreeItemID, FSceneOutlinerTreeItemPtr>& Items) const
+FSceneOutlinerTreeItemPtr FActorHierarchy::FindOrCreateParentItem(const ISceneOutlinerTreeItem& Item, const TMap<FSceneOutlinerTreeItemID, FSceneOutlinerTreeItemPtr>& Items, bool bCreate)
 {
 	if (Item.IsA<FWorldTreeItem>())
 	{
@@ -139,16 +139,16 @@ FSceneOutlinerTreeItemPtr FActorHierarchy::FindParent(const ISceneOutlinerTreeIt
 		if (AActor* Actor = ActorTreeItem->Actor.Get())
 		{
 			// Parent Actor (Actor attachement / parenting)
-			if (const AActor* ParentActor = Actor->GetSceneOutlinerParent())
+			if (AActor* ParentActor = Actor->GetSceneOutlinerParent())
 			{
 				if (const FSceneOutlinerTreeItemPtr* ParentItem = Items.Find(ParentActor))
 				{
 					return *ParentItem;
 				}
 				// If Parent can be listed in SceneOutliner return nullptr so it gets created
-				else if(ParentActor->IsListedInSceneOutliner())
+				else if (ParentActor->IsListedInSceneOutliner())
 				{
-					return nullptr;
+					return bCreate ? Mode->CreateItemFor<FActorTreeItem>(ParentActor, true) : nullptr;
 				}
 			}
 
@@ -162,25 +162,25 @@ FSceneOutlinerTreeItemPtr FActorHierarchy::FindParent(const ISceneOutlinerTreeIt
 				}
 				else
 				{
-					return nullptr;
+					return bCreate ? Mode->CreateItemFor<FActorFolderTreeItem>(FActorFolderTreeItem(Folder, ActorTreeItem->Actor->GetWorld()), true) : nullptr;
 				}
 			}
 
 			// Parent Level Instance
-			if (const ALevelInstance* OwningLevelInstance = Cast<ALevelInstance>(Folder.GetRootObjectPtr()))
+			if (ALevelInstance* OwningLevelInstance = Cast<ALevelInstance>(Folder.GetRootObjectPtr()))
 			{
 				const ALevelInstance* LevelInstanceActor = Cast<ALevelInstance>(Actor);
 				const bool bIsAnEditingLevelInstance = LevelInstanceActor ? LevelInstanceActor->IsEditing() : false;
 				// Parent this to a LevelInstance if the parent LevelInstance is being edited or if this is a sub LevelInstance which is being edited
 				if (bShowingLevelInstances || (OwningLevelInstance->IsEditing() || bIsAnEditingLevelInstance))
 				{
-					if (const FSceneOutlinerTreeItemPtr* OwningLevelInstanceItem = Items.Find(OwningLevelInstance))
+					if (const FSceneOutlinerTreeItemPtr* ParentItem = Items.Find(OwningLevelInstance))
 					{
-						return *OwningLevelInstanceItem;
+						return *ParentItem;
 					}
 					else
 					{
-						return nullptr;
+						return bCreate ? Mode->CreateItemFor<FActorTreeItem>(OwningLevelInstance, true) : nullptr;
 					}
 				}
 			}
@@ -194,18 +194,12 @@ FSceneOutlinerTreeItemPtr FActorHierarchy::FindParent(const ISceneOutlinerTreeIt
 				}
 				else
 				{
-					return nullptr;
+					return bCreate ? Mode->CreateItemFor<FLevelTreeItem>(OwningLevel, true) : nullptr;
 				}
 			}
-
-			// Parent world
-			if (const FSceneOutlinerTreeItemPtr* ParentItem = Items.Find(RepresentingWorld.Get()))
-			{
-				return *ParentItem;
-			}
 		}
 	}
-	else if (const FFolderTreeItem* FolderItem = Item.CastTo<FFolderTreeItem>())
+	else if (const FActorFolderTreeItem* FolderItem = Item.CastTo<FActorFolderTreeItem>())
 	{
 		// We should never call FindParents on a folder item if folders are not being shown
 		check(Mode->ShouldShowFolders());
@@ -219,19 +213,23 @@ FSceneOutlinerTreeItemPtr FActorHierarchy::FindParent(const ISceneOutlinerTreeIt
 			{
 				return *ParentItem;
 			}
+			else
+			{
+				return bCreate ? Mode->CreateItemFor<FActorFolderTreeItem>(FActorFolderTreeItem(ParentPath, FolderItem->World), true) : nullptr;
+			}
 		}
 		// Parent Level Instance
 		else if (ALevelInstance* OwningLevelInstance = Cast<ALevelInstance>(ParentPath.GetRootObjectPtr()))
 		{
 			if (bShowingLevelInstances || OwningLevelInstance->IsEditing())
 			{
-				if (const FSceneOutlinerTreeItemPtr* OwningLevelInstanceItem = Items.Find(OwningLevelInstance))
+				if (const FSceneOutlinerTreeItemPtr* ParentItem = Items.Find(OwningLevelInstance))
 				{
-					return *OwningLevelInstanceItem;
+					return *ParentItem;
 				}
 				else
 				{
-					return nullptr;
+					return bCreate ? Mode->CreateItemFor<FActorTreeItem>(OwningLevelInstance, true) : nullptr;
 				}
 			}
 		}
@@ -244,33 +242,29 @@ FSceneOutlinerTreeItemPtr FActorHierarchy::FindParent(const ISceneOutlinerTreeIt
 			}
 			else
 			{
-				return nullptr;
+				return bCreate ? Mode->CreateItemFor<FLevelTreeItem>(OwningLevel, true) : nullptr;
 			}
 		}
-		// Parent World
-		else if (const FSceneOutlinerTreeItemPtr* WorldItem = Items.Find(RepresentingWorld.Get()))
-		{
-			return *WorldItem;
-		}
-
-		return nullptr;
 	}
 	else if (const FComponentTreeItem* ComponentTreeItem = Item.CastTo<FComponentTreeItem>())
 	{
-		const AActor* Owner = ComponentTreeItem->Component->GetOwner();
-		if (Owner)
+		if (AActor* Owner = ComponentTreeItem->Component->GetOwner())
 		{
 			if (const FSceneOutlinerTreeItemPtr* ParentItem = Items.Find(Owner))
 			{
 				return *ParentItem;
 			}
+			else
+			{
+				return bCreate ? Mode->CreateItemFor<FActorTreeItem>(Owner, true) : nullptr;
+			}
 		}
+		// do not default to world on Component items
+		return nullptr;
 	}
 	else if (const FActorDescTreeItem* ActorDescItem = Item.CastTo<FActorDescTreeItem>())
 	{
-		const FWorldPartitionActorDesc* ActorDesc = ActorDescItem->ActorDescHandle.GetActorDesc();
-
-		if (ActorDesc)
+		if (const FWorldPartitionActorDesc* ActorDesc = ActorDescItem->ActorDescHandle.GetActorDesc())
 		{
 			const FName FolderPath = GetFolderPathFromActorDesc(RepresentingWorld.Get(), ActorDesc);
 			if (!FolderPath.IsNone())
@@ -279,16 +273,23 @@ FSceneOutlinerTreeItemPtr FActorHierarchy::FindParent(const ISceneOutlinerTreeIt
 				{
 					return *ParentItem;
 				}
-				return nullptr;
+				else
+				{
+					return bCreate ? Mode->CreateItemFor<FActorFolderTreeItem>(FActorFolderTreeItem(FFolder(FolderPath), RepresentingWorld.Get()), true) : nullptr;
+				}
 			}
 		}
-		// Default to the world
-		if (const FSceneOutlinerTreeItemPtr* ParentItem = Items.Find(RepresentingWorld.Get()))
-		{
-			return *ParentItem;
-		}
 	}
-	return nullptr;
+
+	// If we get here return world item
+	if (const FSceneOutlinerTreeItemPtr* ParentItem = Items.Find(RepresentingWorld.Get()))
+	{
+		return *ParentItem;
+	}
+	else
+	{
+		return bCreate ? Mode->CreateItemFor<FWorldTreeItem>(RepresentingWorld.Get(), true) : nullptr;
+	}
 }
 
 void FActorHierarchy::CreateComponentItems(const AActor* Actor, TArray<FSceneOutlinerTreeItemPtr>& OutItems) const
@@ -492,116 +493,6 @@ void FActorHierarchy::CreateChildren(const FSceneOutlinerTreeItemPtr& Item, TArr
 	}
 }
 
-FSceneOutlinerTreeItemPtr FActorHierarchy::CreateParentItem(const FSceneOutlinerTreeItemPtr& Item) const
-{
-	if (Item->IsA<FWorldTreeItem>())
-	{
-		return nullptr;
-	}
-	else if (const FActorTreeItem* ActorTreeItem = Item->CastTo<FActorTreeItem>())
-	{
-		if (const AActor* Actor = ActorTreeItem->Actor.Get())
-		{
-			// Parent Actor (Actor attachement / parenting)
-			if (AActor* ParentActor = Actor->GetSceneOutlinerParent())
-			{
-				if (ParentActor->IsListedInSceneOutliner())
-				{
-					return Mode->CreateItemFor<FActorTreeItem>(ParentActor, true);
-				}
-			}
-
-			// Parent Folder
-			FFolder Folder = Actor->GetFolder();
-			if (Mode->ShouldShowFolders() && !Folder.IsNone())
-			{
-				return Mode->CreateItemFor<FActorFolderTreeItem>(FActorFolderTreeItem(Folder, ActorTreeItem->Actor->GetWorld()), true);
-			}
-
-			// Parent Object
-			if (ALevelInstance* OwningLevelInstance = Cast<ALevelInstance>(Folder.GetRootObjectPtr()))
-			{
-				const ULevelInstanceSubsystem* LevelInstanceSubsystem = RepresentingWorld->GetSubsystem<ULevelInstanceSubsystem>();
-				ALevelInstance* ParentLevelInstance = LevelInstanceSubsystem ? LevelInstanceSubsystem->GetParentLevelInstance(Actor) : nullptr;
-				check(OwningLevelInstance == ParentLevelInstance);
-				{
-					const ALevelInstance* LevelInstanceActor = Cast<ALevelInstance>(Actor);
-					const bool bIsAnEditingLevelInstance = LevelInstanceActor ? LevelInstanceActor->IsEditing() : false;
-					if (bShowingLevelInstances || (ParentLevelInstance->IsEditing() || bIsAnEditingLevelInstance))
-					{
-						return Mode->CreateItemFor<FActorTreeItem>(ParentLevelInstance, true);
-					}
-				}
-			}
-
-			// Parent Level Using Actor Folders
-			if (ULevel* OwningLevel = Cast<ULevel>(Folder.GetRootObjectPtr()))
-			{
-				return Mode->CreateItemFor<FLevelTreeItem>(OwningLevel, true);
-			}
-
-			// Parent World
-			UWorld* OwningWorld = ActorTreeItem->Actor->GetWorld();
-			check(OwningWorld);
-			return Mode->CreateItemFor<FWorldTreeItem>(OwningWorld, true);
-		}
-	}
-	else if (const FComponentTreeItem* ComponentTreeItem = Item->CastTo<FComponentTreeItem>())
-	{
-		if (AActor* ParentActor = ComponentTreeItem->Component->GetOwner())
-		{
-			return Mode->CreateItemFor<FActorTreeItem>(ParentActor, true);
-		}
-	}
-	else if (const FActorFolderTreeItem* FolderTreeItem = Item->CastTo<FActorFolderTreeItem>())
-	{
-		check(Mode->ShouldShowFolders());
-
-		FFolder Folder = FolderTreeItem->GetFolder();
-
-		// Parent Folder
-		const FFolder ParentFolder = Folder.GetParent();
-		if (!ParentFolder.IsNone())
-		{
-			return Mode->CreateItemFor<FActorFolderTreeItem>(FActorFolderTreeItem(ParentFolder, FolderTreeItem->World), true);
-		}
-
-		// Parent Level Instance
-		if (ALevelInstance* OwningLevelInstance = Cast<ALevelInstance>(Folder.GetRootObjectPtr()))
-		{
-			if (bShowingLevelInstances || OwningLevelInstance->IsEditing())
-			{
-				return Mode->CreateItemFor<FActorTreeItem>(OwningLevelInstance, true);
-			}
-		}
-
-		// Parent Level Using Actor Folders
-		if (ULevel* OwningLevel = Cast<ULevel>(Folder.GetRootObjectPtr()))
-		{
-			return Mode->CreateItemFor<FLevelTreeItem>(OwningLevel, true);
-		}
-
-		// Parent World
-		UWorld* OwningWorld = FolderTreeItem->World.Get();
-		check(OwningWorld);
-		return Mode->CreateItemFor<FWorldTreeItem>(OwningWorld, true);
-	}
-	else if (const FActorDescTreeItem* ActorDescItem = Item->CastTo<FActorDescTreeItem>())
-	{
-		if (const FWorldPartitionActorDesc* ActorDesc = ActorDescItem->ActorDescHandle.GetActorDesc())
-		{
-			UWorld* World = RepresentingWorld.Get();
-			const FName FolderPath = GetFolderPathFromActorDesc(World, ActorDesc);
-			if (World && Mode->ShouldShowFolders() && !FolderPath.IsNone())
-			{
-				return Mode->CreateItemFor<FActorFolderTreeItem>(FActorFolderTreeItem(FFolder(FolderPath), World), true);
-			}
-		}
-	}
-
-	return nullptr;
-}
-
 void FActorHierarchy::FullRefreshEvent()
 {
 	FSceneOutlinerHierarchyChangedData EventData;
diff --git a/Engine/Source/Editor/SceneOutliner/Private/SSceneOutliner.cpp b/Engine/Source/Editor/SceneOutliner/Private/SSceneOutliner.cpp
--- a/Engine/Source/Editor/SceneOutliner/Private/SSceneOutliner.cpp
+++ b/Engine/Source/Editor/SceneOutliner/Private/SSceneOutliner.cpp
@@ -751,7 +751,7 @@ FSceneOutlinerTreeItemPtr SSceneOutliner::EnsureParentForItem(FSceneOutlinerTree
 {
 	if (SharedData->bShowParentTree)
 	{
-		FSceneOutlinerTreeItemPtr Parent = Mode->GetHierarchy()->FindParent(*Item, TreeItemMap);
+		FSceneOutlinerTreeItemPtr Parent = Mode->GetHierarchy()->FindOrCreateParentItem(*Item, TreeItemMap, /*bCreate=*/false);
 		if (Parent.IsValid())
 		{
 			return Parent;
@@ -759,13 +759,7 @@ FSceneOutlinerTreeItemPtr SSceneOutliner::EnsureParentForItem(FSceneOutlinerTree
 		else
 		{
 			// Try to find the parent in the pending items
-			Parent = Mode->GetHierarchy()->FindParent(*Item, PendingTreeItemMap);
-			if (!Parent.IsValid())
-			{
-				// If there isn't already a parent for this item, try to create one for it
-				Parent = Mode->GetHierarchy()->CreateParentItem(Item);
-			}
-				
+			Parent = Mode->GetHierarchy()->FindOrCreateParentItem(*Item, PendingTreeItemMap, /*bCreate=*/true);				
 			if (Parent.IsValid())
 			{
 				AddUnfilteredItemToTree(Parent.ToSharedRef());
@@ -2201,10 +2195,10 @@ void SSceneOutliner::UnpinSelectedItems()
 
 FSceneOutlinerTreeItemPtr SSceneOutliner::FindParent(const ISceneOutlinerTreeItem& InItem) const
 {
-	FSceneOutlinerTreeItemPtr Parent = Mode->GetHierarchy()->FindParent(InItem, TreeItemMap);
+	FSceneOutlinerTreeItemPtr Parent = Mode->GetHierarchy()->FindOrCreateParentItem(InItem, TreeItemMap, /*bCreate=*/false);
 	if (!Parent.IsValid())
 	{
-		Parent = Mode->GetHierarchy()->FindParent(InItem, PendingTreeItemMap);
+		Parent = Mode->GetHierarchy()->FindOrCreateParentItem(InItem, PendingTreeItemMap, /*bCreate=*/false);
 	}
 	return Parent;
 }
diff --git a/Engine/Source/Editor/SceneOutliner/Public/ActorFolderHierarchy.h b/Engine/Source/Editor/SceneOutliner/Public/ActorFolderHierarchy.h
--- a/Engine/Source/Editor/SceneOutliner/Public/ActorFolderHierarchy.h
+++ b/Engine/Source/Editor/SceneOutliner/Public/ActorFolderHierarchy.h
@@ -13,10 +13,9 @@ public:
 
 
 	/* Begin ISceneOutlinerHierarchy Implementation */
-	virtual FSceneOutlinerTreeItemPtr FindParent(const ISceneOutlinerTreeItem& Item, const TMap<FSceneOutlinerTreeItemID, FSceneOutlinerTreeItemPtr>& Items) const override;
 	virtual void CreateItems(TArray<FSceneOutlinerTreeItemPtr>& OutItems) const override;
 	virtual void CreateChildren(const FSceneOutlinerTreeItemPtr& Item, TArray<FSceneOutlinerTreeItemPtr>& OutChildren) const override;
-	virtual FSceneOutlinerTreeItemPtr CreateParentItem(const FSceneOutlinerTreeItemPtr& Item) const override;
+	virtual FSceneOutlinerTreeItemPtr FindOrCreateParentItem(const ISceneOutlinerTreeItem& Item, const TMap<FSceneOutlinerTreeItemID, FSceneOutlinerTreeItemPtr>& Items, bool bCreate = false) override;
 	/* End ISceneOutlinerHierarchy Implementation */
 private:
 	/** Adds all the direct and indirect children of a world to OutItems */
diff --git a/Engine/Source/Editor/SceneOutliner/Public/ActorHierarchy.h b/Engine/Source/Editor/SceneOutliner/Public/ActorHierarchy.h
--- a/Engine/Source/Editor/SceneOutliner/Public/ActorHierarchy.h
+++ b/Engine/Source/Editor/SceneOutliner/Public/ActorHierarchy.h
@@ -16,15 +16,12 @@ public:
 
 	static TUniquePtr<FActorHierarchy> Create(ISceneOutlinerMode* Mode, const TWeakObjectPtr<UWorld>& World);
 
-	/** Find the parent of a tree item in an existing item map, it if exists. */
-	virtual FSceneOutlinerTreeItemPtr FindParent(const ISceneOutlinerTreeItem& Item, const TMap<FSceneOutlinerTreeItemID, FSceneOutlinerTreeItemPtr>& Items) const override;
-
 	/** Create a linearization of all applicable items in the hierarchy */
 	virtual void CreateItems(TArray<FSceneOutlinerTreeItemPtr>& OutItems) const override;
 	/** Create a linearization of all direct and indirect children of a given item in the hierarchy */
 	virtual void CreateChildren(const FSceneOutlinerTreeItemPtr& Item, TArray<FSceneOutlinerTreeItemPtr>& OutChildren) const override;
 	/** Forcibly create a parent item for a given tree item */
-	virtual FSceneOutlinerTreeItemPtr CreateParentItem(const FSceneOutlinerTreeItemPtr& Item) const override;
+	virtual FSceneOutlinerTreeItemPtr FindOrCreateParentItem(const ISceneOutlinerTreeItem& Item, const TMap<FSceneOutlinerTreeItemID, FSceneOutlinerTreeItemPtr>& Items, bool bCreate = false) override;
 
 	void SetShowingComponents(bool bInShowingComponents) { bShowingComponents = bInShowingComponents; }
 	void SetShowingOnlyActorWithValidComponents(bool bInShowingOnlyActorWithValidComponents) { bShowingOnlyActorWithValidComponents = bInShowingOnlyActorWithValidComponents; }
diff --git a/Engine/Source/Editor/SceneOutliner/Public/ISceneOutlinerHierarchy.h b/Engine/Source/Editor/SceneOutliner/Public/ISceneOutlinerHierarchy.h
--- a/Engine/Source/Editor/SceneOutliner/Public/ISceneOutlinerHierarchy.h
+++ b/Engine/Source/Editor/SceneOutliner/Public/ISceneOutlinerHierarchy.h
@@ -13,16 +13,13 @@ class ISceneOutlinerHierarchy : public TSharedFromThis<ISceneOutlinerHierarchy>
 public:
 	ISceneOutlinerHierarchy(ISceneOutlinerMode* InMode) : Mode(InMode) {}
 	virtual ~ISceneOutlinerHierarchy() { HierarchyChangedEvent.Clear(); }
-
-	/** Find the parent of a tree item in an existing item map, it if exists. */
-	virtual FSceneOutlinerTreeItemPtr FindParent(const ISceneOutlinerTreeItem& Item, const TMap<FSceneOutlinerTreeItemID, FSceneOutlinerTreeItemPtr>& Items) const = 0;
-		
+			
 	/** Create a linearization of all applicable items in the hierarchy */
 	virtual void CreateItems(TArray<FSceneOutlinerTreeItemPtr>& OutItems) const = 0;
 	/** Create a linearization of all direct and indirect children of a given item in the hierarchy */
 	virtual void CreateChildren(const FSceneOutlinerTreeItemPtr& Item, TArray<FSceneOutlinerTreeItemPtr>& OutChildren) const = 0;
-	/** Forcibly create a parent item for a given tree item */
-	virtual FSceneOutlinerTreeItemPtr CreateParentItem(const FSceneOutlinerTreeItemPtr& Item) const = 0;
+	/** Find or optionally create a parent item for a given tree item */
+	virtual FSceneOutlinerTreeItemPtr FindOrCreateParentItem(const ISceneOutlinerTreeItem& Item, const TMap<FSceneOutlinerTreeItemID, FSceneOutlinerTreeItemPtr>& Items, bool bCreate = false) = 0;
 
 	DECLARE_EVENT_OneParam(ISceneOutlinerHierarchy, FHierarchyChangedEvent, FSceneOutlinerHierarchyChangedData)
 	FHierarchyChangedEvent& OnHierarchyChanged() { return HierarchyChangedEvent; }
