World Partition: cleanup scene outliner code handling actor descriptors following 18564793.

diff --git a/Engine/Source/Editor/DataLayerEditor/Private/DataLayer/DataLayersActorDescTreeItem.h b/Engine/Source/Editor/DataLayerEditor/Private/DataLayer/DataLayersActorDescTreeItem.h
--- a/Engine/Source/Editor/DataLayerEditor/Private/DataLayer/DataLayersActorDescTreeItem.h
+++ b/Engine/Source/Editor/DataLayerEditor/Private/DataLayer/DataLayersActorDescTreeItem.h
@@ -50,12 +50,12 @@ public:
 
 	bool Filter(FFilterPredicate Pred) const
 	{
-		return Pred.Execute(ActorDescHandle.GetActorDesc(), DataLayer.Get());
+		return Pred.Execute(ActorDescHandle.Get(), DataLayer.Get());
 	}
 
 	bool GetInteractiveState(FInteractivePredicate Pred) const
 	{
-		return Pred.Execute(ActorDescHandle.GetActorDesc(), DataLayer.Get());
+		return Pred.Execute(ActorDescHandle.Get(), DataLayer.Get());
 	}
 
 	/* Begin ISceneOutlinerTreeItem Implementation */
diff --git a/Engine/Source/Editor/DataLayerEditor/Private/DataLayer/SDataLayerBrowser.cpp b/Engine/Source/Editor/DataLayerEditor/Private/DataLayer/SDataLayerBrowser.cpp
--- a/Engine/Source/Editor/DataLayerEditor/Private/DataLayer/SDataLayerBrowser.cpp
+++ b/Engine/Source/Editor/DataLayerEditor/Private/DataLayer/SDataLayerBrowser.cpp
@@ -64,7 +64,7 @@ void SDataLayerBrowser::Construct(const FArguments& InArgs)
 		}
 		else if (const FDataLayerActorDescTreeItem* ActorDescItem = Item.CastTo<FDataLayerActorDescTreeItem>())
 		{
-			if (const FWorldPartitionActorDesc* ActorDesc = ActorDescItem->ActorDescHandle.GetActorDesc())
+			if (const FWorldPartitionActorDesc* ActorDesc = ActorDescItem->ActorDescHandle.Get())
 			{
 				return ActorDesc->GetActorName().ToString();
 			}
diff --git a/Engine/Source/Editor/SceneOutliner/Private/ActorBrowsingMode.cpp b/Engine/Source/Editor/SceneOutliner/Private/ActorBrowsingMode.cpp
--- a/Engine/Source/Editor/SceneOutliner/Private/ActorBrowsingMode.cpp
+++ b/Engine/Source/Editor/SceneOutliner/Private/ActorBrowsingMode.cpp
@@ -749,7 +749,7 @@ void FActorBrowsingMode::OnItemPassesFilters(const ISceneOutlinerTreeItem& Item)
 	}
 	else if (const FActorDescTreeItem* const ActorDescItem = Item.CastTo<FActorDescTreeItem>(); ActorDescItem && ActorDescItem->IsValid())
 	{
-		ApplicableUnloadedActors.Add(ActorDescItem->ActorDescHandle.GetActorDesc());
+		ApplicableUnloadedActors.Add(ActorDescItem->ActorDescHandle.Get());
 	}
 }
 
diff --git a/Engine/Source/Editor/SceneOutliner/Private/ActorDescTreeItem.cpp b/Engine/Source/Editor/SceneOutliner/Private/ActorDescTreeItem.cpp
--- a/Engine/Source/Editor/SceneOutliner/Private/ActorDescTreeItem.cpp
+++ b/Engine/Source/Editor/SceneOutliner/Private/ActorDescTreeItem.cpp
@@ -113,7 +113,7 @@ private:
 	{
 		if (TSharedPtr<FActorDescTreeItem> TreeItem = TreeItemPtr.Pin())
 		{
-			if (const FWorldPartitionActorDesc* ActorDesc = TreeItem->ActorDescHandle.GetActorDesc())
+			if (const FWorldPartitionActorDesc* ActorDesc = TreeItem->ActorDescHandle.Get())
 			{
 				FFormatNamedArguments Args;
 				Args.Add(TEXT("ActorLabel"), FText::FromName(ActorDesc->GetActorLabel()));
@@ -133,7 +133,7 @@ private:
 	{
 		if (TSharedPtr<FActorDescTreeItem> TreeItem = TreeItemPtr.Pin())
 		{
-			if (const FWorldPartitionActorDesc* ActorDesc = TreeItem->ActorDescHandle.GetActorDesc())
+			if (const FWorldPartitionActorDesc* ActorDesc = TreeItem->ActorDescHandle.Get())
 			{
 				return FText::FromName(ActorDesc->GetActorClass()->GetFName());
 			}
@@ -153,7 +153,7 @@ private:
 
 		if (TreeItem.IsValid() && WeakSceneOutliner.IsValid())
 		{
-			if (const FWorldPartitionActorDesc* ActorDesc = TreeItem->ActorDescHandle.GetActorDesc())
+			if (const FWorldPartitionActorDesc* ActorDesc = TreeItem->ActorDescHandle.Get())
 			{
 				const FName IconName = ActorDesc->GetActorClass()->GetFName();
 
@@ -212,12 +212,12 @@ private:
 
 FActorDescTreeItem::FActorDescTreeItem(const FGuid& InActorGuid, UActorDescContainer* Container)
 	: ISceneOutlinerTreeItem(Type)
-	, ActorDescHandle(InActorGuid, Container)
+	, ActorDescHandle(Container, InActorGuid)
 	, ActorGuid(InActorGuid)
 {
 	ID = FSceneOutlinerTreeItemID(InActorGuid);
 
-	if (const FWorldPartitionActorDesc* const ActorDesc = ActorDescHandle.GetActorDesc())
+	if (const FWorldPartitionActorDesc* const ActorDesc = ActorDescHandle.Get())
 	{
 		DisplayString = ActorDesc->GetActorLabel().ToString();
 	}
@@ -249,7 +249,7 @@ TSharedRef<SWidget> FActorDescTreeItem::GenerateLabelWidget(ISceneOutliner& Outl
 
 void FActorDescTreeItem::FocusActorBounds() const
 {
-	if (FWorldPartitionActorDesc const* ActorDesc = ActorDescHandle.GetActorDesc())
+	if (FWorldPartitionActorDesc const* ActorDesc = ActorDescHandle.Get())
 	{
 		if (FLevelEditorViewportClient* LevelViewportClient = GCurrentLevelEditingViewportClient)
 		{
@@ -264,11 +264,6 @@ void FActorDescTreeItem::GenerateContextMenu(UToolMenu* Menu, SSceneOutliner&)
 	Section.AddMenuEntry("FocusActorBounds", LOCTEXT("FocusActorBounds", "Focus Actor Bounds"), FText(), FSlateIcon(), FUIAction(FExecuteAction::CreateSP(this, &FActorDescTreeItem::FocusActorBounds)));
 }
 
-void FActorDescTreeItem::OnVisibilityChanged(const bool bNewVisibility)
-{
-
-}
-
 bool FActorDescTreeItem::GetVisibility() const
 {
 	return true;
@@ -278,7 +273,7 @@ bool FActorDescTreeItem::GetPinnedState() const
 {
 	if (ActorDescHandle.IsValid())
 	{
-		return ActorDescHandle.GetActorDescContainer()->IsActorPinned(GetGuid());
+		return ActorDescHandle->GetContainer()->IsActorPinned(GetGuid());
 	}
 	return false;
 }
diff --git a/Engine/Source/Editor/SceneOutliner/Private/ActorHierarchy.cpp b/Engine/Source/Editor/SceneOutliner/Private/ActorHierarchy.cpp
--- a/Engine/Source/Editor/SceneOutliner/Private/ActorHierarchy.cpp
+++ b/Engine/Source/Editor/SceneOutliner/Private/ActorHierarchy.cpp
@@ -264,7 +264,7 @@ FSceneOutlinerTreeItemPtr FActorHierarchy::FindOrCreateParentItem(const ISceneOu
 	}
 	else if (const FActorDescTreeItem* ActorDescItem = Item.CastTo<FActorDescTreeItem>())
 	{
-		if (const FWorldPartitionActorDesc* ActorDesc = ActorDescItem->ActorDescHandle.GetActorDesc())
+		if (const FWorldPartitionActorDesc* ActorDesc = ActorDescItem->ActorDescHandle.Get())
 		{
 			const FName FolderPath = GetFolderPathFromActorDesc(RepresentingWorld.Get(), ActorDesc);
 			if (!FolderPath.IsNone())
diff --git a/Engine/Source/Editor/SceneOutliner/Private/SceneOutlinerActorInfoColumn.cpp b/Engine/Source/Editor/SceneOutliner/Private/SceneOutlinerActorInfoColumn.cpp
--- a/Engine/Source/Editor/SceneOutliner/Private/SceneOutlinerActorInfoColumn.cpp
+++ b/Engine/Source/Editor/SceneOutliner/Private/SceneOutlinerActorInfoColumn.cpp
@@ -59,7 +59,7 @@ struct FGetInfo
 		}
 		else if (const FActorDescTreeItem* ActorDescItem = Item.CastTo<FActorDescTreeItem>())
 		{
-			const FWorldPartitionActorDesc* ActorDesc = ActorDescItem->ActorDescHandle.GetActorDesc();
+			const FWorldPartitionActorDesc* ActorDesc = ActorDescItem->ActorDescHandle.Get();
 
 			if (!ActorDesc)
 			{
diff --git a/Engine/Source/Editor/SceneOutliner/Private/SceneOutlinerActorSCCColumn.cpp b/Engine/Source/Editor/SceneOutliner/Private/SceneOutlinerActorSCCColumn.cpp
--- a/Engine/Source/Editor/SceneOutliner/Private/SceneOutlinerActorSCCColumn.cpp
+++ b/Engine/Source/Editor/SceneOutliner/Private/SceneOutlinerActorSCCColumn.cpp
@@ -70,7 +70,7 @@ public:
 			}
 			else if (FActorDescTreeItem* ActorDescItem = TreeItemPtr->CastTo<FActorDescTreeItem>())
 			{
-				if (const FWorldPartitionActorDesc* ActorDesc = ActorDescItem->ActorDescHandle.GetActorDesc())
+				if (const FWorldPartitionActorDesc* ActorDesc = ActorDescItem->ActorDescHandle.Get())
 				{
 					ExternalPackageName =  USourceControlHelpers::PackageFilename(ActorDesc->GetActorPackage().ToString());
 				}
diff --git a/Engine/Source/Editor/SceneOutliner/Private/SceneOutlinerModule.cpp b/Engine/Source/Editor/SceneOutliner/Private/SceneOutlinerModule.cpp
--- a/Engine/Source/Editor/SceneOutliner/Private/SceneOutlinerModule.cpp
+++ b/Engine/Source/Editor/SceneOutliner/Private/SceneOutlinerModule.cpp
@@ -245,7 +245,7 @@ void FSceneOutlinerModule::CreateActorInfoColumns(FSceneOutlinerInitializationOp
 		}
 		else if (const FActorDescTreeItem* ActorDescItem = Item.CastTo<FActorDescTreeItem>())
 		{
-			const FWorldPartitionActorDesc* ActorDesc = ActorDescItem->ActorDescHandle.GetActorDesc();
+			const FWorldPartitionActorDesc* ActorDesc = ActorDescItem->ActorDescHandle.Get();
 
 			if (ActorDesc)
 			{
@@ -306,7 +306,7 @@ void FSceneOutlinerModule::CreateActorInfoColumns(FSceneOutlinerInitializationOp
 		}
 		else if (const FActorDescTreeItem* ActorDescItem = Item.CastTo<FActorDescTreeItem>())
 		{
-			if (const FWorldPartitionActorDesc* ActorDesc = ActorDescItem->ActorDescHandle.GetActorDesc())
+			if (const FWorldPartitionActorDesc* ActorDesc = ActorDescItem->ActorDescHandle.Get())
 			{
 				return ActorDesc->GetActorName().ToString();
 			}
@@ -333,7 +333,7 @@ void FSceneOutlinerModule::CreateActorInfoColumns(FSceneOutlinerInitializationOp
 		}
 		else if (const FActorDescTreeItem* ActorDescItem = Item.CastTo<FActorDescTreeItem>())
 		{
-			if (const FWorldPartitionActorDesc* ActorDesc = ActorDescItem->ActorDescHandle.GetActorDesc())
+			if (const FWorldPartitionActorDesc* ActorDesc = ActorDescItem->ActorDescHandle.Get())
 			{
 				return FPackageName::GetShortName(ActorDesc->GetActorPackage());
 			}
diff --git a/Engine/Source/Editor/SceneOutliner/Public/ActorDescTreeItem.h b/Engine/Source/Editor/SceneOutliner/Public/ActorDescTreeItem.h
--- a/Engine/Source/Editor/SceneOutliner/Public/ActorDescTreeItem.h
+++ b/Engine/Source/Editor/SceneOutliner/Public/ActorDescTreeItem.h
@@ -14,56 +14,22 @@ class UToolMenu;
 /** A tree item that represents an actor in the world */
 struct SCENEOUTLINER_API FActorDescTreeItem : ISceneOutlinerTreeItem
 {
-	struct FActorDescHandle
-	{
-	public:
-		FActorDescHandle(const FGuid& ActorGuid, UActorDescContainer* InContainer)
-			: ActorDesc(InContainer, ActorGuid)
-			, Container(InContainer)
-		{
-		}
-
-		~FActorDescHandle()
-		{
-			// HACK: prevent the destructor of the handle from accessing the deleted pointer if the container is invalid by manually resetting it.
-			if (!Container.IsValid())
-			{
-				ActorDesc.Reset();
-			}
-		}
-
-		bool IsValid() const { return Container.IsValid(); }
-		
-		const FWorldPartitionActorDesc* GetActorDesc() const
-		{
-			return Container.IsValid() ? ActorDesc.Get() : nullptr;
-		}
-
-		TWeakObjectPtr<UActorDescContainer> GetActorDescContainer() const
-		{
-			return Container;
-		}
-	private:
-		FWorldPartitionHandle ActorDesc;
-		TWeakObjectPtr<UActorDescContainer> Container;
-	};
-
 public:
 	DECLARE_DELEGATE_RetVal_OneParam(bool, FFilterPredicate, const FWorldPartitionActorDesc*);
 	DECLARE_DELEGATE_RetVal_OneParam(bool, FInteractivePredicate, const FWorldPartitionActorDesc*);
 
 	bool Filter(FFilterPredicate Pred) const
 	{
-		return Pred.Execute(ActorDescHandle.GetActorDesc());
+		return Pred.Execute(ActorDescHandle.Get());
 	}
 
 	bool GetInteractiveState(FInteractivePredicate Pred) const
 	{
-		return Pred.Execute(ActorDescHandle.GetActorDesc());
+		return Pred.Execute(ActorDescHandle.Get());
 	}
 
 	/** The actor desc this tree item is associated with. */
-	FActorDescHandle ActorDescHandle;
+	FWorldPartitionHandle ActorDescHandle;
 
 	/** Constant identifier for this tree item */
 	FSceneOutlinerTreeItemID ID;
@@ -75,13 +41,12 @@ public:
 	FActorDescTreeItem(const FGuid& InActorGuid, UActorDescContainer* Container);
 
 	/* Begin ISceneOutlinerTreeItem Implementation */
-	virtual bool IsValid() const override { return ActorDescHandle.GetActorDesc() != nullptr; }
+	virtual bool IsValid() const override { return ActorDescHandle.Get() != nullptr; }
 	virtual FSceneOutlinerTreeItemID GetID() const override;
 	virtual FString GetDisplayString() const override;
 	virtual bool CanInteract() const override;
 	virtual TSharedRef<SWidget> GenerateLabelWidget(ISceneOutliner& Outliner, const STableRow<FSceneOutlinerTreeItemPtr>& InRow) override;
 	virtual void GenerateContextMenu(UToolMenu* Menu, SSceneOutliner& Outliner) override;
-	virtual void OnVisibilityChanged(const bool bNewVisibility) override;
 	virtual bool HasVisibilityInfo() const override { return true; }
 	virtual bool GetVisibility() const override;
 	virtual bool ShouldShowPinnedState() const override { return true; }
