SceneOutliner: Fixed an issue where, under certain circumstances, multiple child entries representing the same object could be added to a scene outliner tree item.
The child pointer would eventually lose refcounts and lead to a stale entry in the children set and cause a validity check to fail.

Additionally, move and remove will now access the existing item in the tree rather than directly using the passed tree item reference since these are sometimes temporary objects.

diff --git a/Engine/Source/Editor/SceneOutliner/Private/SSceneOutliner.cpp b/Engine/Source/Editor/SceneOutliner/Private/SSceneOutliner.cpp
--- a/Engine/Source/Editor/SceneOutliner/Private/SSceneOutliner.cpp
+++ b/Engine/Source/Editor/SceneOutliner/Private/SSceneOutliner.cpp
@@ -633,19 +633,27 @@ void SSceneOutliner::AddPendingItem(FSceneOutlinerTreeItemPtr Item)
 
 void SSceneOutliner::AddPendingItemAndChildren(FSceneOutlinerTreeItemPtr Item)
 {
-	if (Item)
+	if (!Item.IsValid())
 	{
-		AddPendingItem(Item);
+		return;
+	}
 
-		TArray<FSceneOutlinerTreeItemPtr> Children;
-		Mode->GetHierarchy()->CreateChildren(Item, Children);
-		for (auto& Child : Children)
-		{
-			AddPendingItem(Child);
-		}
+	// Verify that there isn't already a pending operation for this item:
+	if (PendingTreeItemMap.Contains(Item->GetID()))
+	{
+		return;
+	}
 
-		Refresh();
+	AddPendingItem(Item);
+
+	TArray<FSceneOutlinerTreeItemPtr> Children;
+	Mode->GetHierarchy()->CreateChildren(Item, Children);
+	for (auto& Child : Children)
+	{
+		AddPendingItem(Child);
 	}
+
+	Refresh();
 }
 
 void SSceneOutliner::RepopulateEntireTree()
@@ -679,16 +687,17 @@ void SSceneOutliner::OnChildRemovedFromParent(ISceneOutlinerTreeItem& Parent)
 	}
 }
 
-void SSceneOutliner::OnItemMoved(const FSceneOutlinerTreeItemRef& Item)
+void SSceneOutliner::OnItemMoved(const FSceneOutlinerTreeItemRef& ReferenceItem)
 {
 	// Just remove the item if it no longer matches the filters
-	if (!Item->Flags.bIsFilteredOut && !SearchBoxFilter->PassesFilter(*Item))
+	if (!ReferenceItem->Flags.bIsFilteredOut && !SearchBoxFilter->PassesFilter(*ReferenceItem))
 	{
 		// This will potentially remove any non-matching, empty parents as well
-		RemoveItemFromTree(Item);
+		RemoveItemFromTree(ReferenceItem);
 	}
-	else if(TreeItemMap.Contains(Item->GetID()))
+	else if (const FSceneOutlinerTreeItemPtr* ItemInTree = TreeItemMap.Find(ReferenceItem->GetID()))
 	{
+		FSceneOutlinerTreeItemRef Item = ItemInTree->ToSharedRef();
 		// The item still matches the filters (or has children that do)
 		// When an item has been asked to move, it will still reside under its old parent
 		FSceneOutlinerTreeItemPtr Parent = Item->GetParent();
@@ -725,25 +734,27 @@ FSceneOutlinerTreeItemPtr SSceneOutliner::GetTreeItem(FSceneOutlinerTreeItemID I
 	return Result;
 }
 
-void SSceneOutliner::RemoveItemFromTree(FSceneOutlinerTreeItemRef InItem)
+void SSceneOutliner::RemoveItemFromTree(FSceneOutlinerTreeItemRef ReferenceItem)
 {
-	if (TreeItemMap.Contains(InItem->GetID()))
+	if (const FSceneOutlinerTreeItemPtr* ItemInTree = TreeItemMap.Find(ReferenceItem->GetID()))
 	{
-		auto Parent = InItem->GetParent();
+		FSceneOutlinerTreeItemRef Item = ItemInTree->ToSharedRef();
+
+		auto Parent = Item->GetParent();
 
 		if (Parent.IsValid())
 		{
-			Parent->RemoveChild(InItem);
+			Parent->RemoveChild(Item);
 			OnChildRemovedFromParent(*Parent);
 		}
 		else
 		{
-			RootTreeItems.Remove(InItem);
+			RootTreeItems.Remove(Item);
 		}
 
-		TreeItemMap.Remove(InItem->GetID());
+		TreeItemMap.Remove(Item->GetID());
 
-		Mode->OnItemRemoved(InItem);
+		Mode->OnItemRemoved(Item);
 	}
 }
 
