Added a ColumnPermissionList to the scene outliner module.

diff --git a/Engine/Source/Editor/SceneOutliner/Private/SSceneOutliner.cpp b/Engine/Source/Editor/SceneOutliner/Private/SSceneOutliner.cpp
--- a/Engine/Source/Editor/SceneOutliner/Private/SSceneOutliner.cpp
+++ b/Engine/Source/Editor/SceneOutliner/Private/SSceneOutliner.cpp
@@ -97,6 +97,9 @@ void SSceneOutliner::Construct(const FArguments& InArgs, const FSceneOutlinerIni
 		SearchBoxFilter = MakeShareable( new SceneOutliner::TreeItemTextFilter( Delegate ) );
 	}
 
+	FSceneOutlinerModule& SceneOutlinerModule = FModuleManager::LoadModuleChecked<FSceneOutlinerModule>("SceneOutliner");
+	SceneOutlinerModule.OnColumnPermissionListChanged().AddSP(this, &SSceneOutliner::OnColumnPermissionListChanged);
+
 	TSharedRef<SVerticalBox> VerticalBox = SNew(SVerticalBox);
 
 	for (auto& ModeFilterInfo : Mode->GetFilterInfos())
@@ -305,16 +308,25 @@ void SSceneOutliner::SetupColumns(SHeaderRow& HeaderRow)
 		SharedData->UseDefaultColumns();
 	}
 
-	Columns.Empty(SharedData->ColumnMap.Num());
+	TMap<FName, FSceneOutlinerColumnInfo> FilteredColumnMap;
+	for(auto It(SharedData->ColumnMap.CreateIterator()); It; ++It)
+	{
+		if (SceneOutlinerModule.GetColumnPermissionList()->PassesFilter(It.Key()))
+		{
+			FilteredColumnMap.Add(It.Key(), It.Value());
+		}
+	}
+
+	Columns.Empty(FilteredColumnMap.Num());
 	HeaderRow.ClearColumns();
 
 	// Get a list of sorted columns IDs to create
 	TArray<FName> SortedIDs;
-	SortedIDs.Reserve(SharedData->ColumnMap.Num());
-	SharedData->ColumnMap.GenerateKeyArray(SortedIDs);
+	SortedIDs.Reserve(FilteredColumnMap.Num());
+	FilteredColumnMap.GenerateKeyArray(SortedIDs);
 
 	SortedIDs.Sort([&](const FName& A, const FName& B){
-		return SharedData->ColumnMap[A].PriorityIndex < SharedData->ColumnMap[B].PriorityIndex;
+		return FilteredColumnMap[A].PriorityIndex < FilteredColumnMap[B].PriorityIndex;
 	});
 
 	TMap<FName, bool> ColumnVisibilities;
@@ -336,16 +348,16 @@ void SSceneOutliner::SetupColumns(SHeaderRow& HeaderRow)
 		{
 			bIsVisible = *ColumnVisibility;
 		}
-		else if (SharedData->ColumnMap[ID].Visibility == ESceneOutlinerColumnVisibility::Invisible)
+		else if (FilteredColumnMap[ID].Visibility == ESceneOutlinerColumnVisibility::Invisible)
 		{
 			bIsVisible = false;
 		}
 
 		TSharedPtr<ISceneOutlinerColumn> Column;
 
-		if (SharedData->ColumnMap[ID].Factory.IsBound())
+		if (FilteredColumnMap[ID].Factory.IsBound())
 		{
-			Column = SharedData->ColumnMap[ID].Factory.Execute(*this);
+			Column = FilteredColumnMap[ID].Factory.Execute(*this);
 		}
 		else
 		{
@@ -365,9 +377,9 @@ void SSceneOutliner::SetupColumns(SHeaderRow& HeaderRow)
 					.OnSort(this, &SSceneOutliner::OnColumnSortModeChanged);
 			}
 
-			if (SharedData->ColumnMap[ID].ColumnLabel.IsSet())
+			if (FilteredColumnMap[ID].ColumnLabel.IsSet())
 			{
-				ColumnArgs.DefaultLabel(SharedData->ColumnMap[ID].ColumnLabel);
+				ColumnArgs.DefaultLabel(FilteredColumnMap[ID].ColumnLabel);
 			}
 			else
 			{
@@ -379,7 +391,7 @@ void SSceneOutliner::SetupColumns(SHeaderRow& HeaderRow)
 				ColumnArgs.DefaultLabel(FText::FromName(ID));
 			}
 			
-			if (!SharedData->ColumnMap[ID].bCanBeHidden)
+			if (!FilteredColumnMap[ID].bCanBeHidden)
 			{
 				ColumnArgs.ShouldGenerateWidget(true);
 			}
@@ -402,8 +414,19 @@ void SSceneOutliner::RefreshColums()
 	bNeedsColumRefresh = true;
 }
 
+void SSceneOutliner::OnColumnPermissionListChanged()
+{
+	RefreshColums();
+	FullRefresh();
+}
+
 SSceneOutliner::~SSceneOutliner()
 {
+	if (FSceneOutlinerModule* SceneOutlinerModule = FModuleManager::GetModulePtr<FSceneOutlinerModule>("SceneOutliner"))
+	{
+		SceneOutlinerModule->OnColumnPermissionListChanged().RemoveAll(this);
+	}
+
 	Mode->GetHierarchy()->OnHierarchyChanged().RemoveAll(this);
 	delete Mode;
 
diff --git a/Engine/Source/Editor/SceneOutliner/Private/SceneOutlinerModule.cpp b/Engine/Source/Editor/SceneOutliner/Private/SceneOutlinerModule.cpp
--- a/Engine/Source/Editor/SceneOutliner/Private/SceneOutlinerModule.cpp
+++ b/Engine/Source/Editor/SceneOutliner/Private/SceneOutlinerModule.cpp
@@ -30,6 +30,12 @@
 /* FSceneOutlinerModule interface
  *****************************************************************************/
 
+FSceneOutlinerModule::FSceneOutlinerModule()
+: ColumnPermissionList(MakeShareable(new FNamePermissionList()))
+{
+	ColumnPermissionList->OnFilterChanged().AddLambda([this]() { ColumnPermissionListChanged.Broadcast(); });
+}
+
 void FSceneOutlinerModule::StartupModule()
 {
 	RegisterDefaultColumnType< FSceneOutlinerItemLabelColumn >(FSceneOutlinerColumnInfo(ESceneOutlinerColumnVisibility::Visible, 10, FCreateSceneOutlinerColumn(), false, TOptional<float>(), FSceneOutlinerBuiltInColumnTypes::Label_Localized()));
@@ -381,13 +387,13 @@ void FSceneOutlinerModule::CreateActorInfoColumns(FSceneOutlinerInitializationOp
 				ColumnName));
 	};
 
-	AddTextInfoColumn("Mobility", LOCTEXT("SceneOutlinerMobilityColumn", "Mobility"), MobilityInfoText);
-	AddTextInfoColumn("Level", LOCTEXT("SceneOutlinerLevelColumn", "Level"), LevelInfoText);
-	AddTextInfoColumn("Layer", LOCTEXT("SceneOutlinerLayerColumn", "Layer"), LayerInfoText);
-	AddTextInfoColumn("Data Layer", LOCTEXT("SceneOutlinerDataLayerColumn", "Data Layer"), DataLayerInfoText);
-	AddTextInfoColumn("ID Name", LOCTEXT("SceneOutlinerIDColumn", "ID Name"), InternalNameInfoText);
-	AddTextInfoColumn("Package Short Name", LOCTEXT("SceneOutlinerPackageShortNameColumn", "Package Short Name"), PackageShortNameInfoText);
-	AddTextInfoColumn("Uncached Lights", LOCTEXT("SceneOutlinerUncachedLightsColumn", "# Uncached Lights"), UncachedLightsInfoText);
+	AddTextInfoColumn(FSceneOutlinerBuiltInColumnTypes::Mobility(), FSceneOutlinerBuiltInColumnTypes::Mobility_Localized(), MobilityInfoText);
+	AddTextInfoColumn(FSceneOutlinerBuiltInColumnTypes::Level(), FSceneOutlinerBuiltInColumnTypes::Level_Localized(), LevelInfoText);
+	AddTextInfoColumn(FSceneOutlinerBuiltInColumnTypes::Layer(), FSceneOutlinerBuiltInColumnTypes::Layer_Localized(), LayerInfoText);
+	AddTextInfoColumn(FSceneOutlinerBuiltInColumnTypes::DataLayer(), FSceneOutlinerBuiltInColumnTypes::DataLayer_Localized(), DataLayerInfoText);
+	AddTextInfoColumn(FSceneOutlinerBuiltInColumnTypes::IDName(), FSceneOutlinerBuiltInColumnTypes::IDName_Localized(), InternalNameInfoText);
+	AddTextInfoColumn(FSceneOutlinerBuiltInColumnTypes::PackageShortName(), FSceneOutlinerBuiltInColumnTypes::PackageShortName_Localized(), PackageShortNameInfoText);
+	AddTextInfoColumn(FSceneOutlinerBuiltInColumnTypes::UncachedLights(), FSceneOutlinerBuiltInColumnTypes::UncachedLights_Localized(), UncachedLightsInfoText);
 }
 
 IMPLEMENT_MODULE(FSceneOutlinerModule, SceneOutliner);
diff --git a/Engine/Source/Editor/SceneOutliner/Public/SSceneOutliner.h b/Engine/Source/Editor/SceneOutliner/Public/SSceneOutliner.h
--- a/Engine/Source/Editor/SceneOutliner/Public/SSceneOutliner.h
+++ b/Engine/Source/Editor/SceneOutliner/Public/SSceneOutliner.h
@@ -769,6 +769,9 @@ private:
 
 private:
 
+	/** Called when SceneOutlinerModule column permission list changes. */
+	void OnColumnPermissionListChanged();
+
 	/** Structure containing information relating to the expansion state of parent items in the tree */
 	typedef TMap<FSceneOutlinerTreeItemID, bool> FParentsExpansionState;
 	
diff --git a/Engine/Source/Editor/SceneOutliner/Public/SceneOutlinerModule.h b/Engine/Source/Editor/SceneOutliner/Public/SceneOutlinerModule.h
--- a/Engine/Source/Editor/SceneOutliner/Public/SceneOutlinerModule.h
+++ b/Engine/Source/Editor/SceneOutliner/Public/SceneOutlinerModule.h
@@ -6,6 +6,7 @@
 #include "Modules/ModuleInterface.h"
 #include "ISceneOutliner.h"
 #include "SceneOutlinerPublicTypes.h"
+#include "Misc/NamePermissionList.h"
 
 class ICustomSceneOutliner;
 class ISceneOutlinerColumn;
@@ -26,6 +27,8 @@ class FSceneOutlinerModule
 {
 public:
 
+	FSceneOutlinerModule();
+
 	/**
 	 * Creates a scene outliner widget
 	 *
@@ -56,6 +59,13 @@ public:
 		const FSceneOutlinerInitializationOptions& InInitOptions,
 		TWeakObjectPtr<UWorld> SpecifiedWorld = nullptr) const;
 
+	/** Column permission list */
+	TSharedRef<FNamePermissionList>& GetColumnPermissionList() { return ColumnPermissionList; }
+
+	/** Delegate that broadcasts when column permission list changes. */
+	DECLARE_MULTICAST_DELEGATE(FOnColumnPermissionListChanged);
+	FOnColumnPermissionListChanged& OnColumnPermissionListChanged() { return ColumnPermissionListChanged; }
+
 public:
 	/** Register a new type of column available to all scene outliners */
 	template< typename T >
@@ -115,6 +125,12 @@ private:
 	/** Map of column type name -> factory delegate */
 	TMap< FName, FCreateSceneOutlinerColumn > ColumnMap;
 
+	/** Column permission list used to filter scene ouliner columns. */
+	TSharedRef<FNamePermissionList> ColumnPermissionList;
+
+	/** Delegate that broadcasts when column permission list changes. */
+	FOnColumnPermissionListChanged ColumnPermissionListChanged;
+
 	// Function to recreate what originally was FActorInfoColumn as separate columns using FTextInfoColumn
 	void CreateActorInfoColumns(FSceneOutlinerInitializationOptions& InInitOptions) const;
 
diff --git a/Engine/Source/Editor/SceneOutliner/Public/SceneOutlinerPublicTypes.h b/Engine/Source/Editor/SceneOutliner/Public/SceneOutlinerPublicTypes.h
--- a/Engine/Source/Editor/SceneOutliner/Public/SceneOutlinerPublicTypes.h
+++ b/Engine/Source/Editor/SceneOutliner/Public/SceneOutlinerPublicTypes.h
@@ -22,73 +22,30 @@ DECLARE_DELEGATE_RetVal_OneParam(ISceneOutlinerMode*, FCreateSceneOutlinerMode,
 /** Container for built in column types. Function-static so they are available without linking */
 struct FSceneOutlinerBuiltInColumnTypes
 {
-	/** The gutter column */
-	static const FName& Gutter()
-	{
-		static FName Gutter("Visibility"); // Renamed "Gutter" to "Visibility" so the purpose is more obvious in editor menus
-		return Gutter;
-	}
-
-	/** Localizable FText name for the Gutter Column (pass into FSceneOutlinerColumnInfo) */
-	static const FText& Gutter_Localized()
-	{
-		static FText Gutter_Localized = LOCTEXT("VisibilityColumnName", "Visibility");
-		return Gutter_Localized;
-	}
-
-	/** The item label column */
-	static const FName& Label()
-	{
-		static FName Label("Item Label");
-		return Label;
+	#define DEFINE_SCENEOUTLINER_BUILTIN_COLUMN_TYPE(ColumnID, ColumnName, ColumnKey, ColumnLocalizedTextLiteral) \
+	static FName& ColumnID() \
+	{ \
+		static FName ColumnID = ColumnName; \
+		return ColumnID; \
+	} \
+	static const FText& ColumnID##_Localized() \
+	{ \
+		static FText ColumnID##_Localized = LOCTEXT(ColumnKey, ColumnLocalizedTextLiteral); \
+		return ColumnID##_Localized; \
 	}
 
-	/** Localizable FText name for the Item Label Column (pass into FSceneOutlinerColumnInfo) */
-	static const FText& Label_Localized()
-	{
-		static FText Label_Localized = LOCTEXT("ItemLabelColumnName", "Item Label");
-		return Label_Localized;
-	}
-
-	/** Generic actor info column */
-	static FName& ActorInfo()
-	{
-		static FName ActorInfo("Type"); // Renamed "Actor Info" to "Type" since it has been refactored to only show type information
-		return ActorInfo;
-	}
-
-	/** Localizable FText name for the Type Column (pass into FSceneOutlinerColumnInfo) */
-	static const FText& ActorInfo_Localized()
-	{
-		static FText ActorInfo_Localized = LOCTEXT("TypeColumnName", "Type");
-		return ActorInfo_Localized;
-	}
-
-	static FName& SourceControl()
-	{
-		static FName SourceControl("Source Control");
-		return SourceControl;
-	}
-
-	/** Localizable FText name for the Type Column (pass into FSceneOutlinerColumnInfo) */
-	static const FText& SourceControl_Localized()
-	{
-		static FText SourceControl_Localized = LOCTEXT("SourceControlColumnName", "Source Control");
-		return SourceControl_Localized;
-	}
-
-	static FName& Pinned()
-	{
-		static FName Pinned("Pinned");
-		return Pinned;
-	}
-
-	/** Localizable FText name for the Type Column (pass into FSceneOutlinerColumnInfo) */
-	static const FText& Pinned_Localized()
-	{
-		static FText Pinned_Localized = LOCTEXT("PinnedColumnName", "Pinned");
-		return Pinned_Localized;
-	}
+	DEFINE_SCENEOUTLINER_BUILTIN_COLUMN_TYPE(Label, "Item Label", "ItemLabelColumnName", "Item Label");
+	DEFINE_SCENEOUTLINER_BUILTIN_COLUMN_TYPE(Gutter, "Visibility", "VisibilityColumnName", "Visibility");
+	DEFINE_SCENEOUTLINER_BUILTIN_COLUMN_TYPE(ActorInfo, "Type", "TypeColumnName", "Type");
+	DEFINE_SCENEOUTLINER_BUILTIN_COLUMN_TYPE(SourceControl, "Source Control", "SourceControlColumnName", "Source Control");
+	DEFINE_SCENEOUTLINER_BUILTIN_COLUMN_TYPE(Mobility, "Mobility", "SceneOutlinerMobilityColumn", "Mobility");
+	DEFINE_SCENEOUTLINER_BUILTIN_COLUMN_TYPE(Level, "Level", "SceneOutlinerLevelColumn", "Level");
+	DEFINE_SCENEOUTLINER_BUILTIN_COLUMN_TYPE(Layer, "Layer", "SceneOutlinerLayerColumn", "Layer");
+	DEFINE_SCENEOUTLINER_BUILTIN_COLUMN_TYPE(DataLayer, "Data Layer", "SceneOutlinerDataLayerColumn", "Data Layer");
+	DEFINE_SCENEOUTLINER_BUILTIN_COLUMN_TYPE(Pinned, "Pinned", "SceneOutlinerPinnedColumn", "Pinned");
+	DEFINE_SCENEOUTLINER_BUILTIN_COLUMN_TYPE(IDName, "ID Name", "SceneOutlinerIDNameColumn", "ID Name");
+	DEFINE_SCENEOUTLINER_BUILTIN_COLUMN_TYPE(PackageShortName, "Package Short Name", "SceneOutlinerPackageShortNameColumn", "Package Short Name");
+	DEFINE_SCENEOUTLINER_BUILTIN_COLUMN_TYPE(UncachedLights, "Uncached Lights", "SceneOutlinerUncachedLightsColumn", "# Uncached Lights");
 };
 
 /** Visibility enum for scene outliner columns */
