SceneOutliner: Filter out double click events on items that !CanInteract()

#rnx

diff --git a/Engine/Source/Editor/SceneOutliner/Private/SOutlinerTreeView.cpp b/Engine/Source/Editor/SceneOutliner/Private/SOutlinerTreeView.cpp
--- a/Engine/Source/Editor/SceneOutliner/Private/SOutlinerTreeView.cpp
+++ b/Engine/Source/Editor/SceneOutliner/Private/SOutlinerTreeView.cpp
@@ -236,6 +236,18 @@ FReply SSceneOutlinerTreeRow::OnMouseButtonUp(const FGeometry& MyGeometry, const
 	return FReply::Handled();
 }
 
+FReply SSceneOutlinerTreeRow::OnMouseButtonDoubleClick(const FGeometry& MyGeometry, const FPointerEvent& MouseEvent)
+{
+	auto ItemPtr = Item.Pin();
+	// We don't want to act on double click on an item that can't be interacted with
+	if (ItemPtr.IsValid() && ItemPtr->CanInteract())
+	{
+		return SMultiColumnTableRow<FSceneOutlinerTreeItemPtr>::OnMouseButtonDoubleClick(MyGeometry, MouseEvent);
+	}
+
+	return FReply::Handled();
+}
+
 TSharedRef<SWidget> SSceneOutlinerTreeRow::GenerateWidgetForColumn( const FName& ColumnName )
 {
 	auto ItemPtr = Item.Pin();
diff --git a/Engine/Source/Editor/SceneOutliner/Public/SOutlinerTreeView.h b/Engine/Source/Editor/SceneOutliner/Public/SOutlinerTreeView.h
--- a/Engine/Source/Editor/SceneOutliner/Public/SOutlinerTreeView.h
+++ b/Engine/Source/Editor/SceneOutliner/Public/SOutlinerTreeView.h
@@ -59,6 +59,7 @@ protected:
 
 	virtual FReply OnMouseButtonDown(const FGeometry& MyGeometry, const FPointerEvent& MouseEvent) override;
 	virtual FReply OnMouseButtonUp(const FGeometry& MyGeometry, const FPointerEvent& MouseEvent) override;
+	virtual FReply OnMouseButtonDoubleClick(const FGeometry& MyGeometry, const FPointerEvent& MouseEvent) override;
 
 	virtual void OnDragEnter( const FGeometry& MyGeometry, const FDragDropEvent& DragDropEvent ) override;
 
