Outliner: Fix a crash caused by duplication of items in the outliner when you performed multiple operations on an item too quickly

diff --git a/Engine/Source/Editor/SceneOutliner/Private/SSceneOutliner.cpp b/Engine/Source/Editor/SceneOutliner/Private/SSceneOutliner.cpp
--- a/Engine/Source/Editor/SceneOutliner/Private/SSceneOutliner.cpp
+++ b/Engine/Source/Editor/SceneOutliner/Private/SSceneOutliner.cpp
@@ -924,6 +924,15 @@ void SSceneOutliner::AddUnfilteredItemToTree(FSceneOutlinerTreeItemRef Item)
 		RootTreeItems.Add(Item);
 	}
 
+	// If there are any pending operations involving this item before it was added, update them
+	for (int i = 0; i < PendingOperations.Num(); i++)
+	{
+		if (PendingOperations[i].Item->GetID() == Item->GetID())
+		{
+			PendingOperations[i].Item = Item;
+		}
+	}
+
 	// keep track of the number of active folders
 	if (Item->IsA<FFolderTreeItem>())
 	{
