Fixed double-click not expanding/collapsing Folder in the World Outliner.

diff --git a/Engine/Source/Editor/SceneOutliner/Private/ActorBrowsingMode.cpp b/Engine/Source/Editor/SceneOutliner/Private/ActorBrowsingMode.cpp
--- a/Engine/Source/Editor/SceneOutliner/Private/ActorBrowsingMode.cpp
+++ b/Engine/Source/Editor/SceneOutliner/Private/ActorBrowsingMode.cpp
@@ -792,10 +792,6 @@ void FActorBrowsingMode::OnItemDoubleClick(FSceneOutlinerTreeItemPtr Item)
 			GEditor->MoveViewportCamerasToActor(*Actor, bActiveViewportOnly);
 		}
 	}
-	else if (Item->IsA<FFolderTreeItem>())
-	{
-		SceneOutliner->SetItemExpansion(Item, !SceneOutliner->IsItemExpanded(Item));
-	}
 }
 
 void FActorBrowsingMode::OnFilterTextCommited(FSceneOutlinerItemSelection& Selection, ETextCommit::Type CommitType)
