SceneOutliner: Fix a bug causing changes in PIE to apply to the actual level instead of replicating in the PIE session

diff --git a/Engine/Source/Editor/LevelEditor/Private/SLevelEditor.cpp b/Engine/Source/Editor/LevelEditor/Private/SLevelEditor.cpp
--- a/Engine/Source/Editor/LevelEditor/Private/SLevelEditor.cpp
+++ b/Engine/Source/Editor/LevelEditor/Private/SLevelEditor.cpp
@@ -725,7 +725,7 @@ TSharedRef<ISceneOutliner> SLevelEditor::CreateSceneOutliner(FName TabIdentifier
 
 	FSceneOutlinerModule& SceneOutlinerModule = FModuleManager::Get().LoadModuleChecked<FSceneOutlinerModule>("SceneOutliner");
 	TSharedRef<ISceneOutliner> SceneOutlinerRef = SceneOutlinerModule.CreateActorBrowser(
-		InitOptions, GetWorld());
+		InitOptions);
 
 	// Update the most recently created outliner
 	SceneOutlinerPtr = SceneOutlinerRef;
diff --git a/Engine/Source/Editor/SceneOutliner/Private/SceneOutlinerModule.cpp b/Engine/Source/Editor/SceneOutliner/Private/SceneOutlinerModule.cpp
--- a/Engine/Source/Editor/SceneOutliner/Private/SceneOutlinerModule.cpp
+++ b/Engine/Source/Editor/SceneOutliner/Private/SceneOutlinerModule.cpp
@@ -24,6 +24,7 @@
 #include "WorldTreeItem.h"
 #include "WorldPartition/DataLayer/DataLayerInstance.h"
 #include "WorldPartition/DataLayer/WorldDataLayers.h"
+#include "Editor.h"
 
 #define LOCTEXT_NAMESPACE "SceneOutlinerModule"
 
@@ -97,7 +98,7 @@ TSharedRef<ISceneOutliner> FSceneOutlinerModule::CreateActorPicker(const FSceneO
 		InitOptions.ColumnMap.Add(FSceneOutlinerBuiltInColumnTypes::Label(), FSceneOutlinerColumnInfo(ESceneOutlinerColumnVisibility::Visible, 0, FCreateSceneOutlinerColumn(), false, TOptional<float>(), FSceneOutlinerBuiltInColumnTypes::Label_Localized()));
 		InitOptions.ColumnMap.Add(FSceneOutlinerBuiltInColumnTypes::ActorInfo(), FSceneOutlinerColumnInfo(ESceneOutlinerColumnVisibility::Visible, 10, FCreateSceneOutlinerColumn(), true, TOptional<float>(), FSceneOutlinerBuiltInColumnTypes::ActorInfo_Localized()));
 		InitOptions.ColumnMap.Add(FSceneOutlinerBuiltInColumnTypes::Pinned(), FSceneOutlinerColumnInfo(ESceneOutlinerColumnVisibility::Invisible, 5, FCreateSceneOutlinerColumn(), true, TOptional<float>(), FSceneOutlinerBuiltInColumnTypes::Pinned_Localized()));
-		CreateActorInfoColumns(InitOptions, SpecifiedWorld);
+		CreateActorInfoColumns(InitOptions);
 	}
 	return CreateSceneOutliner(InitOptions);
 }
@@ -133,7 +134,7 @@ TSharedRef<ISceneOutliner> FSceneOutlinerModule::CreateComponentPicker(const FSc
 	{
 		InitOptions.ColumnMap.Add(FSceneOutlinerBuiltInColumnTypes::Label(), FSceneOutlinerColumnInfo(ESceneOutlinerColumnVisibility::Visible, 0, FCreateSceneOutlinerColumn(), false, TOptional<float>(), FSceneOutlinerBuiltInColumnTypes::Label_Localized()));
 		InitOptions.ColumnMap.Add(FSceneOutlinerBuiltInColumnTypes::ActorInfo(), FSceneOutlinerColumnInfo(ESceneOutlinerColumnVisibility::Visible, 10, FCreateSceneOutlinerColumn(), true, TOptional<float>(), FSceneOutlinerBuiltInColumnTypes::ActorInfo_Localized()));
-		CreateActorInfoColumns(InitOptions, SpecifiedWorld);
+		CreateActorInfoColumns(InitOptions);
 	}
 	return CreateSceneOutliner(InitOptions);
 }
@@ -147,6 +148,7 @@ TSharedRef< ISceneOutliner > FSceneOutlinerModule::CreateActorBrowser(const FSce
 
 	FSceneOutlinerInitializationOptions InitOptions(InInitOptions);
 	InitOptions.ModeFactory = ModeFactory;
+
 	if (InitOptions.ColumnMap.Num() == 0)
 	{
 		InitOptions.UseDefaultColumns();
@@ -154,21 +156,24 @@ TSharedRef< ISceneOutliner > FSceneOutlinerModule::CreateActorBrowser(const FSce
 		InitOptions.ColumnMap.Add(FSceneOutlinerBuiltInColumnTypes::ActorInfo(), FSceneOutlinerColumnInfo(ESceneOutlinerColumnVisibility::Visible, 20, FCreateSceneOutlinerColumn(), true, TOptional<float>(), FSceneOutlinerBuiltInColumnTypes::ActorInfo_Localized()));
 		InitOptions.ColumnMap.Add(FSceneOutlinerBuiltInColumnTypes::SourceControl(), FSceneOutlinerColumnInfo(ESceneOutlinerColumnVisibility::Invisible, 30, FCreateSceneOutlinerColumn(), true, TOptional<float>(), FSceneOutlinerBuiltInColumnTypes::SourceControl_Localized()));
 
-		// We don't want the pinned column in non wp levels
-		if (UWorld* WorldPtr = SpecifiedWorld.Get())
+		// TODO: Figure out a better way to get the current world
+		UWorld* WorldPtr = GEditor->GetEditorWorldContext().World();
+
+		if (WorldPtr)
 		{
+			// We don't want the pinned column in non wp levels
 			if (WorldPtr->IsPartitionedWorld())
 			{
 				InitOptions.ColumnMap.Add(FSceneOutlinerBuiltInColumnTypes::Pinned(), FSceneOutlinerColumnInfo(ESceneOutlinerColumnVisibility::Visible, 5, FCreateSceneOutlinerColumn(), true, TOptional<float>(), FSceneOutlinerBuiltInColumnTypes::Pinned_Localized()));
 			}
+			CreateActorInfoColumns(InitOptions, WorldPtr);
 		}
 		
-		CreateActorInfoColumns(InitOptions, SpecifiedWorld);
 	}
 	return CreateSceneOutliner(InitOptions);
 }
 
-void FSceneOutlinerModule::CreateActorInfoColumns(FSceneOutlinerInitializationOptions& InInitOptions, TWeakObjectPtr<UWorld> SpecifiedWorld) const
+void FSceneOutlinerModule::CreateActorInfoColumns(FSceneOutlinerInitializationOptions& InInitOptions, UWorld *WorldPtr) const
 {
 	FGetTextForItem MobilityInfoText = FGetTextForItem::CreateLambda([](const ISceneOutlinerTreeItem& Item) -> FString
 	{
@@ -382,14 +387,11 @@ void FSceneOutlinerModule::CreateActorInfoColumns(FSceneOutlinerInitializationOp
 	};
 
 	// The "Level" column should be named "Package Short Name" in wp enabled levels
-	auto LevelColumnName = TAttribute<FText>::CreateLambda([SpecifiedWorld]() -> FText
+	auto LevelColumnName = TAttribute<FText>::CreateLambda([WorldPtr]() -> FText
 	{
-		if (UWorld* WorldPtr = SpecifiedWorld.Get())
+		if (WorldPtr && WorldPtr->IsPartitionedWorld())
 		{
-			if (WorldPtr->IsPartitionedWorld())
-			{
-				return FSceneOutlinerBuiltInColumnTypes::PackageShortName_Localized();
-			}
+			return FSceneOutlinerBuiltInColumnTypes::PackageShortName_Localized();
 		}
 
 		return FSceneOutlinerBuiltInColumnTypes::Level_Localized();
diff --git a/Engine/Source/Editor/SceneOutliner/Public/SceneOutlinerModule.h b/Engine/Source/Editor/SceneOutliner/Public/SceneOutlinerModule.h
--- a/Engine/Source/Editor/SceneOutliner/Public/SceneOutlinerModule.h
+++ b/Engine/Source/Editor/SceneOutliner/Public/SceneOutlinerModule.h
@@ -132,7 +132,7 @@ private:
 	FOnColumnPermissionListChanged ColumnPermissionListChanged;
 
 	// Function to recreate what originally was FActorInfoColumn as separate columns using FTextInfoColumn
-	void CreateActorInfoColumns(FSceneOutlinerInitializationOptions& InInitOptions, TWeakObjectPtr<UWorld> SpecifiedWorld) const;
+	void CreateActorInfoColumns(FSceneOutlinerInitializationOptions& InInitOptions, UWorld* WorldPtr = nullptr) const;
 
 public:
 
