SceneOutliner: Add the Socket column back to the outliner

diff --git a/Engine/Source/Editor/SceneOutliner/Private/SceneOutlinerModule.cpp b/Engine/Source/Editor/SceneOutliner/Private/SceneOutlinerModule.cpp
--- a/Engine/Source/Editor/SceneOutliner/Private/SceneOutlinerModule.cpp
+++ b/Engine/Source/Editor/SceneOutliner/Private/SceneOutlinerModule.cpp
@@ -401,6 +401,7 @@ void FSceneOutlinerModule::CreateActorInfoColumns(FSceneOutlinerInitializationOp
 	AddTextInfoColumn(FSceneOutlinerBuiltInColumnTypes::Level(), LevelColumnName, LevelInfoText);
 	AddTextInfoColumn(FSceneOutlinerBuiltInColumnTypes::Layer(), FSceneOutlinerBuiltInColumnTypes::Layer_Localized(), LayerInfoText);
 	AddTextInfoColumn(FSceneOutlinerBuiltInColumnTypes::DataLayer(), FSceneOutlinerBuiltInColumnTypes::DataLayer_Localized(), DataLayerInfoText);
+	AddTextInfoColumn(FSceneOutlinerBuiltInColumnTypes::Socket(), FSceneOutlinerBuiltInColumnTypes::Socket_Localized(), SocketInfoText);
 	AddTextInfoColumn(FSceneOutlinerBuiltInColumnTypes::IDName(), FSceneOutlinerBuiltInColumnTypes::IDName_Localized(), InternalNameInfoText);
 	AddTextInfoColumn(FSceneOutlinerBuiltInColumnTypes::UncachedLights(), FSceneOutlinerBuiltInColumnTypes::UncachedLights_Localized(), UncachedLightsInfoText);
 }
diff --git a/Engine/Source/Editor/SceneOutliner/Public/SceneOutlinerPublicTypes.h b/Engine/Source/Editor/SceneOutliner/Public/SceneOutlinerPublicTypes.h
--- a/Engine/Source/Editor/SceneOutliner/Public/SceneOutlinerPublicTypes.h
+++ b/Engine/Source/Editor/SceneOutliner/Public/SceneOutlinerPublicTypes.h
@@ -46,6 +46,7 @@ struct FSceneOutlinerBuiltInColumnTypes
 	DEFINE_SCENEOUTLINER_BUILTIN_COLUMN_TYPE(IDName, "ID Name", "SceneOutlinerIDNameColumn", "ID Name");
 	DEFINE_SCENEOUTLINER_BUILTIN_COLUMN_TYPE(PackageShortName, "Package Short Name", "SceneOutlinerPackageShortNameColumn", "Package Short Name");
 	DEFINE_SCENEOUTLINER_BUILTIN_COLUMN_TYPE(UncachedLights, "Uncached Lights", "SceneOutlinerUncachedLightsColumn", "# Uncached Lights");
+	DEFINE_SCENEOUTLINER_BUILTIN_COLUMN_TYPE(Socket, "Socket", "SceneOutlinerSocketColumn", "Socket");
 };
 
 /** Visibility enum for scene outliner columns */
