Outliner:
Port ActorBrowsingMode over to the new Editor Config system to have separate settings for each outliner
Enable Stacked Hierarchy Headers by default
Add the ability to disable framing selection
Add the ability to get the most recently used outliner from the level editor
Fix a bug where renaming an item on one outliner would sometimes execute the rename on a different outliner

diff --git a/Engine/Source/Editor/LevelEditor/Private/LevelEditorActions.cpp b/Engine/Source/Editor/LevelEditor/Private/LevelEditorActions.cpp
--- a/Engine/Source/Editor/LevelEditor/Private/LevelEditorActions.cpp
+++ b/Engine/Source/Editor/LevelEditor/Private/LevelEditorActions.cpp
@@ -1587,7 +1587,7 @@ bool FLevelEditorActionCallbacks::Duplicate_CanExecute()
 
 		if (!bCanDuplicate)
 		{
-			if (TSharedPtr<ISceneOutliner> SceneOutlinerPtr = LevelEditor->GetSceneOutliner())
+			if (TSharedPtr<ISceneOutliner> SceneOutlinerPtr = LevelEditor->GetMostRecentlyUsedSceneOutliner())
 			{
 				bCanDuplicate = SceneOutlinerPtr->Copy_CanExecute(); // If we can copy, we can duplicate
 			}
@@ -1625,7 +1625,7 @@ bool FLevelEditorActionCallbacks::Delete_CanExecute()
 
 		if (!bCanDelete)
 		{
-			if (TSharedPtr<ISceneOutliner> SceneOutlinerPtr = LevelEditor->GetSceneOutliner())
+			if (TSharedPtr<ISceneOutliner> SceneOutlinerPtr = LevelEditor->GetMostRecentlyUsedSceneOutliner())
 			{
 				bCanDelete = SceneOutlinerPtr->Delete_CanExecute();
 			}
@@ -1653,7 +1653,7 @@ void FLevelEditorActionCallbacks::Rename_Execute()
 		TWeakPtr<class ILevelEditor> LevelEditor = FModuleManager::GetModuleChecked<FLevelEditorModule>(TEXT("LevelEditor")).GetLevelEditorInstance();
 		if (LevelEditor.IsValid())
 		{
-			TSharedPtr<class ISceneOutliner> SceneOutlinerPtr = LevelEditor.Pin()->GetSceneOutliner();
+			TSharedPtr<class ISceneOutliner> SceneOutlinerPtr = LevelEditor.Pin()->GetMostRecentlyUsedSceneOutliner();
 			if (SceneOutlinerPtr.IsValid())
 			{
 				SceneOutlinerPtr->Rename_Execute();
@@ -1683,7 +1683,7 @@ bool FLevelEditorActionCallbacks::Rename_CanExecute()
 		TWeakPtr<class ILevelEditor> LevelEditor = FModuleManager::GetModuleChecked<FLevelEditorModule>(TEXT("LevelEditor")).GetLevelEditorInstance();
 		if (LevelEditor.IsValid())
 		{
-			TSharedPtr<class ISceneOutliner> SceneOutlinerPtr = LevelEditor.Pin()->GetSceneOutliner();
+			TSharedPtr<class ISceneOutliner> SceneOutlinerPtr = LevelEditor.Pin()->GetMostRecentlyUsedSceneOutliner();
 			if (SceneOutlinerPtr.IsValid())
 			{
 				bCanRename = SceneOutlinerPtr->Rename_CanExecute();
@@ -1733,7 +1733,7 @@ bool FLevelEditorActionCallbacks::Cut_CanExecute()
 		TWeakPtr<class ILevelEditor> LevelEditor = FModuleManager::GetModuleChecked<FLevelEditorModule>(TEXT("LevelEditor")).GetLevelEditorInstance();
 		if (LevelEditor.IsValid())
 		{
-			TSharedPtr<class ISceneOutliner> SceneOutlinerPtr = LevelEditor.Pin()->GetSceneOutliner();
+			TSharedPtr<class ISceneOutliner> SceneOutlinerPtr = LevelEditor.Pin()->GetMostRecentlyUsedSceneOutliner();
 			if (SceneOutlinerPtr.IsValid())
 			{
 				bCanCut = SceneOutlinerPtr->Cut_CanExecute();
@@ -1781,7 +1781,7 @@ bool FLevelEditorActionCallbacks::Copy_CanExecute()
 		TWeakPtr<class ILevelEditor> LevelEditor = FModuleManager::GetModuleChecked<FLevelEditorModule>(TEXT("LevelEditor")).GetLevelEditorInstance();
 		if (LevelEditor.IsValid())
 		{
-			TSharedPtr<class ISceneOutliner> SceneOutlinerPtr = LevelEditor.Pin()->GetSceneOutliner();
+			TSharedPtr<class ISceneOutliner> SceneOutlinerPtr = LevelEditor.Pin()->GetMostRecentlyUsedSceneOutliner();
 			if (SceneOutlinerPtr.IsValid())
 			{
 				bCanCopy = SceneOutlinerPtr->Copy_CanExecute();
@@ -1827,7 +1827,7 @@ bool FLevelEditorActionCallbacks::Paste_CanExecute()
 		TWeakPtr<class ILevelEditor> LevelEditor = FModuleManager::GetModuleChecked<FLevelEditorModule>(TEXT("LevelEditor")).GetLevelEditorInstance();
 		if (LevelEditor.IsValid())
 		{
-			TSharedPtr<class ISceneOutliner> SceneOutlinerPtr = LevelEditor.Pin()->GetSceneOutliner();
+			TSharedPtr<class ISceneOutliner> SceneOutlinerPtr = LevelEditor.Pin()->GetMostRecentlyUsedSceneOutliner();
 			if (SceneOutlinerPtr.IsValid())
 			{
 				bCanPaste = SceneOutlinerPtr->Paste_CanExecute();
diff --git a/Engine/Source/Editor/LevelEditor/Private/LevelEditorContextMenu.cpp b/Engine/Source/Editor/LevelEditor/Private/LevelEditorContextMenu.cpp
--- a/Engine/Source/Editor/LevelEditor/Private/LevelEditorContextMenu.cpp
+++ b/Engine/Source/Editor/LevelEditor/Private/LevelEditorContextMenu.cpp
@@ -588,7 +588,7 @@ void FLevelEditorContextMenu::RegisterSceneOutlinerContextMenu()
 	{
 		if (ULevelEditorContextMenuContext* LevelEditorContext = InMenu->FindContext<ULevelEditorContextMenuContext>())
 		{
-			TWeakPtr<ISceneOutliner> SceneOutlinerPtr = LevelEditorContext->LevelEditor.Pin()->GetSceneOutliner();
+			TWeakPtr<ISceneOutliner> SceneOutlinerPtr = LevelEditorContext->LevelEditor.Pin()->GetMostRecentlyUsedSceneOutliner();
 			if (SceneOutlinerPtr.IsValid())
 			{
 				FToolMenuSection& Section = InMenu->AddSection("SelectVisibilityLevels");
diff --git a/Engine/Source/Editor/LevelEditor/Private/SLevelEditor.cpp b/Engine/Source/Editor/LevelEditor/Private/SLevelEditor.cpp
--- a/Engine/Source/Editor/LevelEditor/Private/SLevelEditor.cpp
+++ b/Engine/Source/Editor/LevelEditor/Private/SLevelEditor.cpp
@@ -228,6 +228,12 @@ void SLevelEditor::Construct( const SLevelEditor::FArguments& InArgs)
 	FEditorDelegates::MapChange.AddRaw(this, &SLevelEditor::HandleEditorMapChange);
 	FEditorDelegates::OnAssetsDeleted.AddRaw(this, &SLevelEditor::HandleAssetsDeleted);
 	HandleEditorMapChange(MapChangeEventFlags::NewMap);
+
+	// Register the display names of the outliners
+	SceneOutlinerDisplayNames.Add(LevelEditorTabIds::LevelEditorSceneOutliner, NSLOCTEXT("LevelEditorTabs", "LevelEditorSceneOutliner", "Outliner 1"));
+	SceneOutlinerDisplayNames.Add(LevelEditorTabIds::LevelEditorSceneOutliner2, NSLOCTEXT("LevelEditorTabs", "LevelEditorSceneOutliner2", "Outliner 2"));
+	SceneOutlinerDisplayNames.Add(LevelEditorTabIds::LevelEditorSceneOutliner3, NSLOCTEXT("LevelEditorTabs", "LevelEditorSceneOutliner3", "Outliner 3"));
+	SceneOutlinerDisplayNames.Add(LevelEditorTabIds::LevelEditorSceneOutliner4, NSLOCTEXT("LevelEditorTabs", "LevelEditorSceneOutliner4", "Outliner 4"));
 }
 
 void SLevelEditor::Initialize( const TSharedRef<SDockTab>& OwnerTab, const TSharedRef<SWindow>& OwnerWindow )
@@ -653,6 +659,64 @@ TArray<TWeakPtr<ISceneOutliner>> SLevelEditor::GetAllSceneOutliners() const
 	return OutValueArray;
 }
 
+void SLevelEditor::SetMostRecentlyUsedSceneOutliner(FName OutlinerIdentifier)
+{
+	TWeakPtr<ISceneOutliner> *Outliner = SceneOutliners.Find(OutlinerIdentifier);
+
+	if(Outliner)
+	{
+		if (TSharedPtr<ISceneOutliner> OutlinerPtr = Outliner->Pin())
+		{
+			SceneOutlinerPtr = OutlinerPtr;
+		}
+	}
+}
+
+void SLevelEditor::ResetMostRecentOutliner()
+{
+	// Pick the first available outliner as the most recent outliner
+	for(const TPair<FName, TWeakPtr<ISceneOutliner>>& Outliner : SceneOutliners)
+	{
+		if(TSharedPtr<ISceneOutliner> OutlinerPin = Outliner.Value.Pin())
+		{
+			SceneOutlinerPtr = OutlinerPin;
+			return;
+		}
+	}
+
+	// No outliners are open currently
+	SceneOutlinerPtr.Reset();
+}
+
+TSharedPtr<ISceneOutliner> SLevelEditor::GetMostRecentlyUsedSceneOutliner()
+{
+	// If the outliner marked as "most recent" is still open, simply return it
+	if(TSharedPtr<ISceneOutliner> SceneOutlinerPin = SceneOutlinerPtr.Pin())
+	{
+		return SceneOutlinerPin;
+	}
+
+	// Otherwise, pick the first available outliner to be the most recent one
+	ResetMostRecentOutliner();
+
+	return SceneOutlinerPtr.Pin();
+}
+
+TSharedPtr<ISceneOutliner> SLevelEditor::GetSceneOutliner() const
+{
+	// Simply pick the first available outliner
+	for(const TPair<FName, TWeakPtr<ISceneOutliner>>& Outliner : SceneOutliners)
+	{
+		if(TSharedPtr<ISceneOutliner> OutlinerPin = Outliner.Value.Pin())
+		{
+			return OutlinerPin;
+		}
+	}
+
+	// No outliners are open currently
+	return nullptr;
+}
+
 TSharedRef<SDockTab> SLevelEditor::SummonDetailsPanel( FName TabIdentifier )
 {
 	TSharedRef<SActorDetails> ActorDetails = StaticCastSharedRef<SActorDetails>( CreateActorDetails( TabIdentifier ) );
@@ -728,15 +792,77 @@ TSharedRef<ISceneOutliner> SLevelEditor::CreateSceneOutliner(FName TabIdentifier
 	TSharedRef<ISceneOutliner> SceneOutlinerRef = SceneOutlinerModule.CreateActorBrowser(
 		InitOptions);
 
-	// Update the most recently created outliner
-	SceneOutlinerPtr = SceneOutlinerRef;
-
 	// Add this to the map of all outliners
 	SceneOutliners.Add(TabIdentifier, SceneOutlinerRef);
 
+	// Update the most recently used outliner
+	SetMostRecentlyUsedSceneOutliner(TabIdentifier);
+
 	return SceneOutlinerRef;
 }
 
+void SLevelEditor::OnExtendSceneOutlinerTabContextMenu(FMenuBuilder& InMenuBuilder)
+{
+	InMenuBuilder.BeginSection("LevelEditorOutlinerContextMenu", LOCTEXT("LevelEditorOutlinerContextMenuOptions", "Outliner Tabs"));
+
+	for (const TPair<FName, FText>& OutlinerTabName : SceneOutlinerDisplayNames)
+	{
+		InMenuBuilder.AddMenuEntry(OutlinerTabName.Value,
+			FText(),
+			FSlateIcon(),
+			FUIAction(
+				FExecuteAction::CreateLambda([OutlinerTabName]()
+					{
+						FLevelEditorModule& LevelEditorModule = FModuleManager::GetModuleChecked<FLevelEditorModule>(LevelEditorModuleName);
+						TSharedPtr<FTabManager> LevelEditorTabManager = LevelEditorModule.GetLevelEditorTabManager();
+
+						// Close the tab if it is already open, otherwise open it
+						if (TSharedPtr<SDockTab> LevelEditorOutlinerTab = LevelEditorTabManager->FindExistingLiveTab(OutlinerTabName.Key))
+						{
+							LevelEditorOutlinerTab->RequestCloseTab();
+						}
+						else
+						{
+							LevelEditorTabManager->TryInvokeTab(OutlinerTabName.Key);
+						}
+					}
+				),
+				FCanExecuteAction(),
+				FIsActionChecked::CreateLambda([OutlinerTabName]()
+					{
+						FLevelEditorModule& LevelEditorModule = FModuleManager::GetModuleChecked<FLevelEditorModule>(LevelEditorModuleName);
+						TSharedPtr<FTabManager> LevelEditorTabManager = LevelEditorModule.GetLevelEditorTabManager();
+
+						return LevelEditorTabManager->FindExistingLiveTab(OutlinerTabName.Key) != nullptr;
+					}
+				)),
+			NAME_None,
+			EUserInterfaceActionType::Check);
+	}
+
+	InMenuBuilder.EndSection();
+}
+
+FText SLevelEditor::GetSceneOutlinerLabel(FName SceneOutlinerTabIdentifier)
+{
+	int32 NumActiveOutliners = 0;
+	for (const TPair<FName, TWeakPtr<SDockTab>>& OutlinerTab : SceneOutlinerTabs)
+	{
+		if (OutlinerTab.Value.IsValid())
+		{
+			++NumActiveOutliners;
+		}
+	}
+
+	// If there is only one outliner open, we don't need to call it "Outliner 1"
+	if (NumActiveOutliners == 1)
+	{
+		return NSLOCTEXT("LevelEditor", "SceneOutlinerTabTitle", "Outliner");
+	}
+
+	return SceneOutlinerDisplayNames[SceneOutlinerTabIdentifier];
+}
+
 /** Method to call when a tab needs to be spawned by the FLayoutService */
 TSharedRef<SDockTab> SLevelEditor::SpawnLevelEditorTab( const FSpawnTabArgs& Args, FName TabIdentifier, FString InitializationPayload )
 {
@@ -790,11 +916,12 @@ TSharedRef<SDockTab> SLevelEditor::SpawnLevelEditorTab( const FSpawnTabArgs& Arg
 	{
 		TSharedRef<ISceneOutliner> SceneOutlinerRef = CreateSceneOutliner(TabIdentifier);
 
-		FText Label = NSLOCTEXT("LevelEditor", "SceneOutlinerTabTitle", "Outliner");
-
+		TAttribute<FText> Label = TAttribute<FText>::Create(TAttribute<FText>::FGetter::CreateRaw(this, &SLevelEditor::GetSceneOutlinerLabel, TabIdentifier));
+		
 		TSharedRef<SDockTab> SceneOutlinerTab = SNew( SDockTab )
 			.Label( Label )
 			.ToolTip( IDocumentation::Get()->CreateToolTip( Label, nullptr, "Shared/LevelEditor", "SceneOutlinerTab" ) )
+			.OnExtendContextMenu(this, &SLevelEditor::OnExtendSceneOutlinerTabContextMenu)
 			[
 				SNew(SBorder)
 				.Padding( 0 )
@@ -1225,25 +1352,25 @@ TSharedRef<SWidget> SLevelEditor::RestoreContentArea( const TSharedRef<SDockTab>
 			const FSlateIcon OutlinerIcon(FAppStyle::GetAppStyleSetName(), "LevelEditor.Tabs.Outliner");
 
 			LevelEditorTabManager->RegisterTabSpawner(LevelEditorTabIds::LevelEditorSceneOutliner, FOnSpawnTab::CreateSP<SLevelEditor, FName, FString>(this, &SLevelEditor::SpawnLevelEditorTab, LevelEditorTabIds::LevelEditorSceneOutliner, FString()))
-				.SetDisplayName(NSLOCTEXT("LevelEditorTabs", "LevelEditorSceneOutliner", "Outliner 1"))
+				.SetDisplayName(SceneOutlinerDisplayNames[LevelEditorTabIds::LevelEditorSceneOutliner])
 				.SetTooltipText(OutlinerTooltip)
 				.SetGroup(MenuStructure.GetLevelEditorOutlinerCategory())
 				.SetIcon(OutlinerIcon);
 
 			LevelEditorTabManager->RegisterTabSpawner(LevelEditorTabIds::LevelEditorSceneOutliner2, FOnSpawnTab::CreateSP<SLevelEditor, FName, FString>(this, &SLevelEditor::SpawnLevelEditorTab, LevelEditorTabIds::LevelEditorSceneOutliner2, FString()))
-				.SetDisplayName(NSLOCTEXT("LevelEditorTabs", "LevelEditorSceneOutliner2", "Outliner 2"))
+				.SetDisplayName(SceneOutlinerDisplayNames[LevelEditorTabIds::LevelEditorSceneOutliner2])
 				.SetTooltipText(OutlinerTooltip)
 				.SetGroup(MenuStructure.GetLevelEditorOutlinerCategory())
 				.SetIcon(OutlinerIcon);
 
 			LevelEditorTabManager->RegisterTabSpawner(LevelEditorTabIds::LevelEditorSceneOutliner3, FOnSpawnTab::CreateSP<SLevelEditor, FName, FString>(this, &SLevelEditor::SpawnLevelEditorTab, LevelEditorTabIds::LevelEditorSceneOutliner3, FString()))
-				.SetDisplayName(NSLOCTEXT("LevelEditorTabs", "LevelEditorSceneOutliner3", "Outliner 3"))
+				.SetDisplayName(SceneOutlinerDisplayNames[LevelEditorTabIds::LevelEditorSceneOutliner3])
 				.SetTooltipText(OutlinerTooltip)
 				.SetGroup(MenuStructure.GetLevelEditorOutlinerCategory())
 				.SetIcon(OutlinerIcon);
 
 			LevelEditorTabManager->RegisterTabSpawner(LevelEditorTabIds::LevelEditorSceneOutliner4, FOnSpawnTab::CreateSP<SLevelEditor, FName, FString>(this, &SLevelEditor::SpawnLevelEditorTab, LevelEditorTabIds::LevelEditorSceneOutliner4, FString()))
-				.SetDisplayName(NSLOCTEXT("LevelEditorTabs", "LevelEditorSceneOutliner4", "Outliner 4"))
+				.SetDisplayName(SceneOutlinerDisplayNames[LevelEditorTabIds::LevelEditorSceneOutliner4])
 				.SetTooltipText(OutlinerTooltip)
 				.SetGroup(MenuStructure.GetLevelEditorOutlinerCategory())
 				.SetIcon(OutlinerIcon);
@@ -1742,15 +1869,15 @@ void SLevelEditor::HandleEditorMapChange( uint32 MapChangeFlags )
 	// We want to recreate all the Scene Outliners on loading a new map
 	if (MapChangeFlags == MapChangeEventFlags::NewMap)
 	{
-		for (TPair<FName, TWeakPtr<SDockTab>> OutlinerTabIdentifier : SceneOutlinerTabs)
+		for (TPair<FName, TWeakPtr<SDockTab>> OutlinerTab : SceneOutlinerTabs)
 		{
-			if (TSharedPtr<SDockTab> SceneOutlinerTabPin = OutlinerTabIdentifier.Value.Pin())
+			if (TSharedPtr<SDockTab> SceneOutlinerTabPin = OutlinerTab.Value.Pin())
 			{
 				// Create the new outliner
-				TSharedRef<ISceneOutliner> SceneOutlinerRef = CreateSceneOutliner(OutlinerTabIdentifier.Key);
+				TSharedRef<ISceneOutliner> SceneOutlinerRef = CreateSceneOutliner(OutlinerTab.Key);
 
 				// Add it to the list of outliners
-				SceneOutliners.Add(OutlinerTabIdentifier.Key, SceneOutlinerRef);
+				SceneOutliners.Add(OutlinerTab.Key, SceneOutlinerRef);
 
 				// Set the tabs content to the newly created outliner
 				SceneOutlinerTabPin->SetContent(SceneOutlinerRef);
diff --git a/Engine/Source/Editor/LevelEditor/Private/SLevelEditor.h b/Engine/Source/Editor/LevelEditor/Private/SLevelEditor.h
--- a/Engine/Source/Editor/LevelEditor/Private/SLevelEditor.h
+++ b/Engine/Source/Editor/LevelEditor/Private/SLevelEditor.h
@@ -150,14 +150,19 @@ public:
 	/** Attaches a sequencer asset editor used to animate objects in the level to this level editor */
 	void AttachSequencer( TSharedPtr<SWidget> SequencerWidget, TSharedPtr<IAssetEditorInstance> NewSequencerAssetEditor );
 
-	/** Returns the most recently created scene outliner. Use GetAllSceneOutliners() to get all potentially active outliners
-	 *  TODO: Add a way to identify the most recently interacted with outliner, and change/deprecate this function to return it
-	 */
-	virtual TSharedPtr<ISceneOutliner> GetSceneOutliner() const override { return SceneOutlinerPtr.Pin();  }
-
 	/** Get an array containing weak pointers to all 4 Scene Outliners which could be potentially active */
 	virtual TArray<TWeakPtr<ISceneOutliner>> GetAllSceneOutliners() const override;
 
+	/** Set the outliner with the given name as the most recently interacted with */
+	virtual void SetMostRecentlyUsedSceneOutliner(FName OutlinerIdentifier) override;
+	
+	/** Return the most recently interacted with Outliner */
+	virtual TSharedPtr<ISceneOutliner> GetMostRecentlyUsedSceneOutliner() override;
+	
+	/** Return the most recently interacted with Outliner */
+	UE_DEPRECATED(5.1, "The Level Editor has multiple outliners, use GetAllSceneOutliners() or GetMostRecentlyUsedSceneOutliner() instead to avoid ambiguity")
+	virtual TSharedPtr<ISceneOutliner> GetSceneOutliner() const override;
+	
 	TSharedRef<SWidget> GetTitleBarMessageWidget() const { return TtileBarMessageBox.ToSharedRef(); }
 private:
 	
@@ -266,6 +271,15 @@ private:
 	/** Create a Scene Outliner for the Level Editor */
 	TSharedRef<ISceneOutliner> CreateSceneOutliner(FName TabIdentifier);
 
+	/** Function to extend the context menu for the outliner tab */
+	void OnExtendSceneOutlinerTabContextMenu(FMenuBuilder& InMenuBuilder);
+
+	/** Helper function to get the display name for an outliner */
+	FText GetSceneOutlinerLabel(FName SceneOutlinerTabIdentifier);
+
+	/** Sets the first available outliner as the most recent outliner */
+	void ResetMostRecentOutliner();
+
 private:
 
 	// Tracking the active viewports in this level editor.
@@ -325,6 +339,9 @@ private:
 	/** Map containing Weak pointers to all the Scene Outliner Tabs in the level editor */
 	TMap<FName, TWeakPtr<SDockTab>> SceneOutlinerTabs;
 
+	/** Map containing the display names of all the outliners */
+	TMap<FName, FText> SceneOutlinerDisplayNames;
+
 	/** Handle to the registered OnPreviewFeatureLevelChanged delegate. */
 	FDelegateHandle PreviewFeatureLevelChangedHandle;
 
diff --git a/Engine/Source/Editor/LevelEditor/Public/ILevelEditor.h b/Engine/Source/Editor/LevelEditor/Public/ILevelEditor.h
--- a/Engine/Source/Editor/LevelEditor/Public/ILevelEditor.h
+++ b/Engine/Source/Editor/LevelEditor/Public/ILevelEditor.h
@@ -88,9 +88,18 @@ public:
 	/** Sets the UI customization of the SCSEditor inside the level editor details panel. */
 	virtual void SetActorDetailsSCSEditorUICustomization(TSharedPtr<class ISCSEditorUICustomization> ActorDetailsSCSEditorUICustomization) = 0;
 
+	/** Return the most recently interacted with Outliner */
+	UE_DEPRECATED(5.1, "The Level Editor has multiple outliners, use GetAllSceneOutliners() or GetMostRecentlyUsedSceneOutliner() instead to avoid ambiguity")
 	virtual TSharedPtr<ISceneOutliner> GetSceneOutliner() const = 0;
 
+	/** Get an array containing weak pointers to all 4 Scene Outliners which could be potentially active */
 	virtual TArray<TWeakPtr<ISceneOutliner>> GetAllSceneOutliners() const = 0;
+
+	/** Set the outliner with the given name as the most recently interacted with */
+	virtual void SetMostRecentlyUsedSceneOutliner(FName OutlinerIdentifier) = 0;
+
+	/** Return the most recently interacted with Outliner */
+	virtual TSharedPtr<ISceneOutliner> GetMostRecentlyUsedSceneOutliner() = 0;
 };
 
 
diff --git a/Engine/Source/Editor/SceneOutliner/Private/ActorBrowsingMode.cpp b/Engine/Source/Editor/SceneOutliner/Private/ActorBrowsingMode.cpp
--- a/Engine/Source/Editor/SceneOutliner/Private/ActorBrowsingMode.cpp
+++ b/Engine/Source/Editor/SceneOutliner/Private/ActorBrowsingMode.cpp
@@ -40,10 +40,6 @@ DEFINE_LOG_CATEGORY_STATIC(LogActorBrowser, Log, All);
 
 #define LOCTEXT_NAMESPACE "SceneOutliner_ActorBrowsingMode"
 
-UActorBrowsingModeSettings::UActorBrowsingModeSettings(const FObjectInitializer& ObjectInitializer)
-: Super(ObjectInitializer)
-{ }
-
 using FActorFilter = TSceneOutlinerPredicateFilter<FActorTreeItem>;
 using FActorDescFilter = TSceneOutlinerPredicateFilter<FActorDescTreeItem>;
 static const FFolder GWorldRoot(FFolder::GetEmptyPath(), nullptr);
@@ -70,43 +66,68 @@ FActorBrowsingMode::FActorBrowsingMode(SSceneOutliner* InSceneOutliner, TWeakObj
 	FEditorDelegates::OnDeleteActorsBegin.AddRaw(this, &FActorBrowsingMode::OnDeleteActorsBegin);
 	FEditorDelegates::OnDeleteActorsEnd.AddRaw(this, &FActorBrowsingMode::OnDeleteActorsEnd);
 
-	UActorBrowsingModeSettings* SharedSettings = GetMutableDefault<UActorBrowsingModeSettings>();
+	UActorBrowserConfig::Initialize();
+	UActorBrowserConfig::Get()->LoadEditorConfig();
+
+	// Get a MutableConfig here to force create a config for the current outliner if it doesn't exist
+	const FActorBrowsingModeConfig* SavedSettings = GetMutableConfig();
+
+	// Create a local struct to use the default values if this outliner doesn't want to save configs
+	FActorBrowsingModeConfig LocalSettings;
+
+	// If this outliner doesn't want to save config (OutlinerIdentifier is empty, use the defaults)
+	if (SavedSettings)
+	{
+		LocalSettings = *SavedSettings;
+	}
+
 	// Get the OutlinerModule to register FilterInfos with the FilterInfoMap
-	FSceneOutlinerFilterInfo ShowOnlySelectedActorsInfo(LOCTEXT("ToggleShowOnlySelected", "Only Selected"), LOCTEXT("ToggleShowOnlySelectedToolTip", "When enabled, only displays actors that are currently selected."), SharedSettings->bShowOnlySelectedActors, FCreateSceneOutlinerFilter::CreateStatic(&FActorBrowsingMode::CreateShowOnlySelectedActorsFilter));
+	FSceneOutlinerFilterInfo ShowOnlySelectedActorsInfo(LOCTEXT("ToggleShowOnlySelected", "Only Selected"), LOCTEXT("ToggleShowOnlySelectedToolTip", "When enabled, only displays actors that are currently selected."), LocalSettings.bShowOnlySelectedActors, FCreateSceneOutlinerFilter::CreateStatic(&FActorBrowsingMode::CreateShowOnlySelectedActorsFilter));
 	ShowOnlySelectedActorsInfo.OnToggle().AddLambda([this](bool bIsActive)
 		{
-			UActorBrowsingModeSettings* Settings = GetMutableDefault<UActorBrowsingModeSettings>();
-			Settings->bShowOnlySelectedActors = bIsActive;
-			Settings->PostEditChange();
+			FActorBrowsingModeConfig* Settings = GetMutableConfig();
+			if(Settings)
+			{
+				Settings->bShowOnlySelectedActors = bIsActive;
+				SaveConfig();
+			}
 		});
 	FilterInfoMap.Add(TEXT("ShowOnlySelectedActors"), ShowOnlySelectedActorsInfo);
 
-	FSceneOutlinerFilterInfo HideTemporaryActorsInfo(LOCTEXT("ToggleHideTemporaryActors", "Hide Temporary Actors"), LOCTEXT("ToggleHideTemporaryActorsToolTip", "When enabled, hides temporary/run-time Actors."), SharedSettings->bHideTemporaryActors, FCreateSceneOutlinerFilter::CreateStatic(&FActorBrowsingMode::CreateHideTemporaryActorsFilter));
+	FSceneOutlinerFilterInfo HideTemporaryActorsInfo(LOCTEXT("ToggleHideTemporaryActors", "Hide Temporary Actors"), LOCTEXT("ToggleHideTemporaryActorsToolTip", "When enabled, hides temporary/run-time Actors."), LocalSettings.bHideTemporaryActors, FCreateSceneOutlinerFilter::CreateStatic(&FActorBrowsingMode::CreateHideTemporaryActorsFilter));
 	HideTemporaryActorsInfo.OnToggle().AddLambda([this](bool bIsActive)
 		{
-			UActorBrowsingModeSettings* Settings = GetMutableDefault<UActorBrowsingModeSettings>();
-			Settings->bHideTemporaryActors = bIsActive;
-			Settings->PostEditChange();
+			FActorBrowsingModeConfig* Settings = GetMutableConfig();
+			if(Settings)
+			{
+				Settings->bHideTemporaryActors = bIsActive;
+				SaveConfig();
+			}
 		});
 	FilterInfoMap.Add(TEXT("HideTemporaryActors"), HideTemporaryActorsInfo);
 
-	FSceneOutlinerFilterInfo OnlyCurrentLevelInfo(LOCTEXT("ToggleShowOnlyCurrentLevel", "Only in Current Level"), LOCTEXT("ToggleShowOnlyCurrentLevelToolTip", "When enabled, only shows Actors that are in the Current Level."), SharedSettings->bShowOnlyActorsInCurrentLevel, FCreateSceneOutlinerFilter::CreateStatic(&FActorBrowsingMode::CreateIsInCurrentLevelFilter));
+	FSceneOutlinerFilterInfo OnlyCurrentLevelInfo(LOCTEXT("ToggleShowOnlyCurrentLevel", "Only in Current Level"), LOCTEXT("ToggleShowOnlyCurrentLevelToolTip", "When enabled, only shows Actors that are in the Current Level."), LocalSettings.bShowOnlyActorsInCurrentLevel, FCreateSceneOutlinerFilter::CreateStatic(&FActorBrowsingMode::CreateIsInCurrentLevelFilter));
 	OnlyCurrentLevelInfo.OnToggle().AddLambda([this](bool bIsActive)
 		{
-			UActorBrowsingModeSettings* Settings = GetMutableDefault<UActorBrowsingModeSettings>();
-			Settings->bShowOnlyActorsInCurrentLevel = bIsActive;
-			Settings->PostEditChange();
+			FActorBrowsingModeConfig* Settings = GetMutableConfig();
+			if(Settings)
+			{
+				Settings->bShowOnlyActorsInCurrentLevel = bIsActive;
+				SaveConfig();
+			}
 		});
 	FilterInfoMap.Add(TEXT("ShowOnlyCurrentLevel"), OnlyCurrentLevelInfo);
 
-	bHideComponents = SharedSettings->bHideActorComponents;
-	FSceneOutlinerFilterInfo HideComponentsInfo(LOCTEXT("ToggleHideActorComponents", "Hide Actor Components"), LOCTEXT("ToggleHideActorComponentsToolTip", "When enabled, hides components belonging to actors."), SharedSettings->bHideActorComponents, FCreateSceneOutlinerFilter::CreateStatic(&FActorBrowsingMode::CreateHideComponentsFilter));
+	bHideComponents = LocalSettings.bHideActorComponents;
+	FSceneOutlinerFilterInfo HideComponentsInfo(LOCTEXT("ToggleHideActorComponents", "Hide Actor Components"), LOCTEXT("ToggleHideActorComponentsToolTip", "When enabled, hides components belonging to actors."), LocalSettings.bHideActorComponents, FCreateSceneOutlinerFilter::CreateStatic(&FActorBrowsingMode::CreateHideComponentsFilter));
 	HideComponentsInfo.OnToggle().AddLambda([this](bool bIsActive)
 		{
-			UActorBrowsingModeSettings* Settings = GetMutableDefault<UActorBrowsingModeSettings>();
-			Settings->bHideActorComponents = bHideComponents = bIsActive;
-			Settings->PostEditChange();
-
+			FActorBrowsingModeConfig* Settings = GetMutableConfig();
+			if(Settings)
+			{
+				Settings->bHideActorComponents = bHideComponents = bIsActive;
+				SaveConfig();
+			}
 			if (auto ActorHierarchy = StaticCast<FActorHierarchy*>(Hierarchy.Get()))
 			{
 				ActorHierarchy->SetShowingComponents(!bIsActive);
@@ -115,12 +136,15 @@ FActorBrowsingMode::FActorBrowsingMode(SSceneOutliner* InSceneOutliner, TWeakObj
 
 	FilterInfoMap.Add(TEXT("HideComponentsFilter"), HideComponentsInfo);
 
-	FSceneOutlinerFilterInfo HideLevelInstancesInfo(LOCTEXT("ToggleHideLevelInstanceContent", "Hide Level Instance Content"), LOCTEXT("ToggleHideLevelInstancesToolTip", "When enabled, hides all level instance content."), SharedSettings->bHideLevelInstanceHierarchy, FCreateSceneOutlinerFilter::CreateStatic(&FActorBrowsingMode::CreateHideLevelInstancesFilter));
+	FSceneOutlinerFilterInfo HideLevelInstancesInfo(LOCTEXT("ToggleHideLevelInstanceContent", "Hide Level Instance Content"), LOCTEXT("ToggleHideLevelInstancesToolTip", "When enabled, hides all level instance content."), LocalSettings.bHideLevelInstanceHierarchy, FCreateSceneOutlinerFilter::CreateStatic(&FActorBrowsingMode::CreateHideLevelInstancesFilter));
 	HideLevelInstancesInfo.OnToggle().AddLambda([this](bool bIsActive)
 		{
-			UActorBrowsingModeSettings* Settings = GetMutableDefault<UActorBrowsingModeSettings>();
-			Settings->bHideLevelInstanceHierarchy = bHideLevelInstanceHierarchy = bIsActive;
-			Settings->PostEditChange();
+			FActorBrowsingModeConfig* Settings = GetMutableConfig();
+			if(Settings)
+			{
+				Settings->bHideLevelInstanceHierarchy = bHideLevelInstanceHierarchy = bIsActive;
+				SaveConfig();
+			}
 
 			if (auto ActorHierarchy = StaticCast<FActorHierarchy*>(Hierarchy.Get()))
 			{
@@ -129,13 +153,16 @@ FActorBrowsingMode::FActorBrowsingMode(SSceneOutliner* InSceneOutliner, TWeakObj
 		});
 	FilterInfoMap.Add(TEXT("HideLevelInstancesFilter"), HideLevelInstancesInfo);
 	
-	FSceneOutlinerFilterInfo HideUnloadedActorsInfo(LOCTEXT("ToggleHideUnloadedActors", "Hide Unloaded Actors"), LOCTEXT("ToggleHideUnloadedActorsToolTip", "When enabled, hides all unloaded world partition actors."), SharedSettings->bHideUnloadedActors, FCreateSceneOutlinerFilter::CreateStatic(&FActorBrowsingMode::CreateHideUnloadedActorsFilter));
+	FSceneOutlinerFilterInfo HideUnloadedActorsInfo(LOCTEXT("ToggleHideUnloadedActors", "Hide Unloaded Actors"), LOCTEXT("ToggleHideUnloadedActorsToolTip", "When enabled, hides all unloaded world partition actors."), LocalSettings.bHideUnloadedActors, FCreateSceneOutlinerFilter::CreateStatic(&FActorBrowsingMode::CreateHideUnloadedActorsFilter));
 	HideUnloadedActorsInfo.OnToggle().AddLambda([this] (bool bIsActive)
 		{
-			UActorBrowsingModeSettings* Settings = GetMutableDefault<UActorBrowsingModeSettings>();
-			Settings->bHideUnloadedActors = bHideUnloadedActors = bIsActive;
-			Settings->PostEditChange();
-		
+			FActorBrowsingModeConfig* Settings = GetMutableConfig();
+			if(Settings)
+			{
+				Settings->bHideUnloadedActors = bHideUnloadedActors = bIsActive;
+				SaveConfig();
+			}
+
 			if (auto ActorHierarchy = StaticCast<FActorHierarchy*>(Hierarchy.Get()))
 			{
 				ActorHierarchy->SetShowingUnloadedActors(!bIsActive);
@@ -164,6 +191,8 @@ FActorBrowsingMode::FActorBrowsingMode(SSceneOutliner* InSceneOutliner, TWeakObj
 			return true;
 		})));
 
+	bAlwaysFrameSelection = LocalSettings.bAlwaysFrameSelection;
+
 	Rebuild();
 }
 
@@ -280,8 +309,43 @@ FSlateColor FActorBrowsingMode::GetStatusTextColor() const
 	}
 }
 
+void FActorBrowsingMode::OnToggleAlwaysFrameSelection()
+{
+	bAlwaysFrameSelection = !bAlwaysFrameSelection;
+
+	FActorBrowsingModeConfig* Settings = GetMutableConfig();
+	if(Settings)
+	{
+		Settings->bAlwaysFrameSelection = bAlwaysFrameSelection;
+		SaveConfig();
+	}
+}
+
+bool FActorBrowsingMode::ShouldAlwaysFrameSelection()
+{
+	return bAlwaysFrameSelection;
+}
+
 void FActorBrowsingMode::CreateViewContent(FMenuBuilder& MenuBuilder)
 {
+	MenuBuilder.BeginSection("OutlinerSelectionOptions", LOCTEXT("OptionsHeading", "Options"));
+
+	MenuBuilder.AddMenuEntry(
+		LOCTEXT("AlwaysFrameSelection", "Always Frame Selection"),
+		LOCTEXT("AlwaysFrameSelectionTooltip", "Select an item in the outliner when it is selected through the level editor viewport."),
+		FSlateIcon(),
+		FUIAction(
+			FExecuteAction::CreateRaw(this, &FActorBrowsingMode::OnToggleAlwaysFrameSelection),
+			FCanExecuteAction(),
+			FIsActionChecked::CreateRaw(this, &FActorBrowsingMode::ShouldAlwaysFrameSelection)
+		),
+		NAME_None,
+		EUserInterfaceActionType::ToggleButton
+	);
+
+	MenuBuilder.EndSection();
+
+
 	MenuBuilder.BeginSection("AssetThumbnails", LOCTEXT("ShowWorldHeading", "World"));
 	{
 		MenuBuilder.AddSubMenu(
@@ -753,6 +817,9 @@ void FActorBrowsingMode::OnItemSelectionChanged(FSceneOutlinerTreeItemPtr TreeIt
 
 		// Fire selection changed event
 		GEditor->NoteSelectionChange();
+
+		// Set this outliner as the most recently interacted with
+		SetAsMostRecentOutliner();
 	}
 
 	SceneOutliner->RefreshSelection();
@@ -820,6 +887,8 @@ void FActorBrowsingMode::OnFilterTextCommited(FSceneOutlinerItemSelection& Selec
 	// Set keyboard focus to the SceneOutliner, so the user can perform keyboard commands that interact
 	// with selected actors (such as Delete, to delete selected actors.)
 	SceneOutliner->SetKeyboardFocus();
+
+	SetAsMostRecentOutliner();
 }
 
 void FActorBrowsingMode::OnItemPassesFilters(const ISceneOutlinerTreeItem& Item)
@@ -1933,4 +2002,35 @@ void FActorBrowsingMode::OnDeleteActorsEnd()
 {
 	SceneOutliner->DeleteFoldersEnd();
 }
+
+struct FActorBrowsingModeConfig* FActorBrowsingMode::GetMutableConfig()
+{
+	FName OutlinerIdentifier = SceneOutliner->GetOutlinerIdentifier();
+
+	if (OutlinerIdentifier.IsNone())
+	{
+		return nullptr;
+	}
+
+	return &UActorBrowserConfig::Get()->ActorBrowsers.FindOrAdd(OutlinerIdentifier);
+}
+
+
+const FActorBrowsingModeConfig* FActorBrowsingMode::GetConstConfig() const
+{
+	FName OutlinerIdentifier = SceneOutliner->GetOutlinerIdentifier();
+
+	if (OutlinerIdentifier.IsNone())
+	{
+		return nullptr;
+	}
+
+	return UActorBrowserConfig::Get()->ActorBrowsers.Find(OutlinerIdentifier);
+}
+
+void FActorBrowsingMode::SaveConfig()
+{
+	UActorBrowserConfig::Get()->SaveEditorConfig();
+}
+
 #undef LOCTEXT_NAMESPACE
\ No newline at end of file
diff --git a/Engine/Source/Editor/SceneOutliner/Private/ActorBrowsingModeSettings.cpp b/Engine/Source/Editor/SceneOutliner/Private/ActorBrowsingModeSettings.cpp
new file mode 100644
--- /dev/null
+++ b/Engine/Source/Editor/SceneOutliner/Private/ActorBrowsingModeSettings.cpp
@@ -0,0 +1,14 @@
+// Copyright Epic Games, Inc. All Rights Reserved.
+
+#include "ActorBrowsingModeSettings.h"
+
+TObjectPtr<UActorBrowserConfig> UActorBrowserConfig::Instance = nullptr;
+
+void UActorBrowserConfig::Initialize()
+{
+	if(!Instance)
+	{
+		Instance = NewObject<UActorBrowserConfig>(); 
+		Instance->AddToRoot();
+	}
+}
\ No newline at end of file
diff --git a/Engine/Source/Editor/SceneOutliner/Private/ActorBrowsingModeSettings.h b/Engine/Source/Editor/SceneOutliner/Private/ActorBrowsingModeSettings.h
--- a/Engine/Source/Editor/SceneOutliner/Private/ActorBrowsingModeSettings.h
+++ b/Engine/Source/Editor/SceneOutliner/Private/ActorBrowsingModeSettings.h
@@ -5,39 +5,60 @@
 #include "CoreMinimal.h"
 #include "UObject/ObjectMacros.h"
 #include "UObject/Object.h"
-#include "ActorBrowsingModeSettings.generated.h"
+#include "EditorConfigBase.h"
 
+#include "ActorBrowsingModeSettings.generated.h"
 
-/**
- * Implements the settings for the Actor Browsing Outliner (also known as the World Outliner).
- */
-UCLASS(config = EditorPerProjectUserSettings)
-class UActorBrowsingModeSettings
-	: public UObject
+USTRUCT()
+struct FActorBrowsingModeConfig
 {
-	GENERATED_UCLASS_BODY()
+	GENERATED_BODY()
+
+public:
 
 	/** True when the Scene Outliner is hiding temporary/run-time Actors */
-	UPROPERTY(config)
-	uint32 bHideTemporaryActors : 1;
+	UPROPERTY()
+	bool bHideTemporaryActors = true;
 
 	/** True when the Scene Outliner is showing only Actors that exist in the current level */
-	UPROPERTY(config)
-	uint32 bShowOnlyActorsInCurrentLevel : 1;
+	UPROPERTY()
+	bool bShowOnlyActorsInCurrentLevel = false;
 
 	/** True when the Scene Outliner is only displaying selected Actors */
-	UPROPERTY(config)
-	uint32 bShowOnlySelectedActors : 1;
+	UPROPERTY()
+	bool bShowOnlySelectedActors = false;
 
 	/** True when the Scene Outliner is not displaying Actor Components*/
-	UPROPERTY(config)
-	uint32 bHideActorComponents : 1;
+	UPROPERTY()
+	bool bHideActorComponents = true;
 
 	/** True when the Scene Outliner is not displaying LevelInstances */
-	UPROPERTY(config)
-	uint32 bHideLevelInstanceHierarchy : 1;
+	UPROPERTY()
+	bool bHideLevelInstanceHierarchy = false;
 
 	/** True when the Scene Outliner is not displaying unloaded actors */
-	UPROPERTY(config)
-	uint32 bHideUnloadedActors : 1;
+	UPROPERTY()
+	bool bHideUnloadedActors = false;
+
+	/** True when the Scene Outliner updates when an actor is selected in the viewport */
+	UPROPERTY()
+	bool bAlwaysFrameSelection = true;
 };
+
+UCLASS(EditorConfig="ActorBrowsingMode")
+class UActorBrowserConfig : public UEditorConfigBase
+{
+	GENERATED_BODY()
+	
+public:
+
+	static void Initialize();
+	static UActorBrowserConfig* Get() { return Instance; }
+
+	UPROPERTY(meta=(EditorConfig))
+	TMap<FName, FActorBrowsingModeConfig> ActorBrowsers;
+
+private:
+
+	static TObjectPtr<UActorBrowserConfig> Instance;
+};
\ No newline at end of file
diff --git a/Engine/Source/Editor/SceneOutliner/Private/ActorMode.cpp b/Engine/Source/Editor/SceneOutliner/Private/ActorMode.cpp
--- a/Engine/Source/Editor/SceneOutliner/Private/ActorMode.cpp
+++ b/Engine/Source/Editor/SceneOutliner/Private/ActorMode.cpp
@@ -19,6 +19,8 @@
 #include "LevelInstance/LevelInstanceInterface.h"
 #include "LevelInstance/LevelInstanceSubsystem.h"
 #include "LevelInstance/LevelInstanceEditorInstanceActor.h"
+#include "LevelEditor.h"
+#include "Modules/ModuleManager.h"
 
 #define LOCTEXT_NAMESPACE "SceneOutliner_ActorMode"
 
@@ -359,6 +361,18 @@ void FActorMode::OnFilterTextChanged(const FText& InFilterText)
 
 		// In any other case (i.e., if the object passes the current filter), re-add it
 		SceneOutliner->ScrollItemIntoView(TreeItem);
+
+		SetAsMostRecentOutliner();
+	}
+}
+
+void FActorMode::SetAsMostRecentOutliner() const
+{
+	TWeakPtr<ILevelEditor> LevelEditor = FModuleManager::GetModuleChecked<FLevelEditorModule>(TEXT("LevelEditor")).GetLevelEditorInstance();
+
+	if(TSharedPtr<ILevelEditor> LevelEditorPin = LevelEditor.Pin())
+	{
+		LevelEditorPin->SetMostRecentlyUsedSceneOutliner(SceneOutliner->GetOutlinerIdentifier());
 	}
 }
 
diff --git a/Engine/Source/Editor/SceneOutliner/Private/ActorModeInteractive.cpp b/Engine/Source/Editor/SceneOutliner/Private/ActorModeInteractive.cpp
--- a/Engine/Source/Editor/SceneOutliner/Private/ActorModeInteractive.cpp
+++ b/Engine/Source/Editor/SceneOutliner/Private/ActorModeInteractive.cpp
@@ -4,6 +4,8 @@
 #include "Engine/Selection.h"
 #include "Editor.h"
 #include "ActorTreeItem.h"
+#include "LevelEditor.h"
+#include "Modules/ModuleManager.h"
 
 FActorModeInteractive::FActorModeInteractive(const FActorModeParams& Params)
 	: FActorMode(Params)
@@ -61,16 +63,22 @@ void FActorModeInteractive::OnLevelSelectionChanged(UObject* Obj)
 		SceneOutliner->ClearSelection();
 		SceneOutliner->RefreshSelection();
 
-		// Scroll last item into view - this means if we are multi-selecting, we show newest selection. @TODO Not perfect though
+		// Scroll last item into view if Selection Framing is enabled - this means if we are multi-selecting, we show newest selection. @TODO Not perfect though
 		if (const AActor* LastSelectedActor = GEditor->GetSelectedActors()->GetBottom<AActor>())
 		{
 			if (FSceneOutlinerTreeItemPtr TreeItem = SceneOutliner->GetTreeItem(LastSelectedActor, false))
 			{
-				SceneOutliner->ScrollItemIntoView(TreeItem);
+				if(bAlwaysFrameSelection)
+				{
+					SceneOutliner->ScrollItemIntoView(TreeItem);
+				}
 			}
 			else
 			{
-				SceneOutliner->OnItemAdded(LastSelectedActor, SceneOutliner::ENewItemAction::ScrollIntoView);
+				SceneOutliner::ENewItemAction::Type Action = bAlwaysFrameSelection ? SceneOutliner::ENewItemAction::ScrollIntoView : SceneOutliner::ENewItemAction::Select;
+				
+				SceneOutliner->OnItemAdded(LastSelectedActor, Action);
+				
 			}
 		}
 	}
@@ -78,15 +86,26 @@ void FActorModeInteractive::OnLevelSelectionChanged(UObject* Obj)
 
 void FActorModeInteractive::OnLevelActorRequestsRename(const AActor* Actor)
 {
-	const auto& SelectedItems = SceneOutliner->GetSelectedItems();
-	if (SelectedItems.Num() > 0)
+	TWeakPtr<ILevelEditor> LevelEditor = FModuleManager::GetModuleChecked<FLevelEditorModule>(TEXT("LevelEditor")).GetLevelEditorInstance();
+
+	if(TSharedPtr<ILevelEditor> LevelEditorPin = LevelEditor.Pin())
 	{
-		// Ensure that the item we want to rename is visible in the tree
-		FSceneOutlinerTreeItemPtr ItemToRename = SelectedItems[SelectedItems.Num() - 1];
-		if (SceneOutliner->CanExecuteRenameRequest(*ItemToRename) && ItemToRename->CanInteract())
+		/* We want to execute the rename on the most recently used outliner
+		 * TODO: Add a way to pop-out the outliner the rename is done on
+		 */
+		if(SceneOutliner == LevelEditorPin->GetMostRecentlyUsedSceneOutliner().Get())
 		{
-			SceneOutliner->SetPendingRenameItem(ItemToRename);
-			SceneOutliner->ScrollItemIntoView(ItemToRename);
+			const TArray<FSceneOutlinerTreeItemPtr>& SelectedItems = SceneOutliner->GetSelectedItems();
+			if (SelectedItems.Num() > 0)
+			{
+				// Ensure that the item we want to rename is visible in the tree
+				FSceneOutlinerTreeItemPtr ItemToRename = SelectedItems[SelectedItems.Num() - 1];
+				if (SceneOutliner->CanExecuteRenameRequest(*ItemToRename) && ItemToRename->CanInteract())
+				{
+					SceneOutliner->SetPendingRenameItem(ItemToRename);
+					SceneOutliner->ScrollItemIntoView(ItemToRename);
+				}
+			}
 		}
 	}
 }
diff --git a/Engine/Source/Editor/SceneOutliner/Private/SceneOutlinerConfig.cpp b/Engine/Source/Editor/SceneOutliner/Private/SceneOutlinerConfig.cpp
--- a/Engine/Source/Editor/SceneOutliner/Private/SceneOutlinerConfig.cpp
+++ b/Engine/Source/Editor/SceneOutliner/Private/SceneOutlinerConfig.cpp
@@ -6,6 +6,9 @@ TObjectPtr<UOutlinerConfig> UOutlinerConfig::Instance = nullptr;
 
 void UOutlinerConfig::Initialize()
 {
-	Instance = NewObject<UOutlinerConfig>(); 
-	Instance->AddToRoot();
+	if(!Instance)
+	{
+		Instance = NewObject<UOutlinerConfig>(); 
+		Instance->AddToRoot();
+	}
 }
\ No newline at end of file
diff --git a/Engine/Source/Editor/SceneOutliner/Private/SceneOutlinerConfig.h b/Engine/Source/Editor/SceneOutliner/Private/SceneOutlinerConfig.h
--- a/Engine/Source/Editor/SceneOutliner/Private/SceneOutlinerConfig.h
+++ b/Engine/Source/Editor/SceneOutliner/Private/SceneOutlinerConfig.h
@@ -21,7 +21,7 @@ public:
 
 	/** Whether the hierarchy is pinned at the top of the outliner*/
 	UPROPERTY()
-	bool bShouldStackHierarchyHeaders = false;
+	bool bShouldStackHierarchyHeaders = true;
 };
 
 
diff --git a/Engine/Source/Editor/SceneOutliner/Public/ActorBrowsingMode.h b/Engine/Source/Editor/SceneOutliner/Public/ActorBrowsingMode.h
--- a/Engine/Source/Editor/SceneOutliner/Public/ActorBrowsingMode.h
+++ b/Engine/Source/Editor/SceneOutliner/Public/ActorBrowsingMode.h
@@ -114,6 +114,27 @@ private:
 	static TSharedRef<FSceneOutlinerFilter> CreateHideComponentsFilter();
 	static TSharedRef<FSceneOutlinerFilter> CreateHideLevelInstancesFilter();
 	static TSharedRef<FSceneOutlinerFilter> CreateHideUnloadedActorsFilter();
+
+	/** Functions to expose selection framing to the UI */
+	void OnToggleAlwaysFrameSelection();
+	bool ShouldAlwaysFrameSelection();
+
+	/**
+	 * Get a mutable version of the ActorBrowser config for setting values.
+	 * @returns		The config for this ActorBrowser.
+	 * @note		If OutlinerIdentifier is not set for this outliner, it is not possible to store settings.
+	 */
+	struct FActorBrowsingModeConfig* GetMutableConfig();
+
+	/**
+	 * Get a const version of the ActorBrowser config for getting values.
+	 * @returns		The config for this ActorBrowser.
+	 * @note		If OutlinerIdentifier is not set for this outliner, it is not possible to store settings.
+	 */
+	const FActorBrowsingModeConfig* GetConstConfig() const;
+
+	/** Save the config for this ActorBrowser */
+	void SaveConfig();
 	
 private:
 	/** Number of actors (including unloaded) which have passed through the filters */
diff --git a/Engine/Source/Editor/SceneOutliner/Public/ActorMode.h b/Engine/Source/Editor/SceneOutliner/Public/ActorMode.h
--- a/Engine/Source/Editor/SceneOutliner/Public/ActorMode.h
+++ b/Engine/Source/Editor/SceneOutliner/Public/ActorMode.h
@@ -73,6 +73,9 @@ protected:
 	void SynchronizeActorSelection();
 	bool IsActorDisplayable(const AActor* InActor) const;
 
+	/** Set the Scene Outliner attached to this mode as the most recently used outliner in the Level Editor */
+	void SetAsMostRecentOutliner() const;
+
 	virtual TUniquePtr<ISceneOutlinerHierarchy> CreateHierarchy() override;
 protected:
 	/** The world which we are currently representing */
@@ -91,4 +94,6 @@ protected:
 	bool bHideLevelInstanceHierarchy;
 	/** Should unloaded actors be hidden */
 	bool bHideUnloadedActors;
+	/** Should the outliner scroll to the item on selection */
+	bool bAlwaysFrameSelection;
 };
diff --git a/Engine/Source/Editor/SceneOutliner/Public/SSceneOutliner.h b/Engine/Source/Editor/SceneOutliner/Public/SSceneOutliner.h
--- a/Engine/Source/Editor/SceneOutliner/Public/SSceneOutliner.h
+++ b/Engine/Source/Editor/SceneOutliner/Public/SSceneOutliner.h
@@ -468,6 +468,9 @@ public:
 	/** Used to test if Outliner related selection changes have already been handled */
 	bool GetIsReentrant() const { return bIsReentrant; }
 
+	/** Get the unique identifier associated with this outliner */
+	FName GetOutlinerIdentifier() const { return OutlinerIdentifier; }
+
 private:
 	/** Methods that implement structural modification logic for the tree */
 
@@ -886,14 +889,14 @@ private:
 	/**
 	 * Get a mutable version of the outliner config for setting values.
 	 * @returns		The outliner config for this outliner.
-	 * @note		If FSceneOutlinerInitializationOptions.ViewIdentifier is not set, it is not possible to store settings for this outliner.
+	 * @note		If FSceneOutlinerInitializationOptions.OutlinerIdentifier is not set, it is not possible to store settings for this outliner.
 	 */
 	struct FSceneOutlinerConfig* GetMutableConfig();
 
 	/**
 	 * Get a const version of the outliner config for getting values.
 	 * @returns		The outliner config for this outliner.
-	 * @note		If FSceneOutlinerInitializationOptions.ViewIdentifier is not set, it is not possible to retrieve settings for this outliner.
+	 * @note		If FSceneOutlinerInitializationOptions.OutlinerIdentifier is not set, it is not possible to retrieve settings for this outliner.
 	 */
 	const FSceneOutlinerConfig* GetConstConfig() const;
 
