Fixed static analysis
#rnx

diff --git a/Engine/Source/Editor/DataLayerEditor/Private/DataLayer/DataLayerActorTreeItem.h b/Engine/Source/Editor/DataLayerEditor/Private/DataLayer/DataLayerActorTreeItem.h
--- a/Engine/Source/Editor/DataLayerEditor/Private/DataLayer/DataLayerActorTreeItem.h
+++ b/Engine/Source/Editor/DataLayerEditor/Private/DataLayer/DataLayerActorTreeItem.h
@@ -82,12 +82,14 @@ private:
 	{
 		FActorTreeItem::UpdateDisplayString();
 
-		UWorld* OwningWorld = Actor.IsValid() ? Actor->GetWorld() : nullptr;
-		ULevel* Level = Actor.IsValid() ? Actor->GetLevel() : nullptr;
-		ULevelInstanceSubsystem* LevelInstanceSubsystem = UWorld::GetSubsystem<ULevelInstanceSubsystem>(OwningWorld);
-		if (LevelInstanceSubsystem && Level && (Level != OwningWorld->GetCurrentLevel()))
+		if (UWorld* OwningWorld = Actor.IsValid() ? Actor->GetWorld() : nullptr)
 		{
-			DisplayString = LevelInstanceSubsystem->PrefixWithParentLevelInstanceActorLabels(DisplayString, Actor->GetLevel());
+			ULevel* Level = Actor.IsValid() ? Actor->GetLevel() : nullptr;
+			ULevelInstanceSubsystem* LevelInstanceSubsystem = UWorld::GetSubsystem<ULevelInstanceSubsystem>(OwningWorld);
+			if (LevelInstanceSubsystem && Level && (Level != OwningWorld->GetCurrentLevel()))
+			{
+				DisplayString = LevelInstanceSubsystem->PrefixWithParentLevelInstanceActorLabels(DisplayString, Actor->GetLevel());
+			}
 		}
 	}
 
diff --git a/Engine/Source/Editor/SceneOutliner/Private/SceneOutlinerModule.cpp b/Engine/Source/Editor/SceneOutliner/Private/SceneOutlinerModule.cpp
--- a/Engine/Source/Editor/SceneOutliner/Private/SceneOutlinerModule.cpp
+++ b/Engine/Source/Editor/SceneOutliner/Private/SceneOutlinerModule.cpp
@@ -275,25 +275,26 @@ void FSceneOutlinerModule::CreateActorInfoColumns(FSceneOutlinerInitializationOp
 			{
 				if (const FWorldPartitionActorDesc* ActorDesc = ActorDescItem->ActorDescHandle.Get(); ActorDesc && !ActorDesc->GetDataLayerInstanceNames().IsEmpty())
 				{
-					const UActorDescContainer* ActorDescContainer = ActorDescItem->ActorDescHandle.Container.Get();
-					const UWorld* World = ActorDescContainer ? ActorDescContainer->GetWorld() : nullptr;
-
-					if (const UDataLayerSubsystem* DataLayerSubsystem = UWorld::GetSubsystem<UDataLayerSubsystem>(World))
+					if (const UActorDescContainer* ActorDescContainer = ActorDescItem->ActorDescHandle.Container.Get())
 					{
-						TSet<const UDataLayerInstance*> DataLayerInstances;
-						DataLayerInstances.Append(DataLayerSubsystem->GetDataLayerInstances(ActorDesc->GetDataLayerInstanceNames()));
-						if (ULevelInstanceSubsystem* LevelInstanceSubsystem = UWorld::GetSubsystem<ULevelInstanceSubsystem>(World))
+						const UWorld* World = ActorDescContainer->GetWorld();
+						if (const UDataLayerSubsystem* DataLayerSubsystem = UWorld::GetSubsystem<UDataLayerSubsystem>(World))
 						{
-							UWorld* OuterWorld = ActorDescContainer->GetTypedOuter<UWorld>();
-							// Add parent container Data Layer Instances
-							AActor* CurrentActor = OuterWorld ? Cast<AActor>(LevelInstanceSubsystem->GetOwningLevelInstance(OuterWorld->PersistentLevel)) : nullptr;
-							while (CurrentActor)
+							TSet<const UDataLayerInstance*> DataLayerInstances;
+							DataLayerInstances.Append(DataLayerSubsystem->GetDataLayerInstances(ActorDesc->GetDataLayerInstanceNames()));
+							if (ULevelInstanceSubsystem* LevelInstanceSubsystem = UWorld::GetSubsystem<ULevelInstanceSubsystem>(World))
 							{
-								DataLayerInstances.Append(bUseLevelContext ? CurrentActor->GetDataLayerInstancesForLevel() : CurrentActor->GetDataLayerInstances());
-								CurrentActor = Cast<AActor>(LevelInstanceSubsystem->GetParentLevelInstance(CurrentActor));
-							};
+								UWorld* OuterWorld = ActorDescContainer->GetTypedOuter<UWorld>();
+								// Add parent container Data Layer Instances
+								AActor* CurrentActor = OuterWorld ? Cast<AActor>(LevelInstanceSubsystem->GetOwningLevelInstance(OuterWorld->PersistentLevel)) : nullptr;
+								while (CurrentActor)
+								{
+									DataLayerInstances.Append(bUseLevelContext ? CurrentActor->GetDataLayerInstancesForLevel() : CurrentActor->GetDataLayerInstances());
+									CurrentActor = Cast<AActor>(LevelInstanceSubsystem->GetParentLevelInstance(CurrentActor));
+								};
+							}
+							BuildDataLayers(DataLayerInstances, bUseLevelContext);
 						}
-						BuildDataLayers(DataLayerInstances, bUseLevelContext);
 					}
 				}
 			}
