World Partition: don't allow pin actors for game worlds.

#rnx

diff --git a/Engine/Source/Editor/SceneOutliner/Private/ActorBrowsingMode.cpp b/Engine/Source/Editor/SceneOutliner/Private/ActorBrowsingMode.cpp
--- a/Engine/Source/Editor/SceneOutliner/Private/ActorBrowsingMode.cpp
+++ b/Engine/Source/Editor/SceneOutliner/Private/ActorBrowsingMode.cpp
@@ -241,6 +241,7 @@ void FActorBrowsingMode::Rebuild()
 	ApplicableUnloadedActors.Empty();
 	ApplicableActors.Empty();
 
+	bRepresentingWorldGameWorld = RepresentingWorld.IsValid() && RepresentingWorld->IsGameWorld();
 	bRepresentingWorldPartitionedWorld = RepresentingWorld.IsValid() && RepresentingWorld->IsPartitionedWorld();
 
 	if (bRepresentingWorldPartitionedWorld)
@@ -496,7 +497,7 @@ void FActorBrowsingMode::RegisterContextMenu()
 					if (Context->bShowParentTree && Context->NumSelectedItems > 0 && Context->SceneOutliner.IsValid())
 					{
 						// Only add the menu option to wp levels
-						if (Context->bRepresentingPartitionedWorld)
+						if (!Context->bRepresentingGameWorld && Context->bRepresentingPartitionedWorld)
 						{
 							// If selection contains some unpinned items, show the pin option
 							// If the selection contains folders, always show the pin option
@@ -603,6 +604,7 @@ TSharedPtr<SWidget> FActorBrowsingMode::BuildContextMenu()
 	ContextObject->NumSelectedItems = ItemSelection.Num();
 	ContextObject->NumSelectedFolders = ItemSelection.Num<FFolderTreeItem>();
 	ContextObject->NumWorldsSelected = ItemSelection.Num<FWorldTreeItem>();
+	ContextObject->bRepresentingGameWorld = bRepresentingWorldGameWorld;
 	ContextObject->bRepresentingPartitionedWorld = bRepresentingWorldPartitionedWorld;
 
 	int32 NumPinnedItems = 0;
diff --git a/Engine/Source/Editor/SceneOutliner/Private/ActorFolderTreeItem.cpp b/Engine/Source/Editor/SceneOutliner/Private/ActorFolderTreeItem.cpp
--- a/Engine/Source/Editor/SceneOutliner/Private/ActorFolderTreeItem.cpp
+++ b/Engine/Source/Editor/SceneOutliner/Private/ActorFolderTreeItem.cpp
@@ -350,7 +350,7 @@ TSharedRef<SWidget> FActorFolderTreeItem::GenerateLabelWidget(ISceneOutliner& Ou
 
 bool FActorFolderTreeItem::ShouldShowPinnedState() const
 {
-	return GetRootObject() == FFolder::GetDefaultRootObject() && World.IsValid() && World->IsPartitionedWorld();
+	return GetRootObject() == FFolder::GetDefaultRootObject() && World.IsValid() && !World->IsGameWorld() && World->IsPartitionedWorld();
 }
 
 bool FActorFolderTreeItem::CanInteract() const
diff --git a/Engine/Source/Editor/SceneOutliner/Private/ActorTreeItem.cpp b/Engine/Source/Editor/SceneOutliner/Private/ActorTreeItem.cpp
--- a/Engine/Source/Editor/SceneOutliner/Private/ActorTreeItem.cpp
+++ b/Engine/Source/Editor/SceneOutliner/Private/ActorTreeItem.cpp
@@ -391,7 +391,8 @@ bool FActorTreeItem::ShouldShowPinnedState() const
 	if (const AActor* ActorPtr = Actor.Get())
 	{
 		const ULevel* Level = ActorPtr->GetLevel();
-		return !!Level->GetWorldPartition();
+		const UWorld* World = Level->GetWorld();
+		return !World->IsGameWorld() && !!Level->GetWorldPartition();
 	}
 
 	return false;
diff --git a/Engine/Source/Editor/SceneOutliner/Public/ActorBrowsingMode.h b/Engine/Source/Editor/SceneOutliner/Public/ActorBrowsingMode.h
--- a/Engine/Source/Editor/SceneOutliner/Public/ActorBrowsingMode.h
+++ b/Engine/Source/Editor/SceneOutliner/Public/ActorBrowsingMode.h
@@ -146,5 +146,6 @@ private:
 	/** List of actors which passed the regular filters and may or may not have passed the search filter */
 	TSet<TWeakObjectPtr<AActor>> ApplicableActors;
 
+	bool bRepresentingWorldGameWorld = false;
 	bool bRepresentingWorldPartitionedWorld = false;
 };
diff --git a/Engine/Source/Editor/SceneOutliner/Public/SceneOutlinerMenuContext.h b/Engine/Source/Editor/SceneOutliner/Public/SceneOutlinerMenuContext.h
--- a/Engine/Source/Editor/SceneOutliner/Public/SceneOutlinerMenuContext.h
+++ b/Engine/Source/Editor/SceneOutliner/Public/SceneOutlinerMenuContext.h
@@ -20,6 +20,7 @@ public:
 	{
 		bShowParentTree = false;
 		bRepresentingPartitionedWorld = false;
+		bRepresentingGameWorld = false;
 		NumSelectedItems = 0;
 		NumSelectedFolders = 0;
 		NumWorldsSelected = 0;
@@ -29,6 +30,7 @@ public:
 	TWeakPtr<SSceneOutliner> SceneOutliner;
 
 	bool bShowParentTree;
+	bool bRepresentingGameWorld;
 	bool bRepresentingPartitionedWorld;
 	int32 NumSelectedItems;
 	int32 NumSelectedFolders;
