World Partition: fix unpinning an unloaded actor+ display better unloaded message when pinned actor is unloaded.

#rnx

diff --git a/Engine/Source/Editor/SceneOutliner/Private/ActorBaseTreeItem.cpp b/Engine/Source/Editor/SceneOutliner/Private/ActorBaseTreeItem.cpp
new file mode 100644
--- /dev/null
+++ b/Engine/Source/Editor/SceneOutliner/Private/ActorBaseTreeItem.cpp
@@ -0,0 +1,5 @@
+// Copyright Epic Games, Inc. All Rights Reserved.
+
+#include "ActorBaseTreeItem.h"
+
+const FSceneOutlinerTreeItemType IActorBaseTreeItem::Type(&ISceneOutlinerTreeItem::Type);
\ No newline at end of file
diff --git a/Engine/Source/Editor/SceneOutliner/Private/ActorBrowsingMode.cpp b/Engine/Source/Editor/SceneOutliner/Private/ActorBrowsingMode.cpp
--- a/Engine/Source/Editor/SceneOutliner/Private/ActorBrowsingMode.cpp
+++ b/Engine/Source/Editor/SceneOutliner/Private/ActorBrowsingMode.cpp
@@ -610,9 +610,9 @@ TSharedPtr<SWidget> FActorBrowsingMode::BuildContextMenu()
 	int32 NumPinnedItems = 0;
 	if (const UWorldPartition* const WorldPartition = RepresentingWorld->GetWorldPartition())
 	{
-		ItemSelection.ForEachItem<FActorTreeItem>([WorldPartition, &NumPinnedItems](const FActorTreeItem& ActorItem)
+		ItemSelection.ForEachItem<IActorBaseTreeItem>([WorldPartition, &NumPinnedItems](const IActorBaseTreeItem& ActorItem)
 		{
-			if (ActorItem.Actor.IsValid() && WorldPartition->IsActorPinned(ActorItem.Actor->GetActorGuid()))
+			if (WorldPartition->IsActorPinned(ActorItem.GetGuid()))
 			{
 				++NumPinnedItems;
 			}
diff --git a/Engine/Source/Editor/SceneOutliner/Private/ActorDescTreeItem.cpp b/Engine/Source/Editor/SceneOutliner/Private/ActorDescTreeItem.cpp
--- a/Engine/Source/Editor/SceneOutliner/Private/ActorDescTreeItem.cpp
+++ b/Engine/Source/Editor/SceneOutliner/Private/ActorDescTreeItem.cpp
@@ -24,7 +24,7 @@
 
 #define LOCTEXT_NAMESPACE "SceneOutliner_ActorDescTreeItem"
 
-const FSceneOutlinerTreeItemType FActorDescTreeItem::Type(&ISceneOutlinerTreeItem::Type);
+const FSceneOutlinerTreeItemType FActorDescTreeItem::Type(&IActorBaseTreeItem::Type);
 
 struct SActorDescTreeLabel : FSceneOutlinerCommonLabelData, public SCompoundWidget
 {
@@ -118,8 +118,17 @@ private:
 			{
 				FFormatNamedArguments Args;
 				Args.Add(TEXT("ActorLabel"), FText::FromString(TreeItem->GetDisplayString()));
-				Args.Add(TEXT("UnloadedTag"), LOCTEXT("UnloadedActorLabel", "(Unloaded)"));
-				return FText::Format(LOCTEXT("UnloadedActorDisplay", "{ActorLabel} {UnloadedTag}"), Args);
+
+				if (UWorldPartition* WorldPartition = Cast<UWorldPartition>(ActorDesc->GetContainer()); WorldPartition->IsActorPinned(ActorDesc->GetGuid()))
+				{
+					Args.Add(TEXT("UnloadState"), LOCTEXT("UnloadedDataLayer", "(Unloaded DataLayer)"));
+				}
+				else
+				{
+					Args.Add(TEXT("UnloadState"), LOCTEXT("UnloadedActorLabel", "(Unloaded)"));
+				}
+
+				return FText::Format(LOCTEXT("UnloadedActorDisplay", "{ActorLabel} {UnloadState}"), Args);
 			}
 		}
 		return FText();
@@ -211,7 +220,7 @@ private:
 };
 
 FActorDescTreeItem::FActorDescTreeItem(const FGuid& InActorGuid, UActorDescContainer* Container)
-	: ISceneOutlinerTreeItem(Type)
+	: IActorBaseTreeItem(Type)
 	, ActorDescHandle(Container, InActorGuid)
 	, ActorGuid(InActorGuid)
 {
diff --git a/Engine/Source/Editor/SceneOutliner/Private/ActorTreeItem.cpp b/Engine/Source/Editor/SceneOutliner/Private/ActorTreeItem.cpp
--- a/Engine/Source/Editor/SceneOutliner/Private/ActorTreeItem.cpp
+++ b/Engine/Source/Editor/SceneOutliner/Private/ActorTreeItem.cpp
@@ -24,7 +24,7 @@
 
 #define LOCTEXT_NAMESPACE "SceneOutliner_ActorTreeItem"
 
-const FSceneOutlinerTreeItemType FActorTreeItem::Type(&ISceneOutlinerTreeItem::Type);
+const FSceneOutlinerTreeItemType FActorTreeItem::Type(&IActorBaseTreeItem::Type);
 
 struct SActorTreeLabel : FSceneOutlinerCommonLabelData, public SCompoundWidget
 {
@@ -334,7 +334,7 @@ private:
 };
 
 FActorTreeItem::FActorTreeItem(AActor* InActor)
-	: ISceneOutlinerTreeItem(Type)
+	: IActorBaseTreeItem(Type)
 	, Actor(InActor)
 	, ID(InActor)
 {
@@ -458,6 +458,12 @@ void FActorTreeItem::GenerateContextMenu(UToolMenu* Menu, SSceneOutliner& Outlin
 	}
 }
 
+const FGuid& FActorTreeItem::GetGuid() const
+{
+	static const FGuid InvalidGuid;
+	return Actor.IsValid() ? Actor->GetActorGuid() : InvalidGuid;
+}
+
 void FActorTreeItem::UpdateDisplayString()
 {
 	UpdateDisplayStringInternal();
diff --git a/Engine/Source/Editor/SceneOutliner/Public/ActorBaseTreeItem.h b/Engine/Source/Editor/SceneOutliner/Public/ActorBaseTreeItem.h
new file mode 100644
--- /dev/null
+++ b/Engine/Source/Editor/SceneOutliner/Public/ActorBaseTreeItem.h
@@ -0,0 +1,18 @@
+// Copyright Epic Games, Inc. All Rights Reserved.
+
+#pragma once
+
+#include "CoreMinimal.h"
+#include "ISceneOutlinerTreeItem.h"
+
+/** A tree item that represents an actor, loaded or unloaded */
+struct SCENEOUTLINER_API IActorBaseTreeItem : ISceneOutlinerTreeItem
+{
+public:
+	IActorBaseTreeItem(FSceneOutlinerTreeItemType InType) : ISceneOutlinerTreeItem(InType) {}
+
+	/** Static type identifier for the base class tree item */
+	static const FSceneOutlinerTreeItemType Type;
+
+	virtual const FGuid& GetGuid() const =0;
+};
diff --git a/Engine/Source/Editor/SceneOutliner/Public/ActorDescTreeItem.h b/Engine/Source/Editor/SceneOutliner/Public/ActorDescTreeItem.h
--- a/Engine/Source/Editor/SceneOutliner/Public/ActorDescTreeItem.h
+++ b/Engine/Source/Editor/SceneOutliner/Public/ActorDescTreeItem.h
@@ -4,7 +4,7 @@
 
 #include "CoreMinimal.h"
 #include "Widgets/SWidget.h"
-#include "ISceneOutlinerTreeItem.h"
+#include "ActorBaseTreeItem.h"
 #include "UObject/ObjectKey.h"
 #include "WorldPartition/WorldPartition.h"
 
@@ -12,7 +12,7 @@ class FWorldPartitionActorDesc;
 class UToolMenu;
 
 /** A tree item that represents an actor in the world */
-struct SCENEOUTLINER_API FActorDescTreeItem : ISceneOutlinerTreeItem
+struct SCENEOUTLINER_API FActorDescTreeItem : IActorBaseTreeItem
 {
 public:
 	DECLARE_DELEGATE_RetVal_OneParam(bool, FFilterPredicate, const FWorldPartitionActorDesc*);
@@ -55,7 +55,9 @@ public:
 	virtual bool GetPinnedState() const override;
 	/* End ISceneOutlinerTreeItem Implementation */
 	
-	const FGuid& GetGuid() const { return ActorGuid; }
+	/* Begin IActorBaseTreeItem Implementation */
+	virtual const FGuid& GetGuid() const override { return ActorGuid; }
+	/* End IActorBaseTreeItem Implementation */
 
 protected:
 	FString DisplayString;
diff --git a/Engine/Source/Editor/SceneOutliner/Public/ActorTreeItem.h b/Engine/Source/Editor/SceneOutliner/Public/ActorTreeItem.h
--- a/Engine/Source/Editor/SceneOutliner/Public/ActorTreeItem.h
+++ b/Engine/Source/Editor/SceneOutliner/Public/ActorTreeItem.h
@@ -4,11 +4,11 @@
 
 #include "CoreMinimal.h"
 #include "Widgets/SWidget.h"
-#include "ISceneOutlinerTreeItem.h"
+#include "ActorBaseTreeItem.h"
 #include "UObject/ObjectKey.h"
 
 /** A tree item that represents an actor in the world */
-struct SCENEOUTLINER_API FActorTreeItem : ISceneOutlinerTreeItem
+struct SCENEOUTLINER_API FActorTreeItem : IActorBaseTreeItem
 {
 public:
 	DECLARE_DELEGATE_RetVal_OneParam(bool, FFilterPredicate, const AActor*);
@@ -54,6 +54,10 @@ public:
 	virtual void GenerateContextMenu(UToolMenu* Menu, SSceneOutliner& Outliner) override;
 	/* End ISceneOutlinerTreeItem Implementation */
 
+	/* Begin IActorBaseTreeItem Implementation */
+	virtual const FGuid& GetGuid() const override;
+	/* End IActorBaseTreeItem Implementation */
+
 	/** true if this item exists in both the current world and PIE. */
 	bool bExistsInCurrentWorldAndPIE;
 
