ActorBrowsingModeSettings: Change bHideTemporaryActors to be false by default

diff --git a/Engine/Source/Editor/SceneOutliner/Private/ActorBrowsingModeSettings.h b/Engine/Source/Editor/SceneOutliner/Private/ActorBrowsingModeSettings.h
--- a/Engine/Source/Editor/SceneOutliner/Private/ActorBrowsingModeSettings.h
+++ b/Engine/Source/Editor/SceneOutliner/Private/ActorBrowsingModeSettings.h
@@ -18,7 +18,7 @@ public:
 
 	/** True when the Scene Outliner is hiding temporary/run-time Actors */
 	UPROPERTY()
-	bool bHideTemporaryActors = true;
+	bool bHideTemporaryActors = false;
 
 	/** True when the Scene Outliner is showing only Actors that exist in the current level */
 	UPROPERTY()
