ActorFolder Improvements :
- Optimized FFolder resolving to UActorFolder : now uses and maintains an acceleration table
- Favor creation of FFolder using ActorFolderGuid when available
- FFolder creation now always passes a Root Object to facilitate Root Object ptr resolving (even when it's the main world)
- Fixed Duplicate Hierarchy when using Actor Folders and target level is different
- Fixed copy/paste actor from Persistent to LevelInstance not loosing actor folder
- Fixed mark for delete of an Actor Folder that could generate duplicates
- Modified fix of duplicate Actor Folders in a level : instead of renaming duplicates, mark for delete all duplicates except one and redirect children to the one we keep

diff --git a/Engine/Source/Editor/DataLayerEditor/Private/DataLayer/DataLayerMode.cpp b/Engine/Source/Editor/DataLayerEditor/Private/DataLayer/DataLayerMode.cpp
--- a/Engine/Source/Editor/DataLayerEditor/Private/DataLayer/DataLayerMode.cpp
+++ b/Engine/Source/Editor/DataLayerEditor/Private/DataLayer/DataLayerMode.cpp
@@ -294,6 +294,11 @@ FText FDataLayerMode::GetStatusText() const
 	}
 }
 
+FFolder::FRootObject FDataLayerMode::GetRootObject() const
+{
+	return FFolder::GetWorldRootFolder(RepresentingWorld.Get()).GetRootObject();
+}
+
 SDataLayerBrowser* FDataLayerMode::GetDataLayerBrowser() const
 {
 	return DataLayerBrowser;
diff --git a/Engine/Source/Editor/DataLayerEditor/Private/DataLayer/DataLayerMode.h b/Engine/Source/Editor/DataLayerEditor/Private/DataLayer/DataLayerMode.h
--- a/Engine/Source/Editor/DataLayerEditor/Private/DataLayer/DataLayerMode.h
+++ b/Engine/Source/Editor/DataLayerEditor/Private/DataLayer/DataLayerMode.h
@@ -56,6 +56,7 @@ public:
 	virtual bool ShowFilterOptions() const override { return true; }
 	virtual FText GetStatusText() const override;
 	virtual FSlateColor GetStatusTextColor() const override { return FSlateColor::UseForeground(); }
+	virtual FFolder::FRootObject GetRootObject() const override;
 
 	virtual void SynchronizeSelection() override;
 	virtual void OnItemDoubleClick(FSceneOutlinerTreeItemPtr Item) override;
diff --git a/Engine/Source/Editor/SceneOutliner/Private/ActorBrowsingMode.cpp b/Engine/Source/Editor/SceneOutliner/Private/ActorBrowsingMode.cpp
--- a/Engine/Source/Editor/SceneOutliner/Private/ActorBrowsingMode.cpp
+++ b/Engine/Source/Editor/SceneOutliner/Private/ActorBrowsingMode.cpp
@@ -42,7 +42,6 @@ DEFINE_LOG_CATEGORY_STATIC(LogActorBrowser, Log, All);
 
 using FActorFilter = TSceneOutlinerPredicateFilter<FActorTreeItem>;
 using FActorDescFilter = TSceneOutlinerPredicateFilter<FActorDescTreeItem>;
-static const FFolder GWorldRoot(FFolder::GetEmptyPath(), nullptr);
 
 FActorBrowsingMode::FActorBrowsingMode(SSceneOutliner* InSceneOutliner, TWeakObjectPtr<UWorld> InSpecifiedWorldToDisplay)
 	: FActorModeInteractive(FActorModeParams(InSceneOutliner, InSpecifiedWorldToDisplay,  /* bHideComponents */ true, /* bHideLevelInstanceHierarchy */ false, /* bHideUnloadedActors */ false))
@@ -560,6 +559,19 @@ void FActorBrowsingMode::RegisterContextMenu()
 											}
 										}
 									}
+								}),
+								FCanExecuteAction::CreateLambda([SceneOutliner]
+								{
+									const FSceneOutlinerItemSelection& Selection = SceneOutliner->GetSelection();
+									if (Selection.SelectedItems.Num() == 1)
+									{
+										FSceneOutlinerTreeItemPtr Item = Selection.SelectedItems[0].Pin();
+										if (FActorFolderTreeItem* FolderItem = Item->CastTo<FActorFolderTreeItem>())
+										{
+											return FolderItem->World.IsValid() && (FolderItem->World->GetCurrentLevel() == FolderItem->GetFolder().GetRootObjectAssociatedLevel());
+										}
+									}
+									return false;
 								})
 							)
 						);
@@ -576,10 +588,13 @@ void FActorBrowsingMode::RegisterContextMenu()
 								if (RepresentingWorld.IsValid())
 								{
 									const FScopedTransaction Transaction(NSLOCTEXT("UnrealEd", "ClearCurrentActorFolder", "Clear Current Actor Folder"));
-									FActorFolders::Get().SetActorEditorContextFolder(*RepresentingWorld.Get(), FFolder());
+									FActorFolders::Get().SetActorEditorContextFolder(*RepresentingWorld.Get(), FFolder::GetWorldRootFolder(RepresentingWorld.Get()));
 								}
 							}),
-							FCanExecuteAction::CreateLambda([this] { return RepresentingWorld.IsValid() && !FActorFolders::Get().GetActorEditorContextFolder(*RepresentingWorld.Get()).IsNone(); })
+							FCanExecuteAction::CreateLambda([this]
+							{ 
+								return RepresentingWorld.IsValid() && !FActorFolders::Get().GetActorEditorContextFolder(*RepresentingWorld.Get()).IsNone(); 
+							})
 						)
 					);
 				}
@@ -1002,34 +1017,9 @@ bool FActorBrowsingMode::CanPasteFoldersOnlyFromClipboard() const
 	return PasteString.StartsWith("BEGIN FOLDERLIST");
 }
 
-bool FActorBrowsingMode::GetFolderNamesFromFolders(const TArray<FFolder>& InFolders, TArray<FName>& OutFolders, FFolder::FRootObject& OutCommonRootObject) const
-{
-	OutFolders.Reset();
-	OutCommonRootObject = FFolder::GetDefaultRootObject();
-
-	TOptional<FFolder::FRootObject> CommonRootObject;
-	for (const FFolder& Folder : InFolders)
-	{
-		OutFolders.Add(Folder.GetPath());
-		if (!CommonRootObject.IsSet())
-		{
-			CommonRootObject = Folder.GetRootObject();
-		}
-		else if (CommonRootObject.GetValue() != Folder.GetRootObject())
-		{
-			OutFolders.Reset();
-			CommonRootObject.Reset();
-			break;
-		}
-	};
-
-	OutCommonRootObject = CommonRootObject.Get(FFolder::GetDefaultRootObject());
-	return CommonRootObject.IsSet();
-}
-
 bool FActorBrowsingMode::GetFolderNamesFromPayload(const FSceneOutlinerDragDropPayload& InPayload, TArray<FName>& OutFolders, FFolder::FRootObject& OutCommonRootObject) const
 {
-	return GetFolderNamesFromFolders(InPayload.GetData<FFolder>(SceneOutliner::FFolderPathSelector()), OutFolders, OutCommonRootObject);
+	return FFolder::GetFolderPathsAndCommonRootObject(InPayload.GetData<FFolder>(SceneOutliner::FFolderPathSelector()), OutFolders, OutCommonRootObject);
 }
 
 TSharedPtr<FDragDropOperation> FActorBrowsingMode::CreateDragDropOperation(const FPointerEvent& MouseEvent, const TArray<FSceneOutlinerTreeItemPtr>& InTreeItems) const
@@ -1075,7 +1065,7 @@ bool FActorBrowsingMode::ParseDragDrop(FSceneOutlinerDragDropPayload& OutPayload
 		{
 			for (const auto& Folder : FolderOp->Folders)
 			{
-				OutPayload.DraggedItems.Add(SceneOutliner->GetTreeItem(FFolder(Folder, FolderOp->RootObject)));
+				OutPayload.DraggedItems.Add(SceneOutliner->GetTreeItem(FFolder(FolderOp->RootObject, Folder)));
 			}
 		}
 		if (const auto& ActorOp = OutlinerOp.GetSubOp<FActorDragDropOp>())
@@ -1099,6 +1089,11 @@ bool FActorBrowsingMode::ParseDragDrop(FSceneOutlinerDragDropPayload& OutPayload
 	return false;
 }
 
+FFolder FActorBrowsingMode::GetWorldDefaultRootFolder() const
+{
+	return FFolder::GetWorldRootFolder(RepresentingWorld.Get());
+}
+
 FSceneOutlinerDragValidationInfo FActorBrowsingMode::ValidateDrop(const ISceneOutlinerTreeItem& DropTarget, const FSceneOutlinerDragDropPayload& Payload) const
 {
 	if (Payload.Has<FFolderTreeItem>())
@@ -1238,7 +1233,7 @@ FSceneOutlinerDragValidationInfo FActorBrowsingMode::ValidateDrop(const ISceneOu
 		const FWorldTreeItem* WorldItem = DropTarget.CastTo<FWorldTreeItem>();
 		const FLevelTreeItem* LevelItem = DropTarget.CastTo<FLevelTreeItem>();
 		// WorldTreeItem and LevelTreeItem are treated as root folders (path = none), with the difference that LevelTreeItem has a RootObject.
-		const FFolder DestinationPath = FolderItem ? FolderItem->GetFolder() : (LevelItem ? FFolder(FFolder::GetEmptyPath(), FFolder::GetOptionalFolderRootObject(LevelItem->Level.Get()).Get(FFolder::GetDefaultRootObject())) : GWorldRoot);
+		const FFolder DestinationPath = FolderItem ? FolderItem->GetFolder() : (LevelItem ? FFolder(FFolder::GetOptionalFolderRootObject(LevelItem->Level.Get()).Get(FFolder::GetInvalidRootObject())) : GetWorldDefaultRootFolder());
 		const FFolder::FRootObject& DestinationRootObject = DestinationPath.GetRootObject();
 		ILevelInstanceInterface* LevelInstanceTarget = Cast<ILevelInstanceInterface>(DestinationPath.GetRootObjectPtr());
 		if (LevelInstanceTarget && !LevelInstanceTarget->IsEditing())
@@ -1258,7 +1253,7 @@ FSceneOutlinerDragValidationInfo FActorBrowsingMode::ValidateDrop(const ISceneOu
 					const FName Leaf = FEditorFolderUtils::GetLeafName(DraggedFolder);
 					const FName Parent = FEditorFolderUtils::GetParentPath(DraggedFolder);
 
-					if ((CommonFolderRootObject != DestinationRootObject) && FFolder::HasRootObject(CommonFolderRootObject) && FFolder::HasRootObject(DestinationRootObject))
+					if ((CommonFolderRootObject != DestinationRootObject) && FFolder::IsRootObjectValid(CommonFolderRootObject) && FFolder::IsRootObjectValid(DestinationRootObject))
 					{
 						FFormatNamedArguments Args;
 						Args.Add(TEXT("SourceName"), FText::FromName(Leaf));
@@ -1290,7 +1285,7 @@ FSceneOutlinerDragValidationInfo FActorBrowsingMode::ValidateDrop(const ISceneOu
 					const FString DstFolderPath = DestinationPath.IsNone() ? FString() : DestinationPath.ToString();
 					const FString NewPath = DstFolderPath / LeafName;
 
-					if (FActorFolders::Get().ContainsFolder(*RepresentingWorld, FFolder(FName(*NewPath), DestinationRootObject)))
+					if (FActorFolders::Get().ContainsFolder(*RepresentingWorld, FFolder(DestinationRootObject, FName(*NewPath))))
 					{
 						// The folder already exists
 						FFormatNamedArguments Args;
@@ -1345,7 +1340,7 @@ FSceneOutlinerDragValidationInfo FActorBrowsingMode::ValidateDrop(const ISceneOu
 				{
 					return FSceneOutlinerDragValidationInfo(ESceneOutlinerDropCompatibility::IncompatibleGeneric, FText::Format(LOCTEXT("Error_AttachChildActor", "Cannot move {0} as it is a child actor."), FText::FromString(Actor->GetActorLabel())));
 				}
-				else if ((Actor->GetFolderRootObject() != DestinationRootObject) && FFolder::HasRootObject(Actor->GetFolderRootObject()) && FFolder::HasRootObject(DestinationRootObject))
+				else if ((Actor->GetFolderRootObject() != DestinationRootObject) && FFolder::IsRootObjectValid(Actor->GetFolderRootObject()) && FFolder::IsRootObjectValid(DestinationRootObject))
 				{
 					FFormatNamedArguments Args;
 					Args.Add(TEXT("SourceName"), FText::FromString(Actor->GetActorLabel()));
@@ -1438,7 +1433,7 @@ void FActorBrowsingMode::OnDrop(ISceneOutlinerTreeItem& DropTarget, const FScene
 				check(TargetLevelInstance->IsEditing());
 				const FScopedTransaction Transaction(LOCTEXT("UndoAction_MoveActorsToLevelInstance", "Move actors to LevelInstance"));
 
-				const FFolder DestinationPath(FFolder::GetEmptyPath(), FFolder::FRootObject(DropActor));
+				const FFolder DestinationPath = FFolder(FFolder::FRootObject(DropActor));
 				auto MoveToDestination = [&DestinationPath](FFolderTreeItem& Item)
 				{
 					Item.MoveTo(DestinationPath);
@@ -1514,7 +1509,7 @@ void FActorBrowsingMode::OnDrop(ISceneOutlinerTreeItem& DropTarget, const FScene
 		const FWorldTreeItem* WorldItem = DropTarget.CastTo<FWorldTreeItem>();
 		const FLevelTreeItem* LevelItem = DropTarget.CastTo<FLevelTreeItem>();
 		// WorldTreeItem and LevelTreeItem are treated as root folders (path = none), with the difference that LevelTreeItem has a RootObject.
-		const FFolder DestinationPath = FolderItem ? FolderItem->GetFolder() : (LevelItem ? FFolder(FFolder::GetEmptyPath(), FFolder::GetOptionalFolderRootObject(LevelItem->Level.Get()).Get(FFolder::GetDefaultRootObject())) : GWorldRoot);
+		const FFolder DestinationPath = FolderItem ? FolderItem->GetFolder() : (LevelItem ? FFolder(FFolder::GetOptionalFolderRootObject(LevelItem->Level.Get()).Get(FFolder::GetInvalidRootObject())) : GetWorldDefaultRootFolder());
 
 		const FScopedTransaction Transaction(LOCTEXT("MoveOutlinerItems", "Move World Outliner Items"));
 
@@ -1539,10 +1534,10 @@ void FActorBrowsingMode::OnDrop(ISceneOutlinerTreeItem& DropTarget, const FScene
 					// First mark this object as a parent, then set its children's path
 					ParentActors.Add(Actor);
 
-					const FFolder Folder = Actor->GetFolder();
+					const FFolder SrcFolder = Actor->GetFolder();
 
 					// If the folder root object changes, 1st pass will put actors at root. 2nd pass will set the destination path.
-					FName NewPath = (Folder.GetRootObject() == DestinationPath.GetRootObject()) ? DestinationPath.GetPath() : NAME_None;
+					FName NewPath = (SrcFolder.GetRootObject() == DestinationPath.GetRootObject()) ? DestinationPath.GetPath() : NAME_None;
 					
 					Actor->SetFolderPath(NewPath);
 					FActorEditorUtils::TraverseActorTree_ParentFirst(Actor, [&](AActor* InActor) {
@@ -1551,7 +1546,7 @@ void FActorBrowsingMode::OnDrop(ISceneOutlinerTreeItem& DropTarget, const FScene
 						return true;
 						}, false);
 					
-					if ((Actor->GetFolderRootObject() != DestinationPath.GetRootObject()) && !Folder.HasRootObject() && DestinationPath.HasRootObject())
+					if ((Actor->GetFolderRootObject() != DestinationPath.GetRootObject()) && SrcFolder.IsRootObjectPersistentLevel() && (DestinationPath.IsRootObjectValid() && !DestinationPath.IsRootObjectPersistentLevel()))
 					{
 						MovingActorsToValidRootObject.Add(Actor);
 					}
@@ -1593,7 +1588,7 @@ void FActorBrowsingMode::OnDrop(ISceneOutlinerTreeItem& DropTarget, const FScene
 				}
 			};
 
-			if (!DestinationPath.HasRootObject())
+			if (DestinationPath.IsRootObjectPersistentLevel())
 			{
 				const ULevelInstanceSubsystem* LevelInstanceSubsystem = RepresentingWorld->GetSubsystem<ULevelInstanceSubsystem>();
 				check(LevelInstanceSubsystem);
@@ -1610,9 +1605,13 @@ void FActorBrowsingMode::OnDrop(ISceneOutlinerTreeItem& DropTarget, const FScene
 						check(ParentLevelInstance->IsEditing());
 						LevelInstanceActorsToMove.Add(Actor);
 					}
-					else if (Actor->GetFolder().HasRootObject())
+					else
 					{
-						ActorsToMoveToPersistentLevel.Add(Actor);
+						const FFolder ActorSrcFolder = Actor->GetFolder();
+						if (ActorSrcFolder.IsRootObjectValid() && !ActorSrcFolder.IsRootObjectPersistentLevel())
+						{
+							ActorsToMoveToPersistentLevel.Add(Actor);
+						}
 					}
 				});
 
@@ -1668,11 +1667,15 @@ FFolder FActorBrowsingMode::CreateNewFolder()
 	return NewFolderName;
 }
 
-FFolder FActorBrowsingMode::CreateFolder(const FFolder& ParentPath, const FName& LeafName)
+FFolder FActorBrowsingMode::GetFolder(const FFolder& ParentPath, const FName& LeafName)
+{
+	// Return a unique folder under the provided parent path & root object and using the provided leaf name
+	return FActorFolders::Get().GetFolderName(*RepresentingWorld, ParentPath, LeafName);
+}
+
+bool FActorBrowsingMode::CreateFolder(const FFolder& NewPath)
 {
-	const FFolder NewPath = FActorFolders::Get().GetFolderName(*RepresentingWorld, ParentPath, LeafName);
-	FActorFolders::Get().CreateFolder(*RepresentingWorld, NewPath);
-	return NewPath;
+	return FActorFolders::Get().CreateFolder(*RepresentingWorld, NewPath);
 }
 
 bool FActorBrowsingMode::ReparentItemToFolder(const FFolder& FolderPath, const FSceneOutlinerTreeItemPtr& Item)
@@ -1908,7 +1911,7 @@ FCreateSceneOutlinerMode FActorBrowsingMode::CreateFolderPickerMode(const FFolde
 	{
 		if (FWorldTreeItem* WorldItem = NewParent->CastTo<FWorldTreeItem>())
 		{
-			SceneOutliner->MoveSelectionTo(GWorldRoot);
+			SceneOutliner->MoveSelectionTo(GetWorldDefaultRootFolder());
 		}
 		else if (FFolderTreeItem* FolderItem = NewParent->CastTo<FFolderTreeItem>())
 		{
@@ -1916,16 +1919,16 @@ FCreateSceneOutlinerMode FActorBrowsingMode::CreateFolderPickerMode(const FFolde
 		}
 		else if (FActorTreeItem* ActorItem = NewParent->CastTo<FActorTreeItem>())
 		{
-			if (FFolder::HasRootObject(InRootObject))
+			if (FFolder::IsRootObjectValid(InRootObject))
 			{
-				SceneOutliner->MoveSelectionTo(FFolder(FFolder::GetEmptyPath(), InRootObject));
+				SceneOutliner->MoveSelectionTo(FFolder(InRootObject));
 			}
 		}
 		else if (FLevelTreeItem* LevelItem = NewParent->CastTo<FLevelTreeItem>())
 		{
-			if (FFolder::HasRootObject(InRootObject))
+			if (FFolder::IsRootObjectValid(InRootObject))
 			{
-				SceneOutliner->MoveSelectionTo(FFolder(FFolder::GetEmptyPath(), InRootObject));
+				SceneOutliner->MoveSelectionTo(FFolder(InRootObject));
 			}
 		}
 	};
@@ -1985,7 +1988,7 @@ void FActorBrowsingMode::OnDuplicateActorsBegin()
 	// Only a callback in actor browsing mode
 	FFolder::FRootObject CommonRootObject;
 	TArray<FName> SelectedFolderPaths;
-	GetFolderNamesFromFolders(SceneOutliner->GetSelection().GetData<FFolder>(SceneOutliner::FFolderPathSelector()), SelectedFolderPaths, CommonRootObject);
+	FFolder::GetFolderPathsAndCommonRootObject(SceneOutliner->GetSelection().GetData<FFolder>(SceneOutliner::FFolderPathSelector()), SelectedFolderPaths, CommonRootObject);
 	SceneOutliner->PasteFoldersBegin(SelectedFolderPaths);
 }
 
diff --git a/Engine/Source/Editor/SceneOutliner/Private/ActorFolderHierarchy.cpp b/Engine/Source/Editor/SceneOutliner/Private/ActorFolderHierarchy.cpp
--- a/Engine/Source/Editor/SceneOutliner/Private/ActorFolderHierarchy.cpp
+++ b/Engine/Source/Editor/SceneOutliner/Private/ActorFolderHierarchy.cpp
@@ -90,7 +90,7 @@ void FActorFolderHierarchy::CreateWorldChildren(UWorld* World, TArray<FSceneOutl
 		return true;
 	});
 
-	if (FFolder::HasRootObject(RootObject))
+	if (FFolder::IsRootObjectValid(RootObject))
 	{
 		UObject* RootObjectPtr = FFolder::GetRootObjectPtr(RootObject);
 		if (ILevelInstanceInterface* RootLevelInstance = Cast<ILevelInstanceInterface>(RootObjectPtr))
@@ -102,9 +102,13 @@ void FActorFolderHierarchy::CreateWorldChildren(UWorld* World, TArray<FSceneOutl
 		}
 		else if (ULevel* RootLevel = Cast<ULevel>(RootObjectPtr))
 		{
-			if (FSceneOutlinerTreeItemPtr ActorItem = Mode->CreateItemFor<FLevelTreeItem>(RootLevel, true))
+			// Don't create a LevelTreeItems for the PersistentLevel
+			if (!RootLevel->IsPersistentLevel())
 			{
-				OutItems.Add(ActorItem);
+				if (FSceneOutlinerTreeItemPtr ActorItem = Mode->CreateItemFor<FLevelTreeItem>(RootLevel, true))
+				{
+					OutItems.Add(ActorItem);
+				}
 			}
 		}
 	}
@@ -114,7 +118,7 @@ void FActorFolderHierarchy::CreateItems(TArray<FSceneOutlinerTreeItemPtr>& OutIt
 {
 	check(RepresentingWorld.IsValid());
 
-	if (!FFolder::HasRootObject(RootObject))
+	if (FFolder::IsRootObjectPersistentLevel(RootObject))
 	{
 		if (FSceneOutlinerTreeItemPtr WorldItem = Mode->CreateItemFor<FWorldTreeItem>(RepresentingWorld.Get()))
 		{
diff --git a/Engine/Source/Editor/SceneOutliner/Private/ActorFolderTreeItem.cpp b/Engine/Source/Editor/SceneOutliner/Private/ActorFolderTreeItem.cpp
--- a/Engine/Source/Editor/SceneOutliner/Private/ActorFolderTreeItem.cpp
+++ b/Engine/Source/Editor/SceneOutliner/Private/ActorFolderTreeItem.cpp
@@ -198,9 +198,9 @@ private:
 		{
 			NewPath = FName(*(NewPath.ToString() / LabelString));
 		}
-		Folder.SetPath(NewPath);
+		FFolder NewFolder(Folder.GetRootObject(), NewPath);
 
-		if (FActorFolders::Get().ContainsFolder(*TreeItem->World.Get(), Folder))
+		if (FActorFolders::Get().ContainsFolder(*TreeItem->World.Get(), NewFolder))
 		{
 			OutErrorMessage = LOCTEXT("RenameFailed_AlreadyExists", "A folder with this name already exists at this level");
 			return false;
@@ -225,7 +225,7 @@ private:
 			{
 				NewPath = FName(*(NewPath.ToString() / InLabel.ToString()));
 			}
-			FFolder TreeItemNewFolder(NewPath, Folder.GetRootObject());
+			FFolder TreeItemNewFolder(Folder.GetRootObject(), NewPath);
 
 			FActorFolders::Get().RenameFolderInWorld(*TreeItem->World.Get(), Folder, TreeItemNewFolder);
 
@@ -254,7 +254,9 @@ FActorFolderTreeItem::FActorFolderTreeItem(const FFolder& InFolder, const TWeakO
 	: FFolderTreeItem(InFolder, Type)
 	, World(InWorld)
 {
-	SetPath(InFolder.GetPath());
+	FFolderTreeItem::SetPath(InFolder.GetPath());
+	// Initialize ActorFolder
+	ActorFolder = InFolder.GetActorFolder();
 }
 
 void FActorFolderTreeItem::OnExpansionChanged()
@@ -321,10 +323,8 @@ void FActorFolderTreeItem::SetPath(const FName& InNewPath)
 {
 	FFolderTreeItem::SetPath(InNewPath);
 
-	if (World.IsValid())
-	{
-		ActorFolder = FWorldPersistentFolders::GetActorFolder(GetFolder(), World.Get());
-	}
+	// Validate that if the item's ActorFolder object is set, that changing the path keeps the same ActorFolder but with a different path
+	check(!ActorFolder.IsValid() || (ActorFolder->GetFolder().GetPath() == Path));
 }
 
 void FActorFolderTreeItem::CreateSubFolder(TWeakPtr<SSceneOutliner> WeakOutliner)
@@ -350,7 +350,18 @@ TSharedRef<SWidget> FActorFolderTreeItem::GenerateLabelWidget(ISceneOutliner& Ou
 
 bool FActorFolderTreeItem::ShouldShowPinnedState() const
 {
-	return GetRootObject() == FFolder::GetDefaultRootObject() && World.IsValid() && !World->IsGameWorld() && World->IsPartitionedWorld();
+	if (World.IsValid() && !World->IsGameWorld() && World->IsPartitionedWorld())
+	{
+		return FFolder::IsRootObjectPersistentLevel(GetRootObject());
+	}
+	return false;
+}
+
+FFolder FActorFolderTreeItem::GetFolder() const
+{
+	// Use Folder resolved by ActorFolder when valid (a null owning world would will fail resolving, see UActorFolder::GetFolder())
+	FFolder Folder = ActorFolder.IsValid() ? ActorFolder->GetFolder() : FFolder::GetInvalidFolder();
+	return Folder.IsValid() ? Folder : FFolderTreeItem::GetFolder();
 }
 
 bool FActorFolderTreeItem::CanInteract() const
diff --git a/Engine/Source/Editor/SceneOutliner/Private/ActorHierarchy.cpp b/Engine/Source/Editor/SceneOutliner/Private/ActorHierarchy.cpp
--- a/Engine/Source/Editor/SceneOutliner/Private/ActorHierarchy.cpp
+++ b/Engine/Source/Editor/SceneOutliner/Private/ActorHierarchy.cpp
@@ -127,13 +127,20 @@ FActorHierarchy::~FActorHierarchy()
 
 static FFolder GetFolderPathFromActorDesc(UWorld* InWorld, const FWorldPartitionActorDesc* InActorDesc)
 {
-	UWorld* World = InActorDesc->GetContainer() ? InActorDesc->GetContainer()->GetTypedOuter<UWorld>() : InWorld;
-	if (World && World->PersistentLevel->IsUsingActorFolders())
+	if (UWorld* OuterWorld = InActorDesc->GetContainer() ? InActorDesc->GetContainer()->GetTypedOuter<UWorld>() : InWorld)
 	{
-		UActorFolder* ActorFolder = World->PersistentLevel->GetActorFolder(InActorDesc->GetFolderGuid());
-		return ActorFolder ? ActorFolder->GetFolder() : FFolder();
+		ULevel* OuterLevel = OuterWorld->PersistentLevel;
+		if (OuterLevel->IsUsingActorFolders())
+		{
+			if (UActorFolder* ActorFolder = OuterLevel->GetActorFolder(InActorDesc->GetFolderGuid()))
+			{
+				return ActorFolder->GetFolder();
+			}
+			return FFolder::GetWorldRootFolder(OuterWorld).GetRootObject();
+		}
+		return FFolder(FFolder::GetWorldRootFolder(OuterWorld).GetRootObject(), InActorDesc->GetFolderPath());
 	}
-	return FFolder(InActorDesc->GetFolderPath());
+	return FFolder::GetInvalidFolder();
 }
 
 FSceneOutlinerTreeItemPtr FActorHierarchy::FindOrCreateParentItem(const ISceneOutlinerTreeItem& Item, const TMap<FSceneOutlinerTreeItemID, FSceneOutlinerTreeItemPtr>& Items, bool bCreate)
@@ -195,7 +202,9 @@ FSceneOutlinerTreeItemPtr FActorHierarchy::FindOrCreateParentItem(const ISceneOu
 			}
 
 			// Parent Level Using Actor Folders
-			if (ULevel* OwningLevel = Cast<ULevel>(Folder.GetRootObjectPtr()))
+			ULevel* OwningLevel = Cast<ULevel>(Folder.GetRootObjectPtr());
+			// For the persistent level, fallback on the world
+			if (OwningLevel && !OwningLevel->IsPersistentLevel())
 			{
 				if (const FSceneOutlinerTreeItemPtr* ParentItem = Items.Find(OwningLevel))
 				{
@@ -246,13 +255,17 @@ FSceneOutlinerTreeItemPtr FActorHierarchy::FindOrCreateParentItem(const ISceneOu
 		// Parent Level Using Actor Folders
 		else if (ULevel* OwningLevel = Cast<ULevel>(ParentPath.GetRootObjectPtr()))
 		{
-			if (const FSceneOutlinerTreeItemPtr* ParentItem = Items.Find(OwningLevel))
-			{
-				return *ParentItem;
-			}
-			else
+			// For the persistent level, fallback on the world
+			if (!OwningLevel->IsPersistentLevel())
 			{
-				return bCreate ? Mode->CreateItemFor<FLevelTreeItem>(OwningLevel, true) : nullptr;
+				if (const FSceneOutlinerTreeItemPtr* ParentItem = Items.Find(OwningLevel))
+				{
+					return *ParentItem;
+				}
+				else
+				{
+					return bCreate ? Mode->CreateItemFor<FLevelTreeItem>(OwningLevel, true) : nullptr;
+				}
 			}
 		}
 	}
@@ -759,7 +772,8 @@ void FActorHierarchy::OnLevelRemoved(ULevel* InLevel, UWorld* InWorld)
 			HierarchyChangedEvent.Broadcast(EventData);
 		}
 
-		if (InLevel->IsUsingActorFolders())
+		// If either this level or the owning world are using actor folders, remove level's actor folders
+		if (InLevel->IsUsingActorFolders() || InWorld->PersistentLevel->IsUsingActorFolders())
 		{
 			FSceneOutlinerHierarchyChangedData EventData;
 			EventData.Type = FSceneOutlinerHierarchyChangedData::Removed;
diff --git a/Engine/Source/Editor/SceneOutliner/Private/ActorMode.cpp b/Engine/Source/Editor/SceneOutliner/Private/ActorMode.cpp
--- a/Engine/Source/Editor/SceneOutliner/Private/ActorMode.cpp
+++ b/Engine/Source/Editor/SceneOutliner/Private/ActorMode.cpp
@@ -95,10 +95,6 @@ FActorMode::FActorMode(const FActorModeParams& Params)
 
 	auto FolderPassesFilter = [this](const FFolder& InFolder, bool bInCheckHideLevelInstanceFlag)
 	{
-		if (!InFolder.HasRootObject())
-		{
-			return true;
-		}
 		if (ILevelInstanceInterface* LevelInstance = Cast<ILevelInstanceInterface>(InFolder.GetRootObjectPtr()))
 		{
 			if (LevelInstance->IsEditing())
@@ -318,6 +314,20 @@ bool FActorMode::IsActorDisplayable(const AActor* Actor) const
 	return FActorMode::IsActorDisplayable(SceneOutliner, Actor);
 }
 
+FFolder::FRootObject FActorMode::GetRootObject() const
+{
+	return FFolder::GetWorldRootFolder(RepresentingWorld.Get()).GetRootObject();
+}
+
+FFolder::FRootObject FActorMode::GetPasteTargetRootObject() const
+{
+	if (UWorld* World = RepresentingWorld.Get())
+	{
+		return FFolder::GetOptionalFolderRootObject(World->GetCurrentLevel()).Get(FFolder::GetWorldRootFolder(World).GetRootObject());
+	}
+	return FFolder::GetInvalidRootObject();
+}
+
 bool FActorMode::IsActorDisplayable(const SSceneOutliner* SceneOutliner, const AActor* Actor)
 {
 	static const FName SequencerActorTag(TEXT("SequencerActor"));
diff --git a/Engine/Source/Editor/SceneOutliner/Private/SOutlinerTreeView.cpp b/Engine/Source/Editor/SceneOutliner/Private/SOutlinerTreeView.cpp
--- a/Engine/Source/Editor/SceneOutliner/Private/SOutlinerTreeView.cpp
+++ b/Engine/Source/Editor/SceneOutliner/Private/SOutlinerTreeView.cpp
@@ -4,6 +4,7 @@
 #include "Framework/Application/SlateApplication.h"
 #include "Styling/AppStyle.h"
 #include "ISceneOutlinerColumn.h"
+#include "ISceneOutlinerMode.h"
 #include "SceneOutlinerPublicTypes.h"
 #include "DragAndDrop/DecoratedDragDropOp.h"
 #include "SceneOutlinerDragDrop.h"
@@ -101,8 +102,12 @@ FReply HandleDrop(TSharedPtr<SSceneOutliner> SceneOutlinerPtr, const FDragDropEv
 	return FReply::Handled();
 }
 
-FReply HandleDropFromWeak(TWeakPtr<SSceneOutliner> SceneOutlinerWeak, const FDragDropEvent& DragDropEvent, ISceneOutlinerTreeItem& DropTarget, FSceneOutlinerDragValidationInfo& ValidationInfo, bool bApplyDrop = false)
+FReply HandleDropFromWeak(TWeakPtr<SSceneOutliner> SceneOutlinerWeak, const FDragDropEvent& DragDropEvent, FSceneOutlinerDragValidationInfo& ValidationInfo, bool bApplyDrop = false)
 {
+	const ISceneOutlinerMode* Mode = SceneOutlinerWeak.IsValid() ? SceneOutlinerWeak.Pin()->GetMode() : nullptr;
+	FFolder::FRootObject RootObject = Mode ? Mode->GetRootObject() : FFolder::GetInvalidRootObject();
+	FFolder RootFolder(RootObject);
+	FFolderTreeItem DropTarget(RootFolder);
 	return HandleDrop(SceneOutlinerWeak.Pin(), DragDropEvent, DropTarget, ValidationInfo, bApplyDrop);
 }
 
@@ -124,9 +129,7 @@ void SSceneOutlinerTreeView::FlashHighlightOnItem( FSceneOutlinerTreeItemPtr Fla
 FReply SSceneOutlinerTreeView::OnDragOver(const FGeometry& MyGeometry, const FDragDropEvent& DragDropEvent)
 {
 	FSceneOutlinerDragValidationInfo ValidationInfo = FSceneOutlinerDragValidationInfo::Invalid();
-	FFolderTreeItem DropTarget(FFolder(FFolder::GetEmptyPath(), nullptr));
-
-	auto Reply = HandleDropFromWeak(SceneOutlinerWeak, DragDropEvent, DropTarget, ValidationInfo);
+	auto Reply = HandleDropFromWeak(SceneOutlinerWeak, DragDropEvent, ValidationInfo);
 	if (Reply.IsEventHandled())
 	{
 		UpdateOperationDecorator(DragDropEvent, ValidationInfo);
@@ -146,9 +149,7 @@ void SSceneOutlinerTreeView::OnDragLeave(const FDragDropEvent& DragDropEvent)
 FReply SSceneOutlinerTreeView::OnDrop(const FGeometry& MyGeometry, const FDragDropEvent& DragDropEvent)
 {
 	FSceneOutlinerDragValidationInfo ValidationInfo = FSceneOutlinerDragValidationInfo::Invalid();
-	FFolderTreeItem DropTarget(FFolder(FFolder::GetEmptyPath(), nullptr));
-
-	return HandleDropFromWeak(SceneOutlinerWeak, DragDropEvent, DropTarget, ValidationInfo, true);
+	return HandleDropFromWeak(SceneOutlinerWeak, DragDropEvent, ValidationInfo, true);
 }
 
 FReply SSceneOutlinerTreeRow::OnDrop( const FGeometry& MyGeometry, const FDragDropEvent& DragDropEvent )
diff --git a/Engine/Source/Editor/SceneOutliner/Private/SSceneOutliner.cpp b/Engine/Source/Editor/SceneOutliner/Private/SSceneOutliner.cpp
--- a/Engine/Source/Editor/SceneOutliner/Private/SSceneOutliner.cpp
+++ b/Engine/Source/Editor/SceneOutliner/Private/SSceneOutliner.cpp
@@ -979,14 +979,6 @@ void SSceneOutliner::GetSelectedFolders(TArray<FFolderTreeItem*>& OutFolders) co
 	return FSceneOutlinerItemSelection(*OutlinerTreeView).Get<FFolderTreeItem>(OutFolders);
 }
 
-TArray<FName> SSceneOutliner::GetSelectedFolderNames() const
-{
-	TArray<FFolder> Folders = GetSelection().GetData<FFolder>(SceneOutliner::FFolderPathSelector());
-	TArray<FName> FolderPaths;
-	Algo::ForEach(Folders, [&FolderPaths](const FFolder& Folder) { FolderPaths.Add(Folder.GetPath()); });
-	return FolderPaths;
-}
-
 TSharedPtr<SWidget> SSceneOutliner::OnOpenContextMenu()
 {
 	return Mode->CreateContextMenu();
@@ -1148,17 +1140,27 @@ void SSceneOutliner::RemoveFromSelection(const TArray<FSceneOutlinerTreeItemPtr>
 
 void SSceneOutliner::AddFolderToSelection(const FName& FolderName)
 {
-	if (FSceneOutlinerTreeItemPtr* ItemPtr = TreeItemMap.Find(FFolder(FolderName)))
+	// Not used (but public) : For backward compatibility, we use Mode->GetRootObject()
+	FFolder::FRootObject RootObject = Mode->GetRootObject();
+	if (FFolder::IsRootObjectValid(RootObject))
 	{
-		OutlinerTreeView->SetItemSelection(*ItemPtr, true);
+		if (FSceneOutlinerTreeItemPtr* ItemPtr = TreeItemMap.Find(FFolder(RootObject, FolderName)))
+		{
+			OutlinerTreeView->SetItemSelection(*ItemPtr, true);
+		}
 	}
 }
 
 void SSceneOutliner::RemoveFolderFromSelection(const FName& FolderName)
 {
-	if (FSceneOutlinerTreeItemPtr* ItemPtr = TreeItemMap.Find(FFolder(FolderName)))
+	// Not used (but public) : For backward compatibility, we use Mode->GetRootObject()
+	FFolder::FRootObject RootObject = Mode->GetRootObject();
+	if (FFolder::IsRootObjectValid(RootObject))
 	{
-		OutlinerTreeView->SetItemSelection(*ItemPtr, false);
+		if (FSceneOutlinerTreeItemPtr* ItemPtr = TreeItemMap.Find(FFolder(RootObject, FolderName)))
+		{
+			OutlinerTreeView->SetItemSelection(*ItemPtr, false);
+		}
 	}
 }
 
@@ -1292,7 +1294,7 @@ void SSceneOutliner::AddMoveToFolderOutliner(UToolMenu* Menu) const
 
 		// Filter in/out root items according to whether it is valid to move to/from the root
 		FSceneOutlinerDragDropPayload DraggedObjects(OutlinerTreeView->GetSelectedItems());
-		const bool bMoveToRootValid = Mode->ValidateDrop(FFolderTreeItem(FFolder(FFolder::GetEmptyPath(), TargetRootObject)), DraggedObjects).IsValid();
+		const bool bMoveToRootValid = Mode->ValidateDrop(FFolderTreeItem(FFolder(TargetRootObject, FFolder::GetEmptyPath())), DraggedObjects).IsValid();
 		if (!bMoveToRootValid)
 		{
 			MiniSceneOutlinerInitOptions.Filters->Add(MakeShared<FFilterRoot>(*this));
@@ -1410,7 +1412,12 @@ void SSceneOutliner::CreateFolder()
 
 void SSceneOutliner::CopyFoldersBegin()
 {
-	CacheFoldersEdit = GetSelectedFolderNames();
+	CacheFoldersEdit.Reset();
+	CacheFoldersEditRootObject = FFolder::GetInvalidRootObject();
+
+	TArray<FFolder> SelectedFolders = GetSelection().GetData<FFolder>(SceneOutliner::FFolderPathSelector());
+	FFolder::GetFolderPathsAndCommonRootObject(SelectedFolders, CacheFoldersEdit, CacheFoldersEditRootObject);
+
 	FPlatformApplicationMisc::ClipboardPaste(CacheClipboardContents);
 }
 
@@ -1420,6 +1427,7 @@ void SSceneOutliner::CopyFoldersEnd()
 	{
 		CopyFoldersToClipboard(CacheFoldersEdit, CacheClipboardContents);
 		CacheFoldersEdit.Reset();
+		CacheFoldersEditRootObject = FFolder::GetInvalidRootObject();
 	}
 }
 
@@ -1464,13 +1472,13 @@ bool SSceneOutliner::GetCommonRootObjectFromSelection(FFolder::FRootObject& OutC
 			}
 			else if (CommonRootObject.GetValue() != TreeItem->GetRootObject())
 			{
-				OutCommonRootObject = FFolder::GetDefaultRootObject();
-				return false;
+				OutCommonRootObject = FFolder::GetInvalidRootObject();
+				break;
 			}
 		}
 	}
-	OutCommonRootObject = CommonRootObject.Get(FFolder::GetDefaultRootObject());
-	return true;
+	OutCommonRootObject = CommonRootObject.Get(FFolder::GetInvalidRootObject());
+	return FFolder::IsRootObjectValid(OutCommonRootObject);
 }
 
 void SSceneOutliner::PasteFoldersBegin(TArray<FName> InFolders)
@@ -1493,8 +1501,8 @@ void SSceneOutliner::PasteFoldersBegin(TArray<FName> InFolders)
 		}
 	};
 
-
 	CacheFoldersEdit.Reset();
+	CacheFoldersEditRootObject = Mode->GetPasteTargetRootObject();
 	CachePasteFolderExistingChildrenMap.Reset();
 	PendingFoldersSelect.Reset();
 
@@ -1503,16 +1511,30 @@ void SSceneOutliner::PasteFoldersBegin(TArray<FName> InFolders)
 	// Sort folder names so parents appear before children
 	CacheFoldersEdit.Sort(FNameLexicalLess());
 
-	// Find common root object from selection
-	FFolder::FRootObject TargetRootObject;
-	GetCommonRootObjectFromSelection(TargetRootObject);
-
 	// Cache existing children
 	for (FName Folder : CacheFoldersEdit)
 	{
-		if (FSceneOutlinerTreeItemPtr* TreeItem = TreeItemMap.Find(FFolder(Folder, TargetRootObject)))
+		if (FSceneOutlinerTreeItemPtr* TreeItem = TreeItemMap.Find(FFolder(CacheFoldersEditRootObject, Folder)))
+		{
+			CacheExistingChildrenAction(*TreeItem, CacheFoldersEditRootObject);
+		}
+	}
+
+	// Prepare CacheFolderMap which maps old to new/duplicate folder names
+	CacheFolderMap.Reset();
+	for (FName Folder : CacheFoldersEdit)
+	{
+		FName ParentPath = FEditorFolderUtils::GetParentPath(Folder);
+		FName LeafName = FEditorFolderUtils::GetLeafName(Folder);
+		if (LeafName != TEXT(""))
 		{
-			CacheExistingChildrenAction(*TreeItem, TargetRootObject);
+			if (FName* NewParentPath = CacheFolderMap.Find(ParentPath))
+			{
+				ParentPath = *NewParentPath;
+			}
+
+			FFolder NewFolderPath = Mode->GetFolder(FFolder(CacheFoldersEditRootObject, ParentPath), LeafName);
+			CacheFolderMap.Add(Folder, NewFolderPath.GetPath());
 		}
 	}
 }
@@ -1521,25 +1543,17 @@ void SSceneOutliner::PasteFoldersEnd()
 {
 	const FScopedTransaction Transaction(NSLOCTEXT("UnrealEd", "PasteItems", "Paste Items"));
 
-	// Find common root object from selection
-	FFolder::FRootObject TargetRootObject;
-	GetCommonRootObjectFromSelection(TargetRootObject);
-
 	// Create new folder
-	TMap<FName, FName> FolderMap;
+	TMap<FName, FName> CreatedFolders;
 	for (FName Folder : CacheFoldersEdit)
 	{
-		FName ParentPath = FEditorFolderUtils::GetParentPath(Folder);
-		FName LeafName = FEditorFolderUtils::GetLeafName(Folder);
-		if (LeafName != TEXT(""))
+		if (FName* NewFolder = CacheFolderMap.Find(Folder))
 		{
-			if (FName* NewParentPath = FolderMap.Find(ParentPath))
+			// When using Actor Folder, duplicated actors might already have created the actor folder (when destination rootobject is different)
+			if (Mode->CreateFolder(FFolder(CacheFoldersEditRootObject, *NewFolder)))
 			{
-				ParentPath = *NewParentPath;
+				CreatedFolders.Add(Folder, *NewFolder);
 			}
-
-			FFolder NewFolderPath = Mode->CreateFolder(FFolder(ParentPath, TargetRootObject), LeafName);
-			FolderMap.Add(Folder, NewFolderPath.GetPath());
 		}
 	}
 
@@ -1550,10 +1564,10 @@ void SSceneOutliner::PasteFoldersEnd()
 	for (FName OldFolderName : CacheFoldersEdit)
 	{
 		// Get the new folder that was created from this name
-		if (const FName* NewFolderName = FolderMap.Find(OldFolderName))
+		if (const FName* NewFolderName = CreatedFolders.Find(OldFolderName))
 		{
-			FFolder NewFolder(*NewFolderName, TargetRootObject);
-			FFolder OldFolder(OldFolderName, TargetRootObject);
+			FFolder NewFolder(CacheFoldersEditRootObject, *NewFolderName);
+			FFolder OldFolder(CacheFoldersEditRootObject, OldFolderName);
 
 			if (FSceneOutlinerTreeItemPtr* OldFolderItem = TreeItemMap.Find(OldFolder))
 			{
@@ -1573,6 +1587,8 @@ void SSceneOutliner::PasteFoldersEnd()
 	}
 
 	CacheFoldersEdit.Reset();
+	CacheFoldersEditRootObject = FFolder::GetInvalidRootObject();
+	CacheFolderMap.Reset();
 	CachePasteFolderExistingChildrenMap.Reset();
 	FullRefresh();
 }
diff --git a/Engine/Source/Editor/SceneOutliner/Public/ActorBrowsingMode.h b/Engine/Source/Editor/SceneOutliner/Public/ActorBrowsingMode.h
--- a/Engine/Source/Editor/SceneOutliner/Public/ActorBrowsingMode.h
+++ b/Engine/Source/Editor/SceneOutliner/Public/ActorBrowsingMode.h
@@ -14,7 +14,7 @@ public:
 
 	/* Begin ISceneOutlinerMode Interface */
 	virtual void Rebuild() override;
-	virtual FCreateSceneOutlinerMode CreateFolderPickerMode(const FFolder::FRootObject& InRootObject = FFolder::GetDefaultRootObject()) const override;
+	virtual FCreateSceneOutlinerMode CreateFolderPickerMode(const FFolder::FRootObject& InRootObject = FFolder::GetInvalidRootObject()) const override;
 	virtual void CreateViewContent(FMenuBuilder& MenuBuilder) override;
 	virtual TSharedPtr<FDragDropOperation> CreateDragDropOperation(const FPointerEvent& MouseEvent, const TArray<FSceneOutlinerTreeItemPtr>& InTreeItems) const override;
 	virtual bool ParseDragDrop(FSceneOutlinerDragDropPayload& OutPayload, const FDragDropOperation& Operation) const override;
@@ -46,7 +46,8 @@ public:
 	virtual bool CanPaste() const override;
 	virtual bool CanSupportDragAndDrop() const { return true; }
 	virtual FFolder CreateNewFolder() override;
-	virtual FFolder CreateFolder(const FFolder& ParentPath, const FName& LeafName) override;
+	virtual FFolder GetFolder(const FFolder& ParentPath, const FName& LeafName) override;
+	virtual bool CreateFolder(const FFolder& NewFolder) override;
 	virtual bool ReparentItemToFolder(const FFolder& FolderPath, const FSceneOutlinerTreeItemPtr& Item) override;
 	virtual void SelectFoldersDescendants(const TArray<FFolderTreeItem*>& FolderItems, bool bSelectImmediateChildrenOnly) override;
 	virtual void PinItem(const FSceneOutlinerTreeItemPtr& InItem) override;
@@ -104,8 +105,8 @@ private:
 	void RegisterContextMenu();
 	bool CanPasteFoldersOnlyFromClipboard() const;
 
-	bool GetFolderNamesFromFolders(const TArray<FFolder>& InFolders, TArray<FName>& OutFolders, FFolder::FRootObject& OutCommonRootObject) const;
 	bool GetFolderNamesFromPayload(const FSceneOutlinerDragDropPayload& InPayload, TArray<FName>& OutFolders, FFolder::FRootObject& OutCommonRootObject) const;
+	FFolder GetWorldDefaultRootFolder() const;
 
 	/** Filter factories */
 	static TSharedRef<FSceneOutlinerFilter> CreateShowOnlySelectedActorsFilter();
diff --git a/Engine/Source/Editor/SceneOutliner/Public/ActorFolderHierarchy.h b/Engine/Source/Editor/SceneOutliner/Public/ActorFolderHierarchy.h
--- a/Engine/Source/Editor/SceneOutliner/Public/ActorFolderHierarchy.h
+++ b/Engine/Source/Editor/SceneOutliner/Public/ActorFolderHierarchy.h
@@ -8,7 +8,7 @@
 class FActorFolderHierarchy : public ISceneOutlinerHierarchy
 {
 public:
-	FActorFolderHierarchy(ISceneOutlinerMode* InMode, const TWeakObjectPtr<UWorld>& World, const FFolder::FRootObject& InRootObject = FFolder::GetDefaultRootObject());
+	FActorFolderHierarchy(ISceneOutlinerMode* InMode, const TWeakObjectPtr<UWorld>& World, const FFolder::FRootObject& InRootObject = FFolder::GetInvalidRootObject());
 	virtual ~FActorFolderHierarchy() {}
 
 
diff --git a/Engine/Source/Editor/SceneOutliner/Public/ActorFolderPickingMode.h b/Engine/Source/Editor/SceneOutliner/Public/ActorFolderPickingMode.h
--- a/Engine/Source/Editor/SceneOutliner/Public/ActorFolderPickingMode.h
+++ b/Engine/Source/Editor/SceneOutliner/Public/ActorFolderPickingMode.h
@@ -7,7 +7,7 @@
 class SCENEOUTLINER_API FActorFolderPickingMode : public FActorMode
 {
 public:
-	FActorFolderPickingMode(SSceneOutliner* InSceneOutliner, FOnSceneOutlinerItemPicked InOnItemPicked, TWeakObjectPtr<UWorld> InSpecifiedWorldToDisplay = nullptr, const FFolder::FRootObject& InRootObject = FFolder::GetDefaultRootObject());
+	FActorFolderPickingMode(SSceneOutliner* InSceneOutliner, FOnSceneOutlinerItemPicked InOnItemPicked, TWeakObjectPtr<UWorld> InSpecifiedWorldToDisplay = nullptr, const FFolder::FRootObject& InRootObject = FFolder::GetInvalidRootObject());
 	virtual ~FActorFolderPickingMode() {}
 
 	/* Begin ISceneOutlinerMode Implementation */
diff --git a/Engine/Source/Editor/SceneOutliner/Public/ActorFolderTreeItem.h b/Engine/Source/Editor/SceneOutliner/Public/ActorFolderTreeItem.h
--- a/Engine/Source/Editor/SceneOutliner/Public/ActorFolderTreeItem.h
+++ b/Engine/Source/Editor/SceneOutliner/Public/ActorFolderTreeItem.h
@@ -23,6 +23,7 @@ public:
 	virtual bool CanInteract() const override;
 	virtual TSharedRef<SWidget> GenerateLabelWidget(ISceneOutliner& Outliner, const STableRow<FSceneOutlinerTreeItemPtr>& InRow) override;
 	virtual bool ShouldShowPinnedState() const override;
+	virtual FFolder GetFolder() const override;
 	/* End FFolderTreeItem Implementation */
 		
 	/* Begin FFolderTreeItem Implementation */
diff --git a/Engine/Source/Editor/SceneOutliner/Public/ActorMode.h b/Engine/Source/Editor/SceneOutliner/Public/ActorMode.h
--- a/Engine/Source/Editor/SceneOutliner/Public/ActorMode.h
+++ b/Engine/Source/Editor/SceneOutliner/Public/ActorMode.h
@@ -61,6 +61,10 @@ public:
 	virtual int32 GetTypeSortPriority(const ISceneOutlinerTreeItem& Item) const override;
 
 	static bool IsActorDisplayable(const SSceneOutliner* SceneOutliner, const AActor* Actor);
+
+	virtual FFolder::FRootObject GetRootObject() const override;
+	virtual FFolder::FRootObject GetPasteTargetRootObject() const override;
+
 private:
 	/** Called when the user selects a world in the world picker menu */
 	void OnSelectWorld(TWeakObjectPtr<UWorld> World);
diff --git a/Engine/Source/Editor/SceneOutliner/Public/FolderTreeItem.h b/Engine/Source/Editor/SceneOutliner/Public/FolderTreeItem.h
--- a/Engine/Source/Editor/SceneOutliner/Public/FolderTreeItem.h
+++ b/Engine/Source/Editor/SceneOutliner/Public/FolderTreeItem.h
@@ -63,7 +63,7 @@ public:
 	virtual void SetPath(const FName& InNewPath);
 	const FName& GetPath() const { return Path; }
 	const FName& GetLeafName() const { return LeafName; }
-	FFolder GetFolder() const { return FFolder(Path, RootObject); }
+	virtual FFolder GetFolder() const { return FFolder(RootObject, Path); }
 	
 	virtual bool ShouldShowPinnedState() const override { return true; }
 	virtual bool HasPinnedStateInfo() const override { return false; }
diff --git a/Engine/Source/Editor/SceneOutliner/Public/ISceneOutlinerMode.h b/Engine/Source/Editor/SceneOutliner/Public/ISceneOutlinerMode.h
--- a/Engine/Source/Editor/SceneOutliner/Public/ISceneOutlinerMode.h
+++ b/Engine/Source/Editor/SceneOutliner/Public/ISceneOutlinerMode.h
@@ -38,7 +38,7 @@ public:
 	virtual void CreateViewContent(FMenuBuilder& MenuBuilder) {}
 
 	/** Returns a factory to create a Folder Picker mode which matches this mode */
-	virtual FCreateSceneOutlinerMode CreateFolderPickerMode(const FFolder::FRootObject& InRootObject = FFolder::GetDefaultRootObject()) const { return FCreateSceneOutlinerMode(); }
+	virtual FCreateSceneOutlinerMode CreateFolderPickerMode(const FFolder::FRootObject& InRootObject = FFolder::GetInvalidRootObject()) const { return FCreateSceneOutlinerMode(); }
 public:
 	/** Construct a new Drag and drop operation */
 	virtual TSharedPtr<FDragDropOperation> CreateDragDropOperation(const FPointerEvent& MouseEvent, const TArray<FSceneOutlinerTreeItemPtr>& InTreeItems) const { return nullptr; }
@@ -89,6 +89,11 @@ public:
 	virtual void SynchronizeSelection() {}
 	/** Trigger a duplication of selected items */
 	virtual void OnDuplicateSelected() {}
+
+	/** Returns the root object of this outliner mode */
+	virtual FFolder::FRootObject GetRootObject() const { return FFolder::GetInvalidRootObject(); }
+	/** Returns the target root object for a paste operation */
+	virtual FFolder::FRootObject GetPasteTargetRootObject() const { return FFolder::GetInvalidRootObject(); }
 public:
 	/** Get the current selection mode */
 	virtual ESelectionMode::Type GetSelectionMode() const { return ESelectionMode::Single; }
@@ -130,9 +135,11 @@ public:
 	/* Folder management */
 		
 	/** Creates a new folder item at the root with a valid name*/
-	virtual FFolder CreateNewFolder() { return FFolder(); }
+	virtual FFolder CreateNewFolder() { return FFolder::GetInvalidFolder(); }
+	/** Returns a unique folder path for a specific parent with a given leaf name (without creating it) */
+	virtual FFolder GetFolder(const FFolder& ParentPath, const FName& LeafName) { return FFolder::GetInvalidFolder(); }
 	/** Create a folder under a specific parent with a given leaf name */
-	virtual FFolder CreateFolder(const FFolder& ParentPath, const FName& LeafName) { return FFolder(); }
+	virtual bool CreateFolder(const FFolder& NewFolder) { return false; }
 	/** Reparent an item to a given folder path. Returns true if the operation is sucessful */
 	virtual bool ReparentItemToFolder(const FFolder& FolderPath, const FSceneOutlinerTreeItemPtr& Item) { return false; }
 	/** Select all descendants of a folder. Optionally select only immediate descendants. */
diff --git a/Engine/Source/Editor/SceneOutliner/Public/ISceneOutlinerTreeItem.h b/Engine/Source/Editor/SceneOutliner/Public/ISceneOutlinerTreeItem.h
--- a/Engine/Source/Editor/SceneOutliner/Public/ISceneOutlinerTreeItem.h
+++ b/Engine/Source/Editor/SceneOutliner/Public/ISceneOutlinerTreeItem.h
@@ -121,7 +121,7 @@ public:
 	virtual FSceneOutlinerTreeItemID GetID() const = 0;
 
 	/** Get the optional context object root of a tree item. */
-	virtual FFolder::FRootObject GetRootObject() const { return FFolder::GetDefaultRootObject(); }
+	virtual FFolder::FRootObject GetRootObject() const { return FFolder::GetInvalidRootObject(); }
 
 	/** Get the raw string to display for this tree item - used for sorting */
 	virtual FString GetDisplayString() const = 0;
diff --git a/Engine/Source/Editor/SceneOutliner/Public/SSceneOutliner.h b/Engine/Source/Editor/SceneOutliner/Public/SSceneOutliner.h
--- a/Engine/Source/Editor/SceneOutliner/Public/SSceneOutliner.h
+++ b/Engine/Source/Editor/SceneOutliner/Public/SSceneOutliner.h
@@ -601,8 +601,6 @@ private:
 	/** Get an array of selected folders */
 	void GetSelectedFolders(TArray<FFolderTreeItem*>& OutFolders) const;
 
-	/** Get an array of selected folder names */
-	TArray<FName> GetSelectedFolderNames() const;
 private:
 	/** Tree view event bindings */
 
@@ -677,6 +675,12 @@ private:
 	/** Cache folders for cut/copy/paste/duplicate */
 	TArray<FName> CacheFoldersEdit;
 
+	/** CacheFoldersEdit target root object */
+	FFolder::FRootObject CacheFoldersEditRootObject;
+
+	/** Cache folders mapping (old to new) for cut/copy/paste/duplicate */
+	TMap<FName, FName> CacheFolderMap;
+
 	/** Cache clipboard contents for cut/copy */
 	FString CacheClipboardContents;
 
diff --git a/Engine/Source/Editor/Sequencer/Private/SSequencer.cpp b/Engine/Source/Editor/Sequencer/Private/SSequencer.cpp
--- a/Engine/Source/Editor/Sequencer/Private/SSequencer.cpp
+++ b/Engine/Source/Editor/Sequencer/Private/SSequencer.cpp
@@ -3615,8 +3615,10 @@ void SSequencer::OnActorsDropped( FActorDragDropOp& DragDropOp )
 void SSequencer::OnFolderDropped( FFolderDragDropOp& DragDropOp )
 {
 	// Sequencer doesn't support dragging folder with a root object
-	if (FFolder::HasRootObject(DragDropOp.RootObject))
+	if (!FFolder::IsRootObjectPersistentLevel(DragDropOp.RootObject))
+	{
 		return;
+	}
 
 	const FScopedTransaction Transaction(LOCTEXT("DropActors", "Drop Actors"));
 
diff --git a/Engine/Source/Editor/UnrealEd/Private/EditorActorFolders.cpp b/Engine/Source/Editor/UnrealEd/Private/EditorActorFolders.cpp
--- a/Engine/Source/Editor/UnrealEd/Private/EditorActorFolders.cpp
+++ b/Engine/Source/Editor/UnrealEd/Private/EditorActorFolders.cpp
@@ -99,7 +99,7 @@ void FActorFolders::BroadcastOnActorFolderCreated(UWorld& InWorld, const FFolder
 {
 	OnFolderCreated.Broadcast(InWorld, InFolder);
 	PRAGMA_DISABLE_DEPRECATION_WARNINGS
-	if (!InFolder.HasRootObject())
+	if (InFolder.IsRootObjectPersistentLevel())
 	{
 		OnFolderCreate.Broadcast(InWorld, InFolder.GetPath());
 	}
@@ -111,7 +111,7 @@ void FActorFolders::BroadcastOnActorFolderDeleted(UWorld& InWorld, const FFolder
 {
 	OnFolderDeleted.Broadcast(InWorld, InFolder);
 	PRAGMA_DISABLE_DEPRECATION_WARNINGS
-	if (!InFolder.HasRootObject())
+	if (InFolder.IsRootObjectPersistentLevel())
 	{
 		OnFolderDelete.Broadcast(InWorld, InFolder.GetPath());
 	}
@@ -123,7 +123,7 @@ void FActorFolders::BroadcastOnActorFolderMoved(UWorld& InWorld, const FFolder&
 {
 	OnFolderMoved.Broadcast(InWorld, InSrcFolder, InDstFolder);
 	PRAGMA_DISABLE_DEPRECATION_WARNINGS
-	if (!InSrcFolder.HasRootObject() && !InDstFolder.HasRootObject())
+	if (InSrcFolder.IsRootObjectPersistentLevel() && InDstFolder.IsRootObjectPersistentLevel())
 	{
 		OnFolderMove.Broadcast(InWorld, InSrcFolder.GetPath(), InDstFolder.GetPath());
 	}
@@ -176,6 +176,11 @@ void FActorFolders::OnWorldSaved(UWorld* World, FObjectPostSaveContext ObjectSav
 	}
 }
 
+bool FActorFolders::IsInitializedForWorld(UWorld& InWorld) const
+{
+	return !!WorldFolders.Find(&InWorld);
+}
+
 UWorldFolders& FActorFolders::GetOrCreateWorldFolders(UWorld& InWorld)
 {
 	if (UWorldFolders** Folders = WorldFolders.Find(&InWorld))
@@ -250,13 +255,13 @@ FFolder FActorFolders::GetDefaultFolderForSelection(UWorld& InWorld, TArray<FFol
 		}
 		else if (CommonFolder.GetValue().GetRootObject() != Folder.GetRootObject())
 		{
-			CommonFolder = FFolder();
+			CommonFolder.Reset();
 			return false;
 		}
 		else if (CommonFolder.GetValue().GetPath() != Folder.GetPath())
 		{
 			// Empty path and continue iterating as we need to continue validating RootObjects
-			CommonFolder.GetValue().SetPath(NAME_None);
+			CommonFolder = FFolder(CommonFolder->GetRootObject());
 		}
 		return true;
 	};
@@ -272,7 +277,7 @@ FFolder FActorFolders::GetDefaultFolderForSelection(UWorld& InWorld, TArray<FFol
 		{
 			if (LevelInstance->IsEditing())
 			{
-				Folder = FFolder(FFolder::GetEmptyPath(), FFolder::FRootObject(Actor));
+				Folder = FFolder(FFolder::FRootObject(Actor), FFolder::GetEmptyPath());
 			}
 		}
 		if (!MergeFolders(Folder))
@@ -292,14 +297,16 @@ FFolder FActorFolders::GetDefaultFolderForSelection(UWorld& InWorld, TArray<FFol
 		}
 	}
 
-	return GetDefaultFolderName(InWorld, CommonFolder.Get(FFolder()));
+	return GetDefaultFolderName(InWorld, CommonFolder.Get(FFolder::GetInvalidFolder()));
 }
 
 FFolder FActorFolders::GetFolderName(UWorld& InWorld, const FFolder& InParentFolder, const FName& InLeafName)
 {
+	FFolder ParentFolder = FFolder::IsRootObjectValid(InParentFolder.GetRootObject()) ? InParentFolder : FFolder(GetWorldFolderRootObject(InWorld), InParentFolder.GetPath());
+
 	// This is potentially very slow but necessary to find a unique name
 	const UWorldFolders& Folders = GetOrCreateWorldFolders(InWorld);
-	const FFolder::FRootObject& RootObject = InParentFolder.GetRootObject();
+	const FFolder::FRootObject& RootObject = ParentFolder.GetRootObject();
 	const FString LeafNameString = InLeafName.ToString();
 
 	// Find the last non-numeric character
@@ -333,14 +340,14 @@ FFolder FActorFolders::GetFolderName(UWorld& InWorld, const FFolder& InParentFol
 
 	FText LeafName = FText::Format(LOCTEXT("FolderNamePattern", "{0}{1}"), FText::FromString(LeafNameRoot), SuffixLen > 0 ? FText::AsNumber(Suffix++, &NumberFormat) : FText::GetEmpty());
 
-	FString ParentFolderPath = InParentFolder.IsNone() ? TEXT("") : InParentFolder.ToString();
+	FString ParentFolderPath = ParentFolder.IsNone() ? TEXT("") : ParentFolder.ToString();
 	if (!ParentFolderPath.IsEmpty())
 	{
 		ParentFolderPath += "/";
 	}
 
 	FName FolderName(*(ParentFolderPath + LeafName.ToString()));
-	while (Folders.ContainsFolder(FFolder(FolderName, RootObject)))
+	while (Folders.ContainsFolder(FFolder(RootObject, FolderName)))
 	{
 		LeafName = FText::Format(LOCTEXT("FolderNamePattern", "{0}{1}"), FText::FromString(LeafNameRoot), FText::AsNumber(Suffix++, &NumberFormat));
 		FolderName = FName(*(ParentFolderPath + LeafName.ToString()));
@@ -352,14 +359,16 @@ FFolder FActorFolders::GetFolderName(UWorld& InWorld, const FFolder& InParentFol
 		}
 	}
 
-	return FFolder(FolderName, RootObject);
+	return FFolder(RootObject, FolderName);
 }
 
 FFolder FActorFolders::GetDefaultFolderName(UWorld& InWorld, const FFolder& InParentFolder)
 {
+	FFolder ParentFolder = FFolder::IsRootObjectValid(InParentFolder.GetRootObject()) ? InParentFolder : FFolder(GetWorldFolderRootObject(InWorld), InParentFolder.GetPath());
+
 	// This is potentially very slow but necessary to find a unique name
 	const UWorldFolders& Folders = GetOrCreateWorldFolders(InWorld);
-	const FFolder::FRootObject& RootObject = InParentFolder.GetRootObject();
+	const FFolder::FRootObject& RootObject = ParentFolder.GetRootObject();
 
 	// Create a valid base name for this folder
 	FNumberFormattingOptions NumberFormat;
@@ -367,14 +376,14 @@ FFolder FActorFolders::GetDefaultFolderName(UWorld& InWorld, const FFolder& InPa
 	uint32 Suffix = 1;
 	FText LeafName = FText::Format(LOCTEXT("DefaultFolderNamePattern", "NewFolder{0}"), FText::AsNumber(Suffix++, &NumberFormat));
 
-	FString ParentFolderPath = InParentFolder.IsNone() ? TEXT("") : InParentFolder.ToString();
+	FString ParentFolderPath = ParentFolder.IsNone() ? TEXT("") : ParentFolder.ToString();
 	if (!ParentFolderPath.IsEmpty())
 	{
 		ParentFolderPath += "/";
 	}
 
 	FName FolderName(*(ParentFolderPath + LeafName.ToString()));
-	while (Folders.ContainsFolder(FFolder(FolderName, RootObject)))
+	while (Folders.ContainsFolder(FFolder(RootObject, FolderName)))
 	{
 		LeafName = FText::Format(LOCTEXT("DefaultFolderNamePattern", "NewFolder{0}"), FText::AsNumber(Suffix++, &NumberFormat));
 		FolderName = FName(*(ParentFolderPath + LeafName.ToString()));
@@ -386,7 +395,7 @@ FFolder FActorFolders::GetDefaultFolderName(UWorld& InWorld, const FFolder& InPa
 		}
 	}
 
-	return FFolder(FolderName, RootObject);
+	return FFolder(RootObject, FolderName);
 }
 
 void FActorFolders::CreateFolderContainingSelection(UWorld& InWorld, const FFolder& InFolder)
@@ -423,14 +432,16 @@ void FActorFolders::SetSelectedFolderPath(const FFolder& InFolder) const
 	}
 }
 
-void FActorFolders::CreateFolder(UWorld& InWorld, const FFolder& InFolder)
+bool FActorFolders::CreateFolder(UWorld& InWorld, const FFolder& InFolder)
 {
 	FScopedTransaction Transaction(LOCTEXT("UndoAction_CreateFolder", "Create Folder"));
 
 	if (AddFolderToWorld(InWorld, InFolder))
 	{
 		BroadcastOnActorFolderCreated(InWorld, InFolder);
+		return true;
 	}
+	return false;
 }
 
 void FActorFolders::OnActorFolderAdded(UActorFolder* InActorFolder)
@@ -440,13 +451,16 @@ void FActorFolders::OnActorFolderAdded(UActorFolder* InActorFolder)
 	check(Level->IsUsingActorFolders());
 
 	AddFolderToWorld(*Level->GetWorld(), InActorFolder->GetFolder());
+	// @todo_ow: This code completely overrides the expanded state that was read from the json file (see UWorldFolders::LoadState)
 	SetIsFolderExpanded(*Level->GetWorld(), InActorFolder->GetFolder(), InActorFolder->IsInitiallyExpanded());
 }
 
 void FActorFolders::OnFolderRootObjectRemoved(UWorld& InWorld, const FFolder::FRootObject& InFolderRootObject)
 {
+	const FFolder::FRootObject FolderRootObject = FFolder::IsRootObjectValid(InFolderRootObject) ? InFolderRootObject : GetWorldFolderRootObject(InWorld);
+
 	TArray<FFolder> FoldersToDelete;
-	ForEachFolderWithRootObject(InWorld, InFolderRootObject, [&FoldersToDelete](const FFolder& Folder)
+	ForEachFolderWithRootObject(InWorld, FolderRootObject, [&FoldersToDelete](const FFolder& Folder)
 	{
 		FoldersToDelete.Add(Folder);
 		return true;
@@ -527,7 +541,7 @@ FFolder FActorFolders::GetActorEditorContextFolder(UWorld& InWorld) const
 	{
 		return (*Folders)->GetActorEditorContextFolder();
 	}
-	return FFolder();
+	return FFolder::GetWorldRootFolder(&InWorld);
 }
 
 void FActorFolders::SetActorEditorContextFolder(UWorld& InWorld, const FFolder& InFolder)
@@ -557,7 +571,7 @@ void FActorFolders::OnExecuteActorEditorContextAction(UWorld* InWorld, const EAc
 			}
 			break;
 		case EActorEditorContextAction::ResetContext:
-			SetActorEditorContextFolder(*InWorld, FFolder());
+			SetActorEditorContextFolder(*InWorld, FFolder::GetWorldRootFolder(InWorld));
 			break;
 		case EActorEditorContextAction::PushContext:
 			if (UWorldFolders** Folders = (UWorldFolders**)WorldFolders.Find(InWorld))
@@ -603,16 +617,24 @@ void FActorFolders::ForEachFolder(UWorld& InWorld, TFunctionRef<bool(const FFold
 	GetOrCreateWorldFolders(InWorld).ForEachFolder(Operation);
 }
 
+FFolder::FRootObject FActorFolders::GetWorldFolderRootObject(UWorld& InWorld)
+{
+	return FFolder::GetWorldRootFolder(&InWorld).GetRootObject();
+}
+
 void FActorFolders::ForEachFolderWithRootObject(UWorld& InWorld, const FFolder::FRootObject& InFolderRootObject, TFunctionRef<bool(const FFolder&)> Operation)
 {
-	GetOrCreateWorldFolders(InWorld).ForEachFolderWithRootObject(InFolderRootObject, Operation);
+	const FFolder::FRootObject FolderRootObject = FFolder::IsRootObjectValid(InFolderRootObject) ? InFolderRootObject : GetWorldFolderRootObject(InWorld);
+	GetOrCreateWorldFolders(InWorld).ForEachFolderWithRootObject(FolderRootObject, Operation);
 }
 
-void FActorFolders::ForEachActorInFolders(UWorld& World, const TArray<FName>& Paths, TFunctionRef<bool(AActor*)> Operation, const FFolder::FRootObject& InFolderRootObject)
+void FActorFolders::ForEachActorInFolders(UWorld& InWorld, const TArray<FName>& Paths, TFunctionRef<bool(AActor*)> Operation, const FFolder::FRootObject& InFolderRootObject /*= FFolder::GetInvalidRootObject()*/)
 {
-	for (FActorIterator ActorIt(&World); ActorIt; ++ActorIt)
+	const FFolder::FRootObject FolderRootObject = FFolder::IsRootObjectValid(InFolderRootObject) ? InFolderRootObject : GetWorldFolderRootObject(InWorld);
+
+	for (FActorIterator ActorIt(&InWorld); ActorIt; ++ActorIt)
 	{
-		if (ActorIt->GetFolderRootObject() != InFolderRootObject)
+		if (ActorIt->GetFolderRootObject() != FolderRootObject)
 		{
 			continue;
 		}
@@ -629,22 +651,26 @@ void FActorFolders::ForEachActorInFolders(UWorld& World, const TArray<FName>& Pa
 	}
 }
 
-void FActorFolders::GetActorsFromFolders(UWorld& World, const TArray<FName>& Paths, TArray<AActor*>& OutActors, const FFolder::FRootObject& InFolderRootObject)
+void FActorFolders::GetActorsFromFolders(UWorld& InWorld, const TArray<FName>& Paths, TArray<AActor*>& OutActors, const FFolder::FRootObject& InFolderRootObject /*= FFolder::GetInvalidRootObject()*/)
 {
-	ForEachActorInFolders(World, Paths, [&OutActors](AActor* InActor)
+	const FFolder::FRootObject FolderRootObject = FFolder::IsRootObjectValid(InFolderRootObject) ? InFolderRootObject : GetWorldFolderRootObject(InWorld);
+
+	ForEachActorInFolders(InWorld, Paths, [&OutActors](AActor* InActor)
 	{
 		OutActors.Add(InActor);
 		return true;
-	}, InFolderRootObject);
+	}, FolderRootObject);
 }
 
-void FActorFolders::GetWeakActorsFromFolders(UWorld& World, const TArray<FName>& Paths, TArray<TWeakObjectPtr<AActor>>& OutActors, const FFolder::FRootObject& InFolderRootObject)
+void FActorFolders::GetWeakActorsFromFolders(UWorld& InWorld, const TArray<FName>& Paths, TArray<TWeakObjectPtr<AActor>>& OutActors, const FFolder::FRootObject& InFolderRootObject /*= FFolder::GetInvalidRootObject()*/)
 {
-	ForEachActorInFolders(World, Paths, [&OutActors](AActor* InActor)
+	const FFolder::FRootObject FolderRootObject = FFolder::IsRootObjectValid(InFolderRootObject) ? InFolderRootObject : GetWorldFolderRootObject(InWorld);
+
+	ForEachActorInFolders(InWorld, Paths, [&OutActors](AActor* InActor)
 	{
 		OutActors.Add(InActor);
 		return true;
-	}, InFolderRootObject);
+	}, FolderRootObject);
 }
 
 ////////////////////////////////////////////
@@ -652,12 +678,12 @@ void FActorFolders::GetWeakActorsFromFolders(UWorld& World, const TArray<FName>&
 
 FActorFolderProps* FActorFolders::GetFolderProperties(UWorld& InWorld, FName InPath)
 {
-	return GetFolderProperties(InWorld, FFolder(InPath));
+	return GetFolderProperties(InWorld, FFolder(GetWorldFolderRootObject(InWorld), InPath));
 }
 
 FName FActorFolders::GetDefaultFolderName(UWorld& InWorld, FName ParentPath)
 {
-	FFolder DefaultFolderName = GetDefaultFolderName(InWorld, FFolder(ParentPath));
+	FFolder DefaultFolderName = GetDefaultFolderName(InWorld, FFolder(GetWorldFolderRootObject(InWorld), ParentPath));
 	return DefaultFolderName.GetPath();
 }
 
@@ -669,33 +695,33 @@ FName FActorFolders::GetDefaultFolderNameForSelection(UWorld& InWorld)
 
 FName FActorFolders::GetFolderName(UWorld& InWorld, FName InParentPath, FName InFolderName)
 {
-	FFolder FolderName = GetFolderName(InWorld, FFolder(InParentPath), InFolderName);
+	FFolder FolderName = GetFolderName(InWorld, FFolder(GetWorldFolderRootObject(InWorld), InParentPath), InFolderName);
 	return FolderName.GetPath();
 }
 
 void FActorFolders::CreateFolder(UWorld& InWorld, FName Path)
 {
-	CreateFolder(InWorld, FFolder(Path));
+	CreateFolder(InWorld, FFolder(GetWorldFolderRootObject(InWorld), Path));
 }
 
 void FActorFolders::CreateFolderContainingSelection(UWorld& InWorld, FName Path)
 {
-	CreateFolderContainingSelection(InWorld, FFolder(Path));
+	CreateFolderContainingSelection(InWorld, FFolder(GetWorldFolderRootObject(InWorld), Path));
 }
 
 void FActorFolders::SetSelectedFolderPath(FName Path) const
 {
-	SetSelectedFolderPath(FFolder(Path));
+	SetSelectedFolderPath(FFolder(FFolder::GetInvalidRootObject(), Path));
 }
 
 void FActorFolders::DeleteFolder(UWorld& InWorld, FName FolderToDelete)
 {
-	DeleteFolder(InWorld, FFolder(FolderToDelete));
+	DeleteFolder(InWorld, FFolder(GetWorldFolderRootObject(InWorld), FolderToDelete));
 }
 
 bool FActorFolders::RenameFolderInWorld(UWorld& InWorld, FName OldPath, FName NewPath)
 {
-	return RenameFolderInWorld(InWorld, FFolder(OldPath), FFolder(NewPath));
+	return RenameFolderInWorld(InWorld, FFolder(GetWorldFolderRootObject(InWorld), OldPath), FFolder(GetWorldFolderRootObject(InWorld), NewPath));
 }
 
 //~ End Deprecated
diff --git a/Engine/Source/Editor/UnrealEd/Private/WorldFolders.cpp b/Engine/Source/Editor/UnrealEd/Private/WorldFolders.cpp
--- a/Engine/Source/Editor/UnrealEd/Private/WorldFolders.cpp
+++ b/Engine/Source/Editor/UnrealEd/Private/WorldFolders.cpp
@@ -40,7 +40,7 @@ void UWorldFolders::RebuildList()
 	TArray<FFolder> FoldersToRemove;
 	ForEachFolder([&FoldersToRemove](const FFolder& Folder)
 	{
-		if (Folder.HasRootObject())
+		if (!Folder.IsRootObjectPersistentLevel())
 		{
 			FoldersToRemove.Add(Folder);
 		}
@@ -58,16 +58,20 @@ void UWorldFolders::RebuildList()
 		AddFolder(ActorIt->GetFolder());
 	}
 
-	for (ULevel* Level : GetWorld()->GetLevels())
+	// Add levels ActorFolders as they are still needed for unloaded actors (only in editor)
+	if (!GetWorld()->IsGameWorld())
 	{
-		const bool bIsLevelVisibleOrAssociating = (Level->bIsVisible && !Level->bIsBeingRemoved) || Level->bIsAssociatingLevel || Level->bIsDisassociatingLevel;
-		if (bIsLevelVisibleOrAssociating)
+		for (ULevel* Level : GetWorld()->GetLevels())
 		{
-			Level->ForEachActorFolder([this](UActorFolder* ActorFolder)
+			const bool bIsLevelVisibleOrAssociating = (Level->bIsVisible && !Level->bIsBeingRemoved) || Level->bIsAssociatingLevel || Level->bIsDisassociatingLevel;
+			if (bIsLevelVisibleOrAssociating)
 			{
-				AddFolder(ActorFolder->GetFolder());
-				return true;
-			}, /*bSkipDeleted*/ true);
+				Level->ForEachActorFolder([this](UActorFolder* ActorFolder)
+				{
+					AddFolder(ActorFolder->GetFolder());
+					return true;
+				}, /*bSkipDeleted*/ true);
+			}
 		}
 	}
 }
@@ -188,13 +192,13 @@ FFolder UWorldFolders::GetActorEditorContextFolder() const
 		ULevelInstanceSubsystem* LevelInstanceSubsystem = GetWorld()->GetSubsystem<ULevelInstanceSubsystem>();
 		ILevelInstanceInterface* EditingLevelInstance = LevelInstanceSubsystem ? LevelInstanceSubsystem->GetEditingLevelInstance() : nullptr;
 		UObject* EditingLevelInstanceObject = EditingLevelInstance ? CastChecked<UObject>(EditingLevelInstance) : nullptr;
-		FFolder::FRootObject RootObject = FFolder::FRootObject(EditingLevelInstanceObject ? EditingLevelInstanceObject : (GetWorld()->GetCurrentLevel() == GetWorld()->PersistentLevel) ? nullptr : GetWorld()->GetCurrentLevel());
+		FFolder::FRootObject RootObject = FFolder::FRootObject(EditingLevelInstanceObject ? EditingLevelInstanceObject : GetWorld()->GetCurrentLevel());
 		if (Folder.GetRootObject() == RootObject)
 		{
 			return Folder;
 		}
 	}
-	return FFolder();
+	return FFolder::GetWorldRootFolder(GetWorld());
 }
 
 bool UWorldFolders::SetActorEditorContextFolder(const FFolder& InFolder)
@@ -202,8 +206,7 @@ bool UWorldFolders::SetActorEditorContextFolder(const FFolder& InFolder)
 	if (InFolder != CurrentFolder.GetFolder())
 	{
 		Modify();
-		CurrentFolder.Path = InFolder.GetPath();
-		CurrentFolder.RootObjectPtr = InFolder.GetRootObjectPtr();
+		CurrentFolder = InFolder;
 		return true;
 	}
 	return false;
@@ -226,7 +229,7 @@ void UWorldFolders::PopActorEditorContext()
 
 bool UWorldFolders::ContainsFolder(const FFolder& InFolder) const
 {
-	return GetImpl(InFolder).ContainsFolder(InFolder);
+	return InFolder.IsValid() && !InFolder.IsNone() && GetImpl(InFolder).ContainsFolder(InFolder);
 }
 
 void UWorldFolders::ForEachFolder(TFunctionRef<bool(const FFolder&)> Operation)
@@ -283,6 +286,10 @@ FString UWorldFolders::GetWorldStateFilename() const
 
 void UWorldFolders::LoadState()
 {
+	FFolder WorldDefaultFolder = FFolder::GetWorldRootFolder(World.Get());
+	check(WorldDefaultFolder.IsRootObjectValid());
+	const FFolder::FRootObject WorldRootObject = WorldDefaultFolder.GetRootObject();
+
 	// Attempt to load the folder properties from user's saved world state directory and apply them.
 	const auto Filename = GetWorldStateFilename();
 	TUniquePtr<FArchive> Ar(IFileManager::Get().CreateFileReader(*Filename));
@@ -299,7 +306,7 @@ void UWorldFolders::LoadState()
 				// Only pull in the folder's properties if this folder still exists in the world.
 				// This means that old stale folders won't re-appear in the world (they'll won't get serialized when the world is saved anyway)
 				auto FolderProperties = KeyValue.Value->AsObject();
-				const FFolder Folder(*KeyValue.Key);
+				const FFolder Folder(WorldRootObject, *KeyValue.Key);
 				const bool bIsExpanded = FolderProperties->GetBoolField(TEXT("bIsExpanded"));
 				if (!SetIsFolderExpanded(Folder, bIsExpanded))
 				{
@@ -324,7 +331,7 @@ void UWorldFolders::SaveState()
 		ForEachFolder([this, &JsonFolders](const FFolder& Folder)
 		{
 			// Only write for World root
-			if (!Folder.HasRootObject())
+			if (Folder.IsRootObjectPersistentLevel())
 			{
 				TSharedRef<FJsonObject> JsonFolder = MakeShareable(new FJsonObject);
 				JsonFolder->SetBoolField(TEXT("bIsExpanded"), IsFolderExpanded(Folder));
diff --git a/Engine/Source/Editor/UnrealEd/Private/WorldPersistentFolders.cpp b/Engine/Source/Editor/UnrealEd/Private/WorldPersistentFolders.cpp
--- a/Engine/Source/Editor/UnrealEd/Private/WorldPersistentFolders.cpp
+++ b/Engine/Source/Editor/UnrealEd/Private/WorldPersistentFolders.cpp
@@ -42,7 +42,7 @@ bool FWorldPersistentFolders::AddFolder(const FFolder& InFolder)
 
 bool FWorldPersistentFolders::RemoveFolder(const FFolder& InFolder, bool bShouldDeleteFolder)
 {
-	UActorFolder* ActorFolder = GetActorFolder(InFolder);
+	UActorFolder* ActorFolder = InFolder.GetActorFolder();
 	if (!IsValid(ActorFolder))
 	{
 		return false;
@@ -69,9 +69,9 @@ bool FWorldPersistentFolders::RenameFolder(const FFolder& InOldFolder, const FFo
 	check(InOldFolder.GetRootObject() == InNewFolder.GetRootObject());
 	const FFolder::FRootObject& RootObject = InOldFolder.GetRootObject();
 
-	UActorFolder* ActorFolder = GetActorFolder(InOldFolder);
+	UActorFolder* ActorFolder = InOldFolder.GetActorFolder();
 	check(IsValid(ActorFolder));
-	UActorFolder* FoundFolder = GetActorFolder(InNewFolder);
+	UActorFolder* FoundFolder = InNewFolder.GetActorFolder();
 	check(!IsValid(FoundFolder) || !FoundFolder->GetPath().IsEqual(InNewFolder.GetPath(), ENameCase::CaseSensitive));
 
 
@@ -81,7 +81,7 @@ bool FWorldPersistentFolders::RenameFolder(const FFolder& InOldFolder, const FFo
 
 	ModifyFolderAndDetectChanges(Level, RootObject, [this, &ActorFolder, &InNewFolder]()
 	{
-		UActorFolder* ParentActorFolder = GetActorFolder(InNewFolder.GetParent());
+		UActorFolder* ParentActorFolder = InNewFolder.GetParent().GetActorFolder();
 		ActorFolder->SetParent(ParentActorFolder);
 		const FString FolderLabel = InNewFolder.GetLeafName().ToString();
 		ActorFolder->SetLabel(FolderLabel);
@@ -118,8 +118,8 @@ void FWorldPersistentFolders::ModifyFolderAndDetectChanges(ULevel* InLevel, cons
 	{
 		// Update FoldersProperties
 		static FActorFolderProps DefaultFolderProperties;
-		FFolder OldFolder(ChangedFolder.Key, InRootObject);
-		FFolder NewFolder(ChangedFolder.Value, InRootObject);
+		FFolder OldFolder(InRootObject, ChangedFolder.Key);
+		FFolder NewFolder(InRootObject, ChangedFolder.Value);
 		FActorFolderProps* OldFolderProperties = Owner.FoldersProperties.Find(OldFolder);
 		FActorFolderProps FolderProperties = OldFolderProperties ? *OldFolderProperties : DefaultFolderProperties;
 		Owner.FoldersProperties.Remove(OldFolder);
@@ -135,31 +135,21 @@ UActorFolder* FWorldPersistentFolders::GetActorFolder(const FFolder& InFolder, U
 	{
 		return nullptr;
 	}
-	const ULevel* Level = GetRootObjectContainer(InFolder, InWorld);
-	UActorFolder* ActorFolder = Level ? Level->GetActorFolder(InFolder.GetPath(), /*bSkipDeleted*/ false) : nullptr;
+
+	UActorFolder* ActorFolder = InFolder.GetActorFolder();
 	return (!ActorFolder && bInAllowCreate) ? CreateActorFolder(InFolder, InWorld) : ActorFolder;
 }
 
 ULevel* FWorldPersistentFolders::GetRootObjectContainer(const FFolder& InFolder, UWorld* InWorld)
 {
-	const UObject* RootObjectPtr = InFolder.HasRootObject() ? InFolder.GetRootObjectPtr() : InWorld;
-
-	if (const UWorld* WorldPtr = Cast<UWorld>(RootObjectPtr))
+	if (InFolder.IsRootObjectValid())
 	{
-		return WorldPtr->PersistentLevel;
+		return InFolder.GetRootObjectAssociatedLevel();
 	}
-	else if (const ILevelInstanceInterface* LevelInstance = Cast<ILevelInstanceInterface>(RootObjectPtr))
+	else if (InWorld)
 	{
-		return LevelInstance->GetLoadedLevel();
+		return InWorld->PersistentLevel;
 	}
-	else if (const ULevel* Level = Cast<ULevel>(RootObjectPtr))
-	{
-		if (UWorld* OuterWorld = Level->GetTypedOuter<UWorld>())
-		{
-			return OuterWorld->PersistentLevel;
-		}
-	}
-
 	return nullptr;
 }
 
diff --git a/Engine/Source/Editor/UnrealEd/Private/WorldTransientFolders.cpp b/Engine/Source/Editor/UnrealEd/Private/WorldTransientFolders.cpp
--- a/Engine/Source/Editor/UnrealEd/Private/WorldTransientFolders.cpp
+++ b/Engine/Source/Editor/UnrealEd/Private/WorldTransientFolders.cpp
@@ -34,7 +34,7 @@ bool FWorldTransientFolders::RenameFolder(const FFolder& OldPath, const FFolder&
 
 		if (OldPath == Path || Path.IsChildOf(OldPath))
 		{
-			const FFolder NewFolder = FFolder(OldPathToNewPath(OldPathString, NewPathString, FolderPath), RootObject);
+			const FFolder NewFolder = FFolder(RootObject, OldPathToNewPath(OldPathString, NewPathString, FolderPath));
 
 			// Needs to be done this way otherwise case insensitive comparison is used.
 			bool ContainsFolder = false;
diff --git a/Engine/Source/Editor/UnrealEd/Public/DragAndDrop/FolderDragDropOp.h b/Engine/Source/Editor/UnrealEd/Public/DragAndDrop/FolderDragDropOp.h
--- a/Engine/Source/Editor/UnrealEd/Public/DragAndDrop/FolderDragDropOp.h
+++ b/Engine/Source/Editor/UnrealEd/Public/DragAndDrop/FolderDragDropOp.h
@@ -20,7 +20,7 @@ public:
 	/** Root object of folder (can be invalid) */
 	FFolder::FRootObject RootObject;
 
-	void Init(TArray<FName> InFolders, UWorld* InWorld, const FFolder::FRootObject& InRootObject = FFolder::GetDefaultRootObject())
+	void Init(TArray<FName> InFolders, UWorld* InWorld, const FFolder::FRootObject& InRootObject = FFolder::GetInvalidRootObject())
 	{
 		check(InWorld != nullptr);
 		Folders = MoveTemp(InFolders);
diff --git a/Engine/Source/Editor/UnrealEd/Public/EditorActorFolders.h b/Engine/Source/Editor/UnrealEd/Public/EditorActorFolders.h
--- a/Engine/Source/Editor/UnrealEd/Public/EditorActorFolders.h
+++ b/Engine/Source/Editor/UnrealEd/Public/EditorActorFolders.h
@@ -99,13 +99,16 @@ struct UNREALED_API FActorFolders : public FGCObject, public IActorEditorContext
 	//~ End Deprecated
 
 	/** Apply an operation to each actor in the given list of folders. Will stop when operation returns false. */
-	static void ForEachActorInFolders(UWorld& InWorld, const TArray<FName>& Paths, TFunctionRef<bool(AActor*)> Operation, const FFolder::FRootObject& InFolderRootObject = FFolder::GetDefaultRootObject());
+	static void ForEachActorInFolders(UWorld& InWorld, const TArray<FName>& Paths, TFunctionRef<bool(AActor*)> Operation, const FFolder::FRootObject& InFolderRootObject = FFolder::GetInvalidRootObject());
 
 	/** Get an array of actors from a list of folders */
-	static void GetActorsFromFolders(UWorld& InWorld, const TArray<FName>& Paths, TArray<AActor*>& OutActors, const FFolder::FRootObject& InFolderRootObject = FFolder::GetDefaultRootObject());
+	static void GetActorsFromFolders(UWorld& InWorld, const TArray<FName>& Paths, TArray<AActor*>& OutActors, const FFolder::FRootObject& InFolderRootObject = FFolder::GetInvalidRootObject());
 
 	/** Get an array of weak actor pointers from a list of folders */
-	static void GetWeakActorsFromFolders(UWorld& InWorld, const TArray<FName>& Paths, TArray<TWeakObjectPtr<AActor>>& OutActors, const FFolder::FRootObject& InFolderRootObject = FFolder::GetDefaultRootObject());
+	static void GetWeakActorsFromFolders(UWorld& InWorld, const TArray<FName>& Paths, TArray<TWeakObjectPtr<AActor>>& OutActors, const FFolder::FRootObject& InFolderRootObject = FFolder::GetInvalidRootObject());
+
+	/** Tests whether a folder container exists for the specified world */
+	bool IsInitializedForWorld(UWorld& InWorld) const;
 
 	/** Get a default folder name under the specified parent path */
 	FFolder GetDefaultFolderName(UWorld& InWorld, const FFolder& InParentFolder);
@@ -117,7 +120,7 @@ struct UNREALED_API FActorFolders : public FGCObject, public IActorEditorContext
 	FFolder GetFolderName(UWorld& InWorld, const FFolder& InParentFolder, const FName& InLeafName);
 
 	/** Create a new folder in the specified world, of the specified path */
-	void CreateFolder(UWorld& InWorld, const FFolder& InFolder);
+	bool CreateFolder(UWorld& InWorld, const FFolder& InFolder);
 
 	/** Same as CreateFolder, but moves the current actor selection into the new folder as well */
 	void CreateFolderContainingSelection(UWorld& InWorld, const FFolder& InFolder);
@@ -163,6 +166,8 @@ struct UNREALED_API FActorFolders : public FGCObject, public IActorEditorContext
 	void SetActorEditorContextFolder(UWorld& InWorld, const FFolder& InFolder);
 
 private:
+	
+	static FFolder::FRootObject GetWorldFolderRootObject(UWorld& InWorld);
 
 	/** Broadcast when actor folder is created. */
 	void BroadcastOnActorFolderCreated(UWorld& InWorld, const FFolder& InFolder);
diff --git a/Engine/Source/Editor/UnrealEd/Public/WorldFolders.h b/Engine/Source/Editor/UnrealEd/Public/WorldFolders.h
--- a/Engine/Source/Editor/UnrealEd/Public/WorldFolders.h
+++ b/Engine/Source/Editor/UnrealEd/Public/WorldFolders.h
@@ -6,6 +6,7 @@
 #include "UObject/ObjectMacros.h"
 #include "UObject/Object.h"
 #include "UObject/WeakObjectPtr.h"
+#include "Misc/Guid.h"
 #include "WorldFoldersImplementation.h"
 #include "WorldPersistentFolders.h"
 #include "WorldTransientFolders.h"
@@ -38,16 +39,29 @@ struct FActorPlacementFolder
 
 	UPROPERTY()
 	TWeakObjectPtr<UObject> RootObjectPtr;
+	
+	UPROPERTY()
+	FGuid ActorFolderGuid;
 
 	void Reset()
 	{
 		Path = NAME_None;
 		RootObjectPtr.Reset();
+		ActorFolderGuid.Invalidate();
+	}
+
+	FActorPlacementFolder& operator = (const FFolder& InOtherFolder)
+	{
+		Path = InOtherFolder.GetPath();
+		RootObjectPtr = InOtherFolder.GetRootObjectPtr();
+		ActorFolderGuid = InOtherFolder.GetActorFolderGuid();
+		return *this;
 	}
 
 	FFolder GetFolder() const
 	{
-		return FFolder(Path, FFolder::FRootObject(RootObjectPtr.Get()));
+		FFolder::FRootObject RootObject = FFolder::FRootObject(RootObjectPtr.Get());
+		return ActorFolderGuid.IsValid() ? FFolder(RootObject, ActorFolderGuid) : FFolder(RootObject, Path);
 	}
 };
 
diff --git a/Engine/Source/Runtime/Engine/Classes/Engine/Level.h b/Engine/Source/Runtime/Engine/Classes/Engine/Level.h
--- a/Engine/Source/Runtime/Engine/Classes/Engine/Level.h
+++ b/Engine/Source/Runtime/Engine/Classes/Engine/Level.h
@@ -40,6 +40,7 @@ class SNotificationItem;
 class UActorFolder;
 class IWorldPartitionCell;
 class UWorldPartitionRuntimeCell;
+struct FFolder;
 
 UINTERFACE()
 class ULevelPartitionInterface : public UInterface
@@ -71,8 +72,11 @@ struct ENGINE_API FLevelActorFoldersHelper
 private:
 	static void SetUseActorFolders(ULevel* InLevel, bool bInEnabled);
 	static void AddActorFolder(ULevel* InLevel, UActorFolder* InActorFolder, bool bInShouldDirtyLevel, bool bInShouldBroadcast = true);
+	static void RenameFolder(ULevel* InLevel, const FFolder& InOldFolder, const FFolder& InNewFolder);
+	static void DeleteFolder(ULevel* InLevel, const FFolder& InFolder);
 
 	friend class UWorld;
+	friend class ULevel;
 	friend class UActorFolder;
 	friend class UWorldPartitionConvertCommandlet;
 	friend class FWorldPartitionLevelHelper;
@@ -108,6 +112,24 @@ public:
 	TMap<FName, TObjectPtr<AActor>> Actors;
 };
 
+USTRUCT()
+struct ENGINE_API FActorFolderSet
+{
+	GENERATED_BODY()
+
+public:
+
+	void Add(UActorFolder* InActorFolder) { ActorFolders.Add(InActorFolder); }
+	int32 Remove(UActorFolder* InActorFolder) { return ActorFolders.Remove(InActorFolder); }
+	bool IsEmpty() const { return ActorFolders.IsEmpty(); }
+	const TSet<TObjectPtr<UActorFolder>>& GetActorFolders() const { return ActorFolders; }
+
+private:
+
+	UPROPERTY(Transient)
+	TSet<TObjectPtr<UActorFolder>> ActorFolders;
+};
+
 /**
  * Structure containing all information needed for determining the screen space
  * size of an object/ texture instance.
@@ -837,6 +859,10 @@ private:
 	UPROPERTY(Transient)
 	TMap<FGuid, TObjectPtr<UActorFolder>> ActorFolders;
 
+	/** Acceleration table used to find an ActorFolder object for a given folder path. */
+	UPROPERTY(Transient)
+	TMap<FString, FActorFolderSet> FolderLabelToActorFolders;
+
 	/** Temporary array containing actor folder objects manually loaded from their external packages (only used while loading the level). */
 	UPROPERTY(Transient)
 	TArray<TObjectPtr<UActorFolder>> LoadedExternalActorFolders;
@@ -1204,7 +1230,7 @@ public:
 	ENGINE_API UActorFolder* GetActorFolder(const FGuid& InGuid, bool bSkipDeleted = true) const;
 
 	/** Finds the level actor folder by its path. Returns null if not found. */
-	ENGINE_API UActorFolder* GetActorFolder(const FName& InPath, bool bSkipDeleted = true) const;
+	ENGINE_API UActorFolder* GetActorFolder(const FName& InPath) const;
 
 	/** Iterates on all valid level actor folders. */
 	ENGINE_API void ForEachActorFolder(TFunctionRef<bool(UActorFolder*)> Operation, bool bSkipDeleted = false);
@@ -1361,11 +1387,17 @@ private:
 	/** Prepares/fixes actor folder objects once level is fully loaded. */
 	void FixupActorFolders();
 
+	void AddActorFolder(UActorFolder* InActorFolder);
+	void RemoveActorFolder(UActorFolder* InActorFolder);
+	void OnFolderMarkAsDeleted(UActorFolder* InActorFolder);
+	void OnFolderLabelChanged(UActorFolder* InActorFolder, const FString& InOldFolderLabel);
+
 	/** Sets the level to use or not the actor folder objects feature. */
 	void SetUseActorFoldersInternal(bool bInEnabled);
 
 	friend struct FLevelActorFoldersHelper;
 	friend class FWorldPartitionLevelHelper;
+	friend class UActorFolder;
 
 	/** Replace the existing LSA (if set) by spawning a new one based on this level's script blueprint */
 	void RegenerateLevelScriptActor();
diff --git a/Engine/Source/Runtime/Engine/Classes/GameFramework/Actor.h b/Engine/Source/Runtime/Engine/Classes/GameFramework/Actor.h
--- a/Engine/Source/Runtime/Engine/Classes/GameFramework/Actor.h
+++ b/Engine/Source/Runtime/Engine/Classes/GameFramework/Actor.h
@@ -976,7 +976,7 @@ private:
 	FName FolderPath;
 
 	/** If the actor's level uses the actor folder objects feature, contains the actor folder unique identifier (invalid=root). */
-	UPROPERTY()
+	UPROPERTY(TextExportTransient)
 	FGuid FolderGuid;
 
 public:
diff --git a/Engine/Source/Runtime/Engine/Private/ActorEditor.cpp b/Engine/Source/Runtime/Engine/Private/ActorEditor.cpp
--- a/Engine/Source/Runtime/Engine/Private/ActorEditor.cpp
+++ b/Engine/Source/Runtime/Engine/Private/ActorEditor.cpp
@@ -975,12 +975,26 @@ void AActor::ClearActorLabel()
 
 FFolder AActor::GetFolder() const
 {
-	return FFolder(GetFolderPath(), GetFolderRootObject());
+	// Favor building FFolder using guid (if available).
+	// The reason is that a lot of calling functions will try resolving at some point the UActorFolder* from the folder
+	// and the guid implementation is much faster.
+	if (GetWorld() && GetFolderGuid().IsValid())
+	{
+		if (UActorFolder* ActorFolder = GetActorFolder())
+		{
+			return ActorFolder->GetFolder();
+		}
+	}
+	return FFolder(GetFolderRootObject(), GetFolderPath());
 }
 
 FFolder::FRootObject AActor::GetFolderRootObject() const
 {
-	return FFolder::GetOptionalFolderRootObject(GetLevel()).Get(FFolder::GetDefaultRootObject());
+	if (GetWorld())
+	{
+		return FFolder::GetOptionalFolderRootObject(GetLevel()).Get(FFolder::GetInvalidRootObject());
+	}
+	return FFolder::GetInvalidRootObject();
 }
 
 static bool IsUsingActorFolders(const AActor* InActor)
@@ -1023,7 +1037,7 @@ bool AActor::CreateOrUpdateActorFolder()
 	if (!ActorFolder)
 	{
 		check(!FolderPath.IsNone());
-		ActorFolder = FWorldPersistentFolders::GetActorFolder(FFolder(FolderPath, GetFolderRootObject()), GetWorld(), /*bAllowCreate*/ true);
+		ActorFolder = FWorldPersistentFolders::GetActorFolder(FFolder(GetFolderRootObject(), FolderPath), GetWorld(), /*bAllowCreate*/ true);
 	}
 
 	// At this point, actor folder should always be valid
@@ -1048,7 +1062,7 @@ UActorFolder* AActor::GetActorFolder(bool bSkipDeleted) const
 		}
 		else if (!FolderPath.IsNone())
 		{
-			ActorFolder = Level->GetActorFolder(FolderPath, bSkipDeleted);
+			ActorFolder = Level->GetActorFolder(FolderPath);
 		}
 	}
 	return ActorFolder;
@@ -1108,13 +1122,14 @@ FGuid AActor::GetFolderGuid() const
 FName AActor::GetFolderPath() const
 {
 	static const FName RootPath = FFolder::GetEmptyPath();
-	if (!FFolder::GetOptionalFolderRootObject(GetLevel()))
+	UWorld* World = GetWorld();
+	if (World && !FFolder::GetOptionalFolderRootObject(GetLevel()))
 	{
 		return RootPath;
 	}
 	if (IsUsingActorFolders(this))
 	{
-		if (UActorFolder* ActorFolder = GetActorFolder())
+		if (UActorFolder* ActorFolder = World ? GetActorFolder() : nullptr)
 		{
 			return ActorFolder->GetPath();
 		}
@@ -1131,7 +1146,7 @@ void AActor::SetFolderPath(const FName& InNewFolderPath)
 		UWorld* World = GetWorld();
 		if (!InNewFolderPath.IsNone() && World)
 		{
-			FFolder NewFolder(InNewFolderPath, GetFolderRootObject());
+			FFolder NewFolder(GetFolderRootObject(), InNewFolderPath);
 			ActorFolder = FWorldPersistentFolders::GetActorFolder(NewFolder, World);
 			if (!ActorFolder)
 			{
diff --git a/Engine/Source/Runtime/Engine/Private/ActorFolder.cpp b/Engine/Source/Runtime/Engine/Private/ActorFolder.cpp
--- a/Engine/Source/Runtime/Engine/Private/ActorFolder.cpp
+++ b/Engine/Source/Runtime/Engine/Private/ActorFolder.cpp
@@ -120,7 +120,9 @@ void UActorFolder::SetLabel(const FString& InFolderLabel)
 	if (!FolderLabel.Equals(InFolderLabel, ESearchCase::CaseSensitive))
 	{
 		Modify();
+		FString OldFolderLabel = FolderLabel;
 		FolderLabel = InFolderLabel;
+		GetTypedOuter<ULevel>()->OnFolderLabelChanged(this, OldFolderLabel);
 	}
 }
 
@@ -156,11 +158,65 @@ void UActorFolder::MarkAsDeleted()
 {
 	Modify();
 
+	ULevel* Level = GetTypedOuter<ULevel>();
+	
+	auto HasParent = [](UActorFolder* InFolder, UActorFolder* InParent)
+	{
+		UActorFolder* Parent = InFolder->GetParent(/*bSkipDeleted*/false);
+		while (Parent)
+		{
+			if (Parent == InParent)
+			{
+				return true;
+			}
+			Parent = Parent->GetParent(/*bSkipDeleted*/false);
+		}
+		return false;
+	};
+
+	TArray<UActorFolder*> FoldersToDelete;
+	TMap<UActorFolder*, UActorFolder*> DuplicateFolders;
+	const FName PathToDelete = GetPath();
+	// Find all SubPaths of PathToDelete
+	Level->ForEachActorFolder([this, Level, PathToDelete, &HasParent, &DuplicateFolders, &FoldersToDelete](UActorFolder* ActorFolder)
+	{
+		if (HasParent(ActorFolder, this))
+		{
+			// Get child folder new path if parent is deleted
+			FName ChildFolderNewPath = ActorFolder->GetPathInternal(this);
+			if (UActorFolder* ExistingFolder = Level->GetActorFolder(ChildFolderNewPath))
+			{
+				DuplicateFolders.Add(ActorFolder, ExistingFolder);
+				FoldersToDelete.Add(ActorFolder);
+			}
+		}
+		return true;
+	}, /*bSkipDeleted*/ false);
+
+	// Sort in descending order so children will be deleted before parents
+	FoldersToDelete.Sort([](const UActorFolder& FolderA, const UActorFolder& FolderB)
+	{
+		return FolderB.GetPath().LexicalLess(FolderA.GetPath());
+	});
+	
+	for (UActorFolder* FolderToDelete : FoldersToDelete)
+	{
+		UActorFolder* NewParent = DuplicateFolders.FindChecked(FolderToDelete);
+		// First move duplicate folder under the single folder we keep. Use a unique name to avoid dealing with name clash
+		const FFolder OldFolder = FolderToDelete->GetFolder();
+		const FString NewPath = FString::Printf(TEXT("%s/%s_Duplicate_%s"), *NewParent->GetPath().ToString(), *FolderToDelete->GetLabel(), *FGuid::NewGuid().ToString());
+		const FFolder NewFolder = FFolder(OldFolder.GetRootObject(), FName(NewPath));
+		FLevelActorFoldersHelper::RenameFolder(Level, OldFolder, NewFolder);
+		// Then delete (mark as deleted) this folder
+		FLevelActorFoldersHelper::DeleteFolder(Level, FolderToDelete->GetFolder());
+	}
+
 	// Deleting a folder must not modify actors part of it nor sub folders.
 	// Here, we simply mark the folder as deleted. 
 	// When marked as deleted, the folder will act as a redirector to its parent.
 	check(!bIsDeleted);
 	bIsDeleted = true;
+	Level->OnFolderMarkAsDeleted(this);
 }
 
 UActorFolder* UActorFolder::GetParent(bool bSkipDeleted) const
@@ -170,23 +226,37 @@ UActorFolder* UActorFolder::GetParent(bool bSkipDeleted) const
 
 FName UActorFolder::GetPath() const
 {
-	if (!FFolder::GetOptionalFolderRootObject(GetOuterULevel()))
+	return GetPathInternal(nullptr);
+}
+
+FORCEINLINE void BuildPath(TStringBuilder<512>& OutStringBuilder, const UActorFolder* InSkipFolder, const UActorFolder* InCurrentFolder)
+{
+	if (InCurrentFolder)
 	{
-		return NAME_None;
+		BuildPath(OutStringBuilder, InSkipFolder, InCurrentFolder->GetParent());
+		if (InCurrentFolder != InSkipFolder)
+		{
+			if (OutStringBuilder.Len())
+			{
+				OutStringBuilder += TEXT("/");
+			}
+			OutStringBuilder += InCurrentFolder->GetLabel();
+		}
 	}
+};
 
-	TStringBuilder<1024> StringBuilder;
+FName UActorFolder::GetPathInternal(UActorFolder* InSkipFolder) const
+{
+	TStringBuilder<512> StringBuilder;
+	BuildPath(StringBuilder, InSkipFolder, GetParent());
 	if (IsValid())
 	{
+		if (StringBuilder.Len())
+		{
+			StringBuilder += TEXT("/");
+		}
 		StringBuilder += FolderLabel;
 	}
-	UActorFolder* Parent = GetParent();
-	while (Parent)
-	{
-		StringBuilder.Prepend(TEXT("/"));
-		StringBuilder.Prepend(Parent->FolderLabel);
-		Parent = Parent->GetParent();
-	}
 	return FName(*StringBuilder);
 }
 
@@ -216,7 +286,24 @@ void UActorFolder::Fixup()
 
 FFolder UActorFolder::GetFolder() const
 {
-	return FFolder(GetPath(), FFolder::GetOptionalFolderRootObject(GetOuterULevel()).Get(FFolder::GetDefaultRootObject()));
+	// Resolve FFolder::FRootObject for this ActorFolder using its outer level
+	ULevel* OuterLevel = GetOuterULevel();
+	check(OuterLevel);
+	check(FolderGuid.IsValid());
+	FFolder::FRootObject RootObject = FFolder::GetOptionalFolderRootObject(OuterLevel).Get(FFolder::GetInvalidRootObject());
+	
+	// Detect case where returned root object is different from outer level (this is the case for UWorldPartitionLevelStreamingDynamic::GetFolderRootObject).
+	ULevel* RootObjectAssociatedLevel = FFolder::GetRootObjectAssociatedLevel(RootObject);
+	if (RootObjectAssociatedLevel != OuterLevel)
+	{
+		// Build and return a FFolder using the root object and ActorFolder's path
+		return FFolder(RootObject, GetPath());
+	}
+
+	// Build and return a FFolder using the root object and ActorFolder's guid
+	FFolder Folder = FFolder(RootObject, FolderGuid);
+	check(Folder.GetActorFolder() == this);
+	return Folder;
 }
 
 void UActorFolder::SetPackageExternal(bool bInExternal, bool bShouldDirty)
diff --git a/Engine/Source/Runtime/Engine/Private/Folder.cpp b/Engine/Source/Runtime/Engine/Private/Folder.cpp
--- a/Engine/Source/Runtime/Engine/Private/Folder.cpp
+++ b/Engine/Source/Runtime/Engine/Private/Folder.cpp
@@ -3,15 +3,164 @@
 #include "Folder.h"
 #include "Engine/Level.h"
 #include "Engine/LevelStreaming.h"
+#include "Engine/World.h"
+#include "LevelInstance/LevelInstanceInterface.h"
+#include "LevelInstance/LevelInstanceSubsystem.h"
+#include "ActorFolder.h"
 #include "LevelUtils.h"
 
 #if WITH_EDITOR
+
+/*
+ * FFolder Helpers methods
+ */
+
+const FFolder& FFolder::GetInvalidFolder()
+{
+	static FFolder InvalidFolder;
+	return InvalidFolder;
+}
+
+const FFolder::FRootObject& FFolder::GetInvalidRootObject()
+{
+	static FRootObject InvalidRootObject;
+	return InvalidRootObject;
+}
+
+bool FFolder::IsRootObjectPersistentLevel(const FRootObject& Key)
+{
+	ULevel* Level = GetRootObjectAssociatedLevel(Key);
+	return Level && Level->IsPersistentLevel();
+}
+
 TOptional<FFolder::FRootObject> FFolder::GetOptionalFolderRootObject(const ULevel* InLevel)
 {
-	if (ULevelStreaming* LevelStreaming = InLevel ? ULevelStreaming::FindStreamingLevel(InLevel) : nullptr)
+	TOptional<FFolder::FRootObject> Result;
+	check(InLevel);
+	if (InLevel)
 	{
-		return LevelStreaming->GetFolderRootObject();
+		if (InLevel->IsPersistentLevel())
+		{
+			Result = FFolder::FRootObject(InLevel);
+		}
+		else if (ULevelStreaming* LevelStreaming = ULevelStreaming::FindStreamingLevel(InLevel))
+		{
+			Result = LevelStreaming->GetFolderRootObject();
+		}
 	}
-	return FFolder::GetDefaultRootObject();
+	return Result;
 }
+
+ULevel* FFolder::GetRootObjectAssociatedLevel(const FRootObject& Key)
+{
+	if (const UObject* RootObjectPtr = FFolder::GetRootObjectPtr(Key))
+	{
+		if (const UWorld* WorldPtr = Cast<UWorld>(RootObjectPtr))
+		{
+			return WorldPtr->PersistentLevel;
+		}
+		else if (const ILevelInstanceInterface* LevelInstance = Cast<ILevelInstanceInterface>(RootObjectPtr))
+		{
+			return LevelInstance->GetLoadedLevel();
+		}
+		else if (const ULevel* Level = Cast<ULevel>(RootObjectPtr))
+		{
+			if (UWorld* OuterWorld = Level->GetTypedOuter<UWorld>())
+			{
+				return OuterWorld->PersistentLevel;
+			}
+		}
+	}
+	return nullptr;
+}
+
+FFolder FFolder::GetWorldRootFolder(UWorld* InWorld)
+{
+	if (InWorld)
+	{
+		return FFolder::GetOptionalFolderRootObject(InWorld->PersistentLevel).GetValue();
+	}
+	return FFolder::GetInvalidFolder();
+}
+
+bool FFolder::GetFolderPathsAndCommonRootObject(const TArray<FFolder>& InFolders, TArray<FName>& OutFolders, FFolder::FRootObject& OutCommonRootObject)
+{
+	OutFolders.Reset();
+
+	TOptional<FFolder::FRootObject> CommonRootObject;
+	for (const FFolder& Folder : InFolders)
+	{
+		OutFolders.Add(Folder.GetPath());
+		if (!CommonRootObject.IsSet())
+		{
+			CommonRootObject = Folder.GetRootObject();
+		}
+		else if (CommonRootObject.GetValue() != Folder.GetRootObject())
+		{
+			OutFolders.Reset();
+			CommonRootObject.Reset();
+			break;
+		}
+	};
+
+	OutCommonRootObject = CommonRootObject.Get(FFolder::GetInvalidRootObject());
+	return CommonRootObject.IsSet();
+}
+
+/*
+ * FFolder implementation
+ */
+
+const FName FFolder::GetPath() const
+{
+	UActorFolder* ActorFolder = nullptr;
+	if (!bPathInitialized)
+	{
+		check(Path.IsNone());
+		check(ActorFolderGuid.IsValid());
+		ActorFolder = GetActorFolder();
+		if (ActorFolder)
+		{
+			// Cache Path
+			Path = ActorFolder->GetPath();
+		}
+		bPathInitialized = true;
+	}
+	return Path;
+}
+
+FFolder FFolder::GetParent() const
+{
+	if (UActorFolder* ActorFolder = GetActorFolder())
+	{
+		UActorFolder* ParentActorFolder = ActorFolder->GetParent();
+		return ParentActorFolder ? ParentActorFolder->GetFolder() : FFolder(GetRootObject());
+	}
+
+	const FName ParentPath(*FPaths::GetPath(Path.ToString()));
+	return FFolder(RootObject, ParentPath);
+}
+
+UActorFolder* FFolder::GetActorFolder() const
+{
+	ULevel* Level = GetRootObjectAssociatedLevel();
+	if (Level && Level->IsUsingActorFolders())
+	{
+		UActorFolder* ActorFolder = Level->GetActorFolder(ActorFolderGuid, /*bSkipDeleted*/ false);
+		if (ActorFolder)
+		{
+			return ActorFolder;
+		}
+		ActorFolder = Level->GetActorFolder(Path);
+		if (ActorFolder)
+		{
+			// Cache ActorFolderGuid
+ 			check(!ActorFolderGuid.IsValid());
+			ActorFolderGuid = ActorFolder->GetGuid();
+			return ActorFolder;
+		}
+	}
+	return nullptr;
+}
+
 #endif
\ No newline at end of file
diff --git a/Engine/Source/Runtime/Engine/Private/Level.cpp b/Engine/Source/Runtime/Engine/Private/Level.cpp
--- a/Engine/Source/Runtime/Engine/Private/Level.cpp
+++ b/Engine/Source/Runtime/Engine/Private/Level.cpp
@@ -71,6 +71,7 @@ Level.cpp: Level-related functions
 #include "Framework/Notifications/NotificationManager.h"
 #include "Settings/LevelEditorMiscSettings.h"
 #include "ExternalPackageHelper.h"
+#include "Folder.h"
 #include "ActorFolder.h"
 #include "Misc/MessageDialog.h"
 #include "ScopedTransaction.h"
@@ -115,7 +116,7 @@ void FLevelActorFoldersHelper::AddActorFolder(ULevel* InLevel, UActorFolder* InA
 {
 	InLevel->Modify(bInShouldDirtyLevel);
 	check(InActorFolder->GetGuid().IsValid());
-	InLevel->ActorFolders.Add(InActorFolder->GetGuid(), InActorFolder);
+	InLevel->AddActorFolder(InActorFolder);
 
 	if (bInShouldBroadcast)
 	{
@@ -123,6 +124,47 @@ void FLevelActorFoldersHelper::AddActorFolder(ULevel* InLevel, UActorFolder* InA
 	}
 }
 
+void FLevelActorFoldersHelper::RenameFolder(ULevel* InLevel, const FFolder& InOldFolder, const FFolder& InNewFolder)
+{
+	check(InLevel);
+	// This implementation can be called both if FActorFolders is or isn't initialized.
+	if (FActorFolders::Get().IsInitializedForWorld(*InLevel->GetWorld()))
+	{
+		FActorFolders::Get().RenameFolderInWorld(*InLevel->GetWorld(), InOldFolder, InNewFolder);
+	}
+	else
+	{
+		UActorFolder* ActorFolder = InOldFolder.GetActorFolder();
+		check(::IsValid(ActorFolder));
+		UActorFolder* FoundFolder = InNewFolder.GetActorFolder();
+		check(!::IsValid(FoundFolder) || !FoundFolder->GetPath().IsEqual(InNewFolder.GetPath(), ENameCase::CaseSensitive));
+
+		UActorFolder* ParentActorFolder = InNewFolder.GetParent().GetActorFolder();
+		ActorFolder->SetParent(ParentActorFolder);
+		const FString FolderLabel = InNewFolder.GetLeafName().ToString();
+		ActorFolder->SetLabel(FolderLabel);
+		check(ActorFolder->GetPath().IsEqual(InNewFolder.GetPath(), ENameCase::CaseSensitive));
+	}
+};
+
+void FLevelActorFoldersHelper::DeleteFolder(ULevel* InLevel, const FFolder& InFolder)
+{
+	check(InLevel);
+	// This implementation can be called both if FActorFolders is or isn't initialized.
+	if (FActorFolders::Get().IsInitializedForWorld(*InLevel->GetWorld()))
+	{
+		FActorFolders::Get().DeleteFolder(*InLevel->GetWorld(), InFolder);
+	}
+	else
+	{
+		UActorFolder* ActorFolder = InFolder.GetActorFolder();
+		if (::IsValid(ActorFolder))
+		{
+			ActorFolder->MarkAsDeleted();
+		}
+	}
+};
+
 FLevelPartitionOperationScope::FLevelPartitionOperationScope(ULevel* InLevel)
 {
 	InterfacePtr = InLevel->GetLevelPartition();
@@ -2451,6 +2493,75 @@ void ULevel::CommitModelSurfaces()
 }
 
 #if WITH_EDITOR
+void ULevel::AddActorFolder(UActorFolder* InActorFolder)
+{
+	Modify(false);
+	check(InActorFolder);
+	const FGuid& FolderGuid = InActorFolder->GetGuid();
+	check(FolderGuid.IsValid());
+	check(!ActorFolders.Contains(FolderGuid));
+	ActorFolders.Add(FolderGuid, InActorFolder);
+	if (!InActorFolder->IsMarkedAsDeleted())
+	{
+		FActorFolderSet& Folders = FolderLabelToActorFolders.FindOrAdd(InActorFolder->GetLabel());
+		Folders.Add(InActorFolder);
+	}
+}
+
+void ULevel::RemoveActorFolder(UActorFolder* InActorFolder)
+{
+	Modify(false);
+	check(InActorFolder);
+	TObjectPtr<UActorFolder> FoundActorFolder;
+	if (ensure(ActorFolders.RemoveAndCopyValue(InActorFolder->GetGuid(), FoundActorFolder)))
+	{
+		if (!InActorFolder->IsMarkedAsDeleted())
+		{
+			FActorFolderSet& Folders = FolderLabelToActorFolders.FindChecked(FoundActorFolder->GetLabel());
+			verify(Folders.Remove(FoundActorFolder));
+		}
+	}
+}
+
+void ULevel::OnFolderMarkAsDeleted(UActorFolder* InActorFolder)
+{
+	check(InActorFolder);
+	check(InActorFolder->IsMarkedAsDeleted());
+	if (ensure(ActorFolders.Contains(InActorFolder->GetGuid())))
+	{
+		Modify(false);
+		FActorFolderSet& Folders = FolderLabelToActorFolders.FindChecked(InActorFolder->GetLabel());
+		verify(Folders.Remove(InActorFolder));
+		if (Folders.IsEmpty())
+		{
+			FolderLabelToActorFolders.Remove(InActorFolder->GetLabel());
+		}
+	}
+}
+
+void ULevel::OnFolderLabelChanged(UActorFolder* InActorFolder, const FString& InOldFolderLabel)
+{
+	check(InActorFolder);
+	if (InOldFolderLabel.IsEmpty())
+	{
+		// We are in the process of creating the actor folder
+		check(!ActorFolders.Contains(InActorFolder->GetGuid()));
+		return;
+	}
+	if (ensure(ActorFolders.Contains(InActorFolder->GetGuid())))
+	{
+		Modify(false);
+		FActorFolderSet& OldLabelFolders = FolderLabelToActorFolders.FindChecked(InOldFolderLabel);
+		verify(OldLabelFolders.Remove(InActorFolder));
+		if (OldLabelFolders.IsEmpty())
+		{
+			FolderLabelToActorFolders.Remove(InOldFolderLabel);
+		}
+		FActorFolderSet& NewLabelFolders = FolderLabelToActorFolders.FindOrAdd(InActorFolder->GetLabel());
+		NewLabelFolders.Add(InActorFolder);
+	}
+}
+
 void ULevel::FixupActorFolders()
 {
 	if (!IsActorFolderObjectsFeatureAvailable())
@@ -2464,39 +2575,66 @@ void ULevel::FixupActorFolders()
 		for (UActorFolder* LoadedActorFolder : LoadedExternalActorFolders)
 		{
 			check(LoadedActorFolder->GetGuid().IsValid());
-			ActorFolders.Add(LoadedActorFolder->GetGuid(), LoadedActorFolder);
+			AddActorFolder(LoadedActorFolder);
 		}
 		LoadedExternalActorFolders.Empty();
+
+		const bool bFixDuplicateFolders = !IsRunningCommandlet();
 	
 		// Discover duplicate paths first to prioritize non-duplicate paths
-		TSet<FName> FolderPaths;
-		TArray<UActorFolder*> DuplicateFolders;
-		ForEachActorFolder([&FolderPaths, &DuplicateFolders](UActorFolder* ActorFolder)
+		TMap<FName, UActorFolder*> ExistingPaths;
+		TArray<UActorFolder*> FoldersToDelete;
+		TMap<UActorFolder*, UActorFolder*> DuplicateFolders;
+		ForEachActorFolder([bFixDuplicateFolders, &ExistingPaths, &DuplicateFolders, &FoldersToDelete](UActorFolder* ActorFolder)
 		{
 			// Detects and clears invalid parent folder
 			ActorFolder->FixupParentFolder();
 
-			bool bIsAlreadyInSet = false;
-			FolderPaths.Add(ActorFolder->GetPath(), &bIsAlreadyInSet);
-			if (bIsAlreadyInSet)
+			if (bFixDuplicateFolders)
 			{
-				DuplicateFolders.Add(ActorFolder);
+				FName FolderPath = ActorFolder->GetPath();
+				UActorFolder** ExistingFolder = ExistingPaths.Find(FolderPath);
+				if (!ExistingFolder)
+				{
+					ExistingPaths.Add(FolderPath, ActorFolder);
+				}
+				else
+				{
+					FoldersToDelete.Add(ActorFolder);
+					DuplicateFolders.Add(ActorFolder, *ExistingFolder);
+				}
 			}
 			return true;
 		}, /*bSkipDeleted*/ true);
 
-		if (!IsRunningCommandlet())
+		if (bFixDuplicateFolders)
 		{
-			// Rename duplicates to a new valid/unique name
-			for (UActorFolder* ActorFolder : DuplicateFolders)
+			// Sort in descending order so children will be deleted before parents
+			FoldersToDelete.Sort([](const UActorFolder& FolderA, const UActorFolder& FolderB)
+			{
+				return FolderB.GetPath().LexicalLess(FolderA.GetPath());
+			});
+
+			for (UActorFolder* FolderToDelete : FoldersToDelete)
 			{
-				FFolder Folder = ActorFolder->GetFolder();
-				const FFolder NewPath = FActorFolders::Get().GetFolderName(*GetWorld(), Folder.GetParent(), Folder.GetLeafName());
-				ActorFolder->SetLabel(NewPath.GetLeafName().ToString());
-				bool bIsAlreadyInSet = false;
-				FolderPaths.Add(ActorFolder->GetPath(), &bIsAlreadyInSet);
-				check(!bIsAlreadyInSet);
-				UE_LOG(LogLevel, Warning, TEXT("Found duplicate actor folder %s, renamed to %s."), *Folder.GetPath().ToString(), *NewPath.GetPath().ToString());
+				UActorFolder* NewParent = DuplicateFolders.FindChecked(FolderToDelete);
+				// First move duplicate folder under the single folder we keep. Use a unique name to avoid dealing with name clash
+				const FFolder OldFolder = FolderToDelete->GetFolder();
+				UE_LOG(LogLevel, Log, TEXT("Merging duplicate actor folder %s."), *OldFolder.GetPath().ToString());
+
+				// Since we can't rename to a dest with the same parent hierarchy, do it in 2 passes.
+				// For example: If we want to rename A/B to A/B/B_Dup123, we need to :
+				// 1- Rename A/B to A/B_Dup123
+				// 2- Rename A/B_Dup123 to A/B/B_Dup123
+				// Then we mark for delete A/B/B_Dup123, so that child actors and folders of A/B/B_Dup123 will be parented back to in A/B.
+				const FString NewPath = FString::Printf(TEXT("%s_%s"), *NewParent->GetPath().ToString(), *FGuid::NewGuid().ToString());
+				const FFolder NewFolder = FFolder(OldFolder.GetRootObject(), FName(NewPath));
+				FLevelActorFoldersHelper::RenameFolder(this, OldFolder, NewFolder);
+				const FString NewPath2 = FString::Printf(TEXT("%s/DuplicateFolder_%s"), *NewParent->GetPath().ToString(), *FolderToDelete->GetLabel(), *FGuid::NewGuid().ToString());
+				const FFolder NewFolder2 = FFolder(OldFolder.GetRootObject(), FName(NewPath2));
+				FLevelActorFoldersHelper::RenameFolder(this, NewFolder, NewFolder2);
+				// Then delete (mark as deleted) this folder
+				FLevelActorFoldersHelper::DeleteFolder(this, FolderToDelete->GetFolder());
 			}
 		}
 	}
@@ -3118,27 +3256,44 @@ bool ULevel::IsUsingExternalObjects() const
 	return IsUsingExternalActors();
 }
 
-UActorFolder* ULevel::GetActorFolder(const FName& InPath, bool bSkipDeleted) const
+static UActorFolder* FindNextFolder(const TMap<FString, FActorFolderSet>& InFolderLabelToActorFolders, const TArray<FString>& InFolderLabels, int32 Index, UActorFolder* ParentFolder)
 {
-	if (!InPath.IsNone() && IsUsingActorFolders())
+	if (const FActorFolderSet* FoundSet = InFolderLabelToActorFolders.Find(InFolderLabels[Index]))
 	{
-		// @todo_ow : Add an acceleration table
-		for (const auto& Pair : ActorFolders)
+		for (const TObjectPtr<UActorFolder>& ActorFolder : FoundSet->GetActorFolders())
 		{
-			UActorFolder* ActorFolder = Pair.Value;
-			if (ActorFolder->GetPath() == InPath)
+			if (ActorFolder->GetParent() == ParentFolder)
 			{
-				if (bSkipDeleted && ActorFolder->IsMarkedAsDeleted())
+				int32 NextIndex = Index + 1;
+				check(NextIndex <= InFolderLabels.Num());
+
+				if (NextIndex == InFolderLabels.Num())
+				{
+					return ActorFolder;
+				}
+				if (UActorFolder* Found = FindNextFolder(InFolderLabelToActorFolders, InFolderLabels, NextIndex, ActorFolder))
 				{
-					return ActorFolder->GetParent();
+					return Found;
 				}
-				return ActorFolder;
 			}
 		}
 	}
 	return nullptr;
 }
 
+UActorFolder* ULevel::GetActorFolder(const FName& InPath) const
+{
+	if (!InPath.IsNone() && IsUsingActorFolders())
+	{
+		TArray<FString> FolderLabels;
+		FString Path = FPaths::RemoveDuplicateSlashes(InPath.ToString());
+		Path.ParseIntoArray(FolderLabels, TEXT("/"));
+		UActorFolder* CurrentFolder = (FolderLabels.Num() > 0) ? FindNextFolder(FolderLabelToActorFolders, FolderLabels, 0, nullptr) : nullptr;
+		return CurrentFolder;
+	}
+	return nullptr;
+}
+
 UActorFolder* ULevel::GetActorFolder(const FGuid& InGuid, bool bSkipDeleted) const
 {
 	if (IsUsingActorFolders() && InGuid.IsValid())
@@ -3220,8 +3375,10 @@ bool ULevel::SetUseActorFolders(bool bInEnabled, bool bInInteractiveMode)
 		}
 	}
 
-	FScopedTransaction Transaction(LOCTEXT("ChangeUseActorFolders", "Change Use Actor Folder Objects"));
 	SetUseActorFoldersInternal(bInEnabled);
+	// Operation cannot be undone
+	GEditor->ResetTransaction(LOCTEXT("LevelUseActorFolderObjectsResetTrans", "Level Use Actor Folder Objects"));
+
 	return true;
 }
 
@@ -3275,7 +3432,7 @@ void ULevel::CreateOrUpdateActorFolders()
 	{
 		GEngine->BroadcastActorFolderRemoved(ActorFolderToDelete);
 
-		ActorFolders.Remove(ActorFolderToDelete->GetGuid());
+		RemoveActorFolder(ActorFolderToDelete);
 	}
 
 	// Avoid broadcasting if no actor folder were/are part of this level
diff --git a/Engine/Source/Runtime/Engine/Private/LevelInstance/LevelInstanceEditorLevelStreaming.cpp b/Engine/Source/Runtime/Engine/Private/LevelInstance/LevelInstanceEditorLevelStreaming.cpp
--- a/Engine/Source/Runtime/Engine/Private/LevelInstance/LevelInstanceEditorLevelStreaming.cpp
+++ b/Engine/Source/Runtime/Engine/Private/LevelInstance/LevelInstanceEditorLevelStreaming.cpp
@@ -44,6 +44,13 @@ TOptional<FFolder::FRootObject> ULevelStreamingLevelInstanceEditor::GetFolderRoo
 		}
 	}
 
+	if (!LevelInstanceID.IsValid())
+	{
+		// When creating a new level instance, we have an invalid LevelInstanceID (until it gets loaded)
+		// Return the world as root object.
+		return FFolder::GetWorldRootFolder(GetWorld()).GetRootObject();
+	}
+
 	return TOptional<FFolder::FRootObject>();
 }
 
diff --git a/Engine/Source/Runtime/Engine/Private/LevelInstance/LevelInstanceLevelStreaming.cpp b/Engine/Source/Runtime/Engine/Private/LevelInstance/LevelInstanceLevelStreaming.cpp
--- a/Engine/Source/Runtime/Engine/Private/LevelInstance/LevelInstanceLevelStreaming.cpp
+++ b/Engine/Source/Runtime/Engine/Private/LevelInstance/LevelInstanceLevelStreaming.cpp
@@ -45,8 +45,9 @@ TOptional<FFolder::FRootObject> ULevelStreamingLevelInstance::GetFolderRootObjec
 			return FFolder::FRootObject(Actor);
 		}
 	}
-
-	return TOptional<FFolder::FRootObject>();
+	// When LevelInstance is null, it's because the level instance is being streamed-out. 
+	// Return the world as root object.
+	return FFolder::GetWorldRootFolder(GetWorld()).GetRootObject();
 }
 
 FBox ULevelStreamingLevelInstance::GetBounds() const
diff --git a/Engine/Source/Runtime/Engine/Private/LevelStreaming.cpp b/Engine/Source/Runtime/Engine/Private/LevelStreaming.cpp
--- a/Engine/Source/Runtime/Engine/Private/LevelStreaming.cpp
+++ b/Engine/Source/Runtime/Engine/Private/LevelStreaming.cpp
@@ -1658,15 +1658,15 @@ void ULevelStreaming::BroadcastLevelVisibleStatus(UWorld* PersistentWorld, FName
 
 	for (ULevelStreaming* StreamingLevel : LevelsToBroadcast)
 	{
-			if (bVisible)
-			{
-				StreamingLevel->OnLevelShown.Broadcast();
-			}
-			else
-			{
-				StreamingLevel->OnLevelHidden.Broadcast();
-			}
+		if (bVisible)
+		{
+			StreamingLevel->OnLevelShown.Broadcast();
+		}
+		else
+		{
+			StreamingLevel->OnLevelHidden.Broadcast();
 		}
+	}
 }
 
 void ULevelStreaming::SetWorldAsset(const TSoftObjectPtr<UWorld>& NewWorldAsset)
@@ -2038,13 +2038,18 @@ void ULevelStreaming::PostEditUndo()
 TOptional<FFolder::FRootObject> ULevelStreaming::GetFolderRootObject() const
 { 
 	// We consider that if either the loaded level or its world persistent level uses actor folder objects, the loaded level is the folder root object.
-	if (LoadedLevel && 
-		(LoadedLevel != LoadedLevel->GetWorld()->PersistentLevel) && 
-		(LoadedLevel->IsUsingActorFolders() || LoadedLevel->GetWorld()->PersistentLevel->IsUsingActorFolders()))
+	if (LoadedLevel)
 	{
-		return FFolder::FRootObject(LoadedLevel);
+		if (LoadedLevel->IsUsingActorFolders() || LoadedLevel->GetWorld()->PersistentLevel->IsUsingActorFolders())
+		{
+			return FFolder::FRootObject(LoadedLevel);
+		}
+		else
+		{
+			return FFolder::GetWorldRootFolder(LoadedLevel->GetWorld()).GetRootObject();
+		}
 	}
-	return FFolder::GetDefaultRootObject();
+	return TOptional<FFolder::FRootObject>();
 }
 
 const FName& ULevelStreaming::GetFolderPath() const
diff --git a/Engine/Source/Runtime/Engine/Private/WorldPartition/WorldPartitionLevelStreamingDynamic.cpp b/Engine/Source/Runtime/Engine/Private/WorldPartition/WorldPartitionLevelStreamingDynamic.cpp
--- a/Engine/Source/Runtime/Engine/Private/WorldPartition/WorldPartitionLevelStreamingDynamic.cpp
+++ b/Engine/Source/Runtime/Engine/Private/WorldPartition/WorldPartitionLevelStreamingDynamic.cpp
@@ -465,7 +465,11 @@ void UWorldPartitionLevelStreamingDynamic::OnCleanupLevel()
 // This could become an option in the world outliner when running PIE.
 TOptional<FFolder::FRootObject> UWorldPartitionLevelStreamingDynamic::GetFolderRootObject() const
 {
-	return FFolder::GetDefaultRootObject();
+	if (UWorld* World = GetWorld())
+	{
+		return FFolder::GetWorldRootFolder(World).GetRootObject();
+	}
+	return FFolder::GetInvalidRootObject();
 }
 
 #endif
diff --git a/Engine/Source/Runtime/Engine/Public/ActorFolder.h b/Engine/Source/Runtime/Engine/Public/ActorFolder.h
--- a/Engine/Source/Runtime/Engine/Public/ActorFolder.h
+++ b/Engine/Source/Runtime/Engine/Public/ActorFolder.h
@@ -36,7 +36,7 @@ public:
 	UActorFolder* GetParent(bool bSkipDeleted = true) const;
 
 	void SetLabel(const FString& InFolderLabel);
-	FString GetLabel() const { return FolderLabel; }
+	const FString& GetLabel() const { return FolderLabel; }
 
 	// Control whether the folder is initially expanded or not
 	void SetIsInitiallyExpanded(bool bInFolderInitiallyExpanded);
@@ -65,10 +65,13 @@ public:
 	* @param bShouldDirty should dirty or not the level package
 	*/
 	void SetPackageExternal(bool bExternal, bool bShouldDirty = true);
-#endif
 
 private:
+	FName GetPathInternal(UActorFolder* InSkipFolder) const;
+#endif
+
 #if WITH_EDITORONLY_DATA
+private:
 	UPROPERTY()
 	FGuid ParentFolderGuid;
 
diff --git a/Engine/Source/Runtime/Engine/Public/Folder.h b/Engine/Source/Runtime/Engine/Public/Folder.h
--- a/Engine/Source/Runtime/Engine/Public/Folder.h
+++ b/Engine/Source/Runtime/Engine/Public/Folder.h
@@ -10,22 +10,44 @@
 #include "Misc/Optional.h"
 
 class ULevel;
+class UActorFolder;
+class UWorld;
 
-struct FFolder
+struct ENGINE_API FFolder
 {
 #if WITH_EDITOR
 	typedef FObjectKey FRootObject;
 
-	FFolder(const FName& InPath = GetEmptyPath(), const FRootObject& InRootObject = GetDefaultRootObject())
-		: Path(InPath)
+	// Only used by containers
+	FFolder()
+		: bPathInitialized(true)
+		, Path(GetEmptyPath())
+		, RootObject(GetInvalidRootObject())
+	{}
+
+	FFolder(const FRootObject& InRootObject, const FName& InPath = GetEmptyPath())
+		: bPathInitialized(true)
+		, Path(InPath)
 		, RootObject(InRootObject)
 	{}
 
-	ENGINE_API static TOptional<FRootObject> GetOptionalFolderRootObject(const ULevel* InLevel);
+	FFolder(const FRootObject& InRootObject, const FGuid& InActorFolderGuid)
+		: bPathInitialized(false)
+		, Path(NAME_None)
+		, RootObject(InRootObject)
+		, ActorFolderGuid(InActorFolderGuid)
+	{
+		check(ActorFolderGuid.IsValid());
+	}
+
+	FORCEINLINE static UObject* GetRootObjectPtr(const FRootObject& InRootObject)
+	{
+		return InRootObject.ResolveObjectPtr();
+	}
 
-	FORCEINLINE static bool HasRootObject(const FRootObject& Key)
+	FORCEINLINE static bool IsRootObjectValid(const FRootObject& Key)
 	{
-		return Key != GetDefaultRootObject();
+		return Key != GetInvalidRootObject();
 	}
 
 	FORCEINLINE static FName GetEmptyPath()
@@ -33,25 +55,24 @@ struct FFolder
 		return NAME_None;
 	}
 
-	FORCEINLINE static FRootObject GetDefaultRootObject()
+	FORCEINLINE bool IsRootObjectValid() const
 	{
-		return FRootObject();
+		return FFolder::IsRootObjectValid(RootObject);
 	}
 
-	FORCEINLINE static UObject* GetRootObjectPtr(const FRootObject& InRootObject)
+	FORCEINLINE ULevel* GetRootObjectAssociatedLevel() const
 	{
-		return InRootObject.ResolveObjectPtr();
+		return FFolder::GetRootObjectAssociatedLevel(RootObject);
 	}
 
-	FORCEINLINE bool HasRootObject() const
+	FORCEINLINE bool IsRootObjectPersistentLevel() const
 	{
-		return FFolder::HasRootObject(RootObject);
+		return FFolder::IsRootObjectPersistentLevel(RootObject);
 	}
 
-	FORCEINLINE FFolder GetParent() const
+	FORCEINLINE bool IsValid() const
 	{
-		const FName ParentPath(*FPaths::GetPath(Path.ToString()));
-		return FFolder(ParentPath, RootObject);
+		return IsRootObjectValid();
 	}
 
 	FORCEINLINE bool IsChildOf(const FFolder& InParent) const
@@ -65,32 +86,28 @@ struct FFolder
 
 	FORCEINLINE bool IsNone() const
 	{
-		return Path.IsNone();
+		return GetPath().IsNone();
 	}
 
-	const FRootObject& GetRootObject() const
+	FORCEINLINE const FRootObject& GetRootObject() const
 	{
 		return RootObject;
 	}
 
 	FORCEINLINE UObject* GetRootObjectPtr() const
 	{
-		return GetRootObjectPtr(RootObject);
-	}
-
-	FORCEINLINE const FName& GetPath() const
-	{
-		return Path;
+		return FFolder::GetRootObjectPtr(RootObject);
 	}
 
-	FORCEINLINE void SetPath(const FName& InPath)
+	FORCEINLINE const FGuid& GetActorFolderGuid() const
 	{
-		Path = InPath;
+		return ActorFolderGuid;
 	}
 
 	FORCEINLINE FName GetLeafName() const
 	{
-		FString PathString = Path.ToString();
+		FName PathLocal = GetPath();
+		FString PathString = PathLocal.ToString();
 		int32 LeafIndex = 0;
 		if (PathString.FindLastChar('/', LeafIndex))
 		{
@@ -98,36 +115,49 @@ struct FFolder
 		}
 		else
 		{
-			return Path;
+			return PathLocal;
 		}
 	}
 
 	FORCEINLINE bool operator == (const FFolder& InOther) const
 	{
-		return (Path == InOther.Path) && (RootObject == InOther.RootObject);
+		return (RootObject == InOther.RootObject) && (GetPath() == InOther.GetPath()); // don't check for ActorFolderGuid as it's only used as an accelerator
 	}
 
 	FORCEINLINE bool operator != (const FFolder& InOther) const
 	{
 		return !operator==(InOther);
 	}
-
+	
 	FORCEINLINE FString ToString() const
-	{
-		return Path.ToString();
+	{ 
+		return GetPath().ToString();
 	}
 
 	FORCEINLINE friend FArchive& operator<<(FArchive& Ar, FFolder& Folder)
 	{
 		check(!Ar.IsPersistent());
-		return Ar << Folder.Path << Folder.RootObject;
+		return Ar << Folder.bPathInitialized << Folder.Path << Folder.RootObject << Folder.ActorFolderGuid;
 	}
 
+	const FName GetPath() const;
+	FFolder GetParent() const;
+	UActorFolder* GetActorFolder() const;
+
+	// Helpers methods
+	static const FFolder& GetInvalidFolder();
+	static const FRootObject& GetInvalidRootObject();
+	static bool IsRootObjectPersistentLevel(const FRootObject& Key);
+	static TOptional<FRootObject> GetOptionalFolderRootObject(const ULevel* InLevel);
+	static ULevel* GetRootObjectAssociatedLevel(const FRootObject& Key);
+	static FFolder GetWorldRootFolder(UWorld* InWorld);
+	static bool GetFolderPathsAndCommonRootObject(const TArray<FFolder>& InFolders, TArray<FName>& OutFolders, FRootObject& OutCommonRootObject);
+
 private:
 	bool PathIsChildOf(const FString& InPotentialChild, const FString& InParent) const
 	{
 		const int32 ParentLen = InParent.Len();
-		
+
 		// If parent is empty and child isn't, consider that path is child of parent
 		if ((InPotentialChild.Len() > 0) && (ParentLen == 0))
 		{
@@ -145,8 +175,10 @@ private:
 		return PathIsChildOf(InPotentialChild.ToString(), InParent.ToString());
 	}
 
-	FName Path;
+	mutable bool bPathInitialized;
+	mutable FName Path;
 	FRootObject RootObject;
+	mutable FGuid ActorFolderGuid; // Optional : Used to find Level's ActorFolder faster than by using the path
 #endif
 };
 
