Scene Outliner: refactored actor pin API to be batched, so we can pin actors all at once instead of doing the operation one by one.

#rnx

diff --git a/Engine/Source/Editor/SceneOutliner/Private/ActorBrowsingMode.cpp b/Engine/Source/Editor/SceneOutliner/Private/ActorBrowsingMode.cpp
--- a/Engine/Source/Editor/SceneOutliner/Private/ActorBrowsingMode.cpp
+++ b/Engine/Source/Editor/SceneOutliner/Private/ActorBrowsingMode.cpp
@@ -1748,6 +1748,38 @@
 			}
 		}
 	}
+
+	static void RecursiveAddItemsToActorGuidList(const TArray<FSceneOutlinerTreeItemPtr>& Items, TArray<FGuid>& List)
+	{
+		for (const FSceneOutlinerTreeItemPtr& Item : Items)
+		{
+			if (const FActorDescTreeItem* const ActorDescTreeItem = Item->CastTo<FActorDescTreeItem>())
+			{
+				List.Add(ActorDescTreeItem->GetGuid());
+			}
+			else if (const FActorTreeItem* const ActorTreeItem = Item->CastTo<FActorTreeItem>())
+			{
+				if (ActorTreeItem->Actor.IsValid())
+				{
+					List.Add(ActorTreeItem->Actor->GetActorGuid());
+				}
+			}
+
+			TArray<FSceneOutlinerTreeItemPtr> ChildrenItems;
+			for (const auto& Child : Item->GetChildren())
+			{
+				if (Child.IsValid())
+				{
+					ChildrenItems.Add(Child.Pin());
+				}
+			}
+
+			if (ChildrenItems.Num())
+			{
+				RecursiveAddItemsToActorGuidList(ChildrenItems, List);
+			}
+		}
+	};
 }
 
 void FActorBrowsingMode::SelectFoldersDescendants(const TArray<FFolderTreeItem*>& FolderItems, bool bSelectImmediateChildrenOnly)
@@ -1775,94 +1807,77 @@
 	GEditor->NoteSelectionChange();
 }
 
-void FActorBrowsingMode::PinItem(const FSceneOutlinerTreeItemPtr& InItem)
+void FActorBrowsingMode::PinItems(const TArray<FSceneOutlinerTreeItemPtr>& InItems)
 {
-	AActor* PinnedActor = nullptr;
-	if (UWorldPartition* const WorldPartition = RepresentingWorld->GetWorldPartition())
+	UWorldPartition* const WorldPartition = RepresentingWorld->GetWorldPartition();
+	if (!WorldPartition)
 	{
-		if (const FActorDescTreeItem* const ActorDescTreeItem = InItem->CastTo<FActorDescTreeItem>())
+		return;
+	}
+
+	TArray<FGuid> ActorsToPin;
+	ActorBrowsingModeUtils::RecursiveAddItemsToActorGuidList(InItems, ActorsToPin);
+
+	if (ActorsToPin.Num())
+	{
+		GEditor->GetSelectedActors()->BeginBatchSelectOperation();
+		GEditor->SelectNone(/*bNoteSelectionChange=*/false, /*bDeselectBSPSurfs=*/true);
+
+		WorldPartition->PinActors(ActorsToPin);
+
+		AActor* LastPinnedActor = nullptr;
+		for (const FGuid& ActorGuid : ActorsToPin)
 		{
-			PinnedActor = WorldPartition->PinActor(ActorDescTreeItem->GetGuid());
-		}
-		else if (const FActorTreeItem* const ActorTreeItem = InItem->CastTo<FActorTreeItem>())
-		{
-			if (ActorTreeItem->Actor.IsValid())
+			if (FWorldPartitionHandle ActorHandle(WorldPartition, ActorGuid); ActorHandle.IsValid())
 			{
-				PinnedActor = WorldPartition->PinActor(ActorTreeItem->Actor->GetActorGuid());
+				if (AActor* PinnedActor = ActorHandle->GetActor())
+				{
+					GEditor->SelectActor(PinnedActor, /*bInSelected=*/true, /*bNotify=*/false);
+					LastPinnedActor = PinnedActor;
+				}
 			}
 		}
-	}
 
-	// Check if we need to start a batch selection
-	bool bIsBatchSelectOwner = false;
-	if (PinnedActor)
-	{
-		const bool bIsBatchSelecting = GEditor->GetSelectedActors()->IsBatchSelecting();
-		if (!bIsBatchSelecting)
+		GEditor->GetSelectedActors()->EndBatchSelectOperation(/*bNotify=*/true);
+
+		if (LastPinnedActor)
 		{
-			bIsBatchSelectOwner = true;
-			GEditor->GetSelectedActors()->BeginBatchSelectOperation();
-			GEditor->SelectNone(/*bNoteSelectionChange=*/false, /*bDeselectBSPSurfs=*/true);
-			SceneOutliner->OnItemAdded(PinnedActor, SceneOutliner::ENewItemAction::ScrollIntoView);
+			SceneOutliner->OnItemAdded(LastPinnedActor, SceneOutliner::ENewItemAction::ScrollIntoView);
 		}
-
-		GEditor->SelectActor(PinnedActor, /*bInSelected=*/true, /*bNotify=*/false);
-	}
-
-	// Recursively pin all children.
-	for (const auto& Child : InItem->GetChildren())
-	{
-		if (Child.IsValid())
-		{
-			PinItem(Child.Pin());
-	    }
-	}
-
-	// End batch selection if needed
-	if (bIsBatchSelectOwner)
-	{
-		GEditor->GetSelectedActors()->EndBatchSelectOperation(/*bNotify=*/true);
 	}
 }
 
-void FActorBrowsingMode::UnpinItem(const FSceneOutlinerTreeItemPtr& InItem)
+void FActorBrowsingMode::UnpinItems(const TArray<FSceneOutlinerTreeItemPtr>& InItems)
 {
-	// Check if we need to start a batch selection
-	bool bIsBatchSelectOwner = !GEditor->GetSelectedActors()->IsBatchSelecting();
-	if (bIsBatchSelectOwner)
+	UWorldPartition* const WorldPartition = RepresentingWorld->GetWorldPartition();
+	if (!WorldPartition)
+	{
+		return;
+	}
+
+	TArray<FGuid> ActorsToUnpin;
+	ActorBrowsingModeUtils::RecursiveAddItemsToActorGuidList(InItems, ActorsToUnpin);
+
+	if (ActorsToUnpin.Num())
 	{
 		GEditor->GetSelectedActors()->BeginBatchSelectOperation();
-	}
-	
-	// Recursively unpin all children
-	for (const auto& Child : InItem->GetChildren())
-	{
-		if (Child.IsValid())
+
+		AActor* LastPinnedActor = nullptr;
+		for (const FGuid& ActorGuid : ActorsToUnpin)
 		{
-			UnpinItem(Child.Pin());
-		}
-	}
-	
-	if (UWorldPartition* const WorldPartition = RepresentingWorld->GetWorldPartition())
-	{
-		if (const FActorDescTreeItem* const ActorDescTreeItem = InItem->CastTo<FActorDescTreeItem>())
-		{
-			WorldPartition->UnpinActor(ActorDescTreeItem->GetGuid());
-		}
-		else if (const FActorTreeItem* const ActorTreeItem = InItem->CastTo<FActorTreeItem>())
-		{
-			if (AActor* PinnedActor = ActorTreeItem->Actor.Get())
+			if (FWorldPartitionHandle ActorHandle(WorldPartition, ActorGuid); ActorHandle.IsValid())
 			{
-				GEditor->SelectActor(PinnedActor, /*bInSelected=*/false, /*bNotify=*/false);
-				WorldPartition->UnpinActor(PinnedActor->GetActorGuid());
+				if (AActor* PinnedActor = ActorHandle->GetActor())
+				{
+					GEditor->SelectActor(PinnedActor, /*bInSelected=*/false, /*bNotify=*/false);
+				}
 			}
 		}
-	}
 
-	if (bIsBatchSelectOwner)
-	{
+		WorldPartition->UnpinActors(ActorsToUnpin);
+
 		GEditor->GetSelectedActors()->EndBatchSelectOperation(/*bNotify=*/true);
-	}
+	}	
 }
 
 void FActorBrowsingMode::PinSelectedItems()
@@ -1870,21 +1885,13 @@
 	const FSceneOutlinerItemSelection Selection = SceneOutliner->GetSelection();
 	if (Selection.Num())
 	{
-		GEditor->GetSelectedActors()->BeginBatchSelectOperation();
-		GEditor->SelectNone(/*bNoteSelectionChange=*/false, /*bDeselectBSPSurfs=*/true);
-
-		Selection.ForEachItem([this](const FSceneOutlinerTreeItemPtr& TreeItem)
+		TArray<FSceneOutlinerTreeItemPtr> ItemsToPin;
+		Selection.ForEachItem([this, &ItemsToPin](const FSceneOutlinerTreeItemPtr& TreeItem)
 		{
-			PinItem(TreeItem);
+			ItemsToPin.Add(TreeItem);
 			return true;
 		});
-
-		GEditor->GetSelectedActors()->EndBatchSelectOperation(/*bNotify=*/true);
-
-		if (AActor* Actor = GEditor->GetSelectedActors()->GetTop<AActor>())
-		{
-			SceneOutliner->OnItemAdded(Actor, SceneOutliner::ENewItemAction::ScrollIntoView);
-		}
+		PinItems(ItemsToPin);
 	}
 }
 
@@ -1893,15 +1900,13 @@
 	const FSceneOutlinerItemSelection Selection = SceneOutliner->GetSelection();
 	if(Selection.Num())
 	{
-		GEditor->GetSelectedActors()->BeginBatchSelectOperation();
-
-		Selection.ForEachItem([this](const FSceneOutlinerTreeItemPtr& TreeItem)
+		TArray<FSceneOutlinerTreeItemPtr> ItemsToUnpin;
+		Selection.ForEachItem([this, &ItemsToUnpin](const FSceneOutlinerTreeItemPtr& TreeItem)
 		{
-			UnpinItem(TreeItem);
+			ItemsToUnpin.Add(TreeItem);
 			return true;
 		});
-
-		GEditor->GetSelectedActors()->EndBatchSelectOperation(/*bNotify=*/true);
+		UnpinItems(ItemsToUnpin);
 	}
 }
 
diff --git a/Engine/Source/Editor/SceneOutliner/Private/SSceneOutliner.cpp b/Engine/Source/Editor/SceneOutliner/Private/SSceneOutliner.cpp
--- a/Engine/Source/Editor/SceneOutliner/Private/SSceneOutliner.cpp
+++ b/Engine/Source/Editor/SceneOutliner/Private/SSceneOutliner.cpp
@@ -2276,14 +2276,14 @@ FReply SSceneOutliner::OnDragOverItem(const FDragDropEvent& Event, const ISceneO
 	return Mode->OnDragOverItem(Event, Item); 
 }
 
-void SSceneOutliner::PinItem(const FSceneOutlinerTreeItemPtr& InItem)
+void SSceneOutliner::PinItems(const TArray<FSceneOutlinerTreeItemPtr>& InItems)
 {
-	Mode->PinItem(InItem);
+	Mode->PinItems(InItems);
 }
 
-void SSceneOutliner::UnpinItem(const FSceneOutlinerTreeItemPtr& InItem)
+void SSceneOutliner::UnpinItems(const TArray<FSceneOutlinerTreeItemPtr>& InItems)
 {
-	Mode->UnpinItem(InItem);
+	Mode->UnpinItems(InItems);
 }
 
 void SSceneOutliner::PinSelectedItems()
diff --git a/Engine/Source/Editor/SceneOutliner/Private/SceneOutlinerPinnedColumn.cpp b/Engine/Source/Editor/SceneOutliner/Private/SceneOutlinerPinnedColumn.cpp
--- a/Engine/Source/Editor/SceneOutliner/Private/SceneOutlinerPinnedColumn.cpp
+++ b/Engine/Source/Editor/SceneOutliner/Private/SceneOutlinerPinnedColumn.cpp
@@ -102,7 +102,7 @@ private:
 			}
 			else
 			{
-				Outliner->PinItem(TreeItem);
+				Outliner->PinItems({ TreeItem });
 			}
 		}
 		else
@@ -113,7 +113,7 @@ private:
 			}
 			else
 			{
-				Outliner->UnpinItem(TreeItem);
+				Outliner->UnpinItems({ TreeItem });
 			}
 		}
 
diff --git a/Engine/Source/Editor/SceneOutliner/Public/ActorBrowsingMode.h b/Engine/Source/Editor/SceneOutliner/Public/ActorBrowsingMode.h
--- a/Engine/Source/Editor/SceneOutliner/Public/ActorBrowsingMode.h
+++ b/Engine/Source/Editor/SceneOutliner/Public/ActorBrowsingMode.h
@@ -50,8 +50,8 @@ public:
 	virtual bool CreateFolder(const FFolder& NewFolder) override;
 	virtual bool ReparentItemToFolder(const FFolder& FolderPath, const FSceneOutlinerTreeItemPtr& Item) override;
 	virtual void SelectFoldersDescendants(const TArray<FFolderTreeItem*>& FolderItems, bool bSelectImmediateChildrenOnly) override;
-	virtual void PinItem(const FSceneOutlinerTreeItemPtr& InItem) override;
-	virtual void UnpinItem(const FSceneOutlinerTreeItemPtr& InItem) override;
+	virtual void PinItems(const TArray<FSceneOutlinerTreeItemPtr>& InItems) override;
+	virtual void UnpinItems(const TArray<FSceneOutlinerTreeItemPtr>& InItems) override;
 	virtual void PinSelectedItems() override;
 	virtual void UnpinSelectedItems() override;
 	/* End ISceneOutlinerMode Interface */
diff --git a/Engine/Source/Editor/SceneOutliner/Public/ISceneOutliner.h b/Engine/Source/Editor/SceneOutliner/Public/ISceneOutliner.h
--- a/Engine/Source/Editor/SceneOutliner/Public/ISceneOutliner.h
+++ b/Engine/Source/Editor/SceneOutliner/Public/ISceneOutliner.h
@@ -133,11 +133,11 @@ public:
 	/** Set the item selection of the outliner based on a selector function. Any items which return true will be added */
 	virtual void SetSelection(const TFunctionRef<bool(ISceneOutlinerTreeItem&)> Selector) = 0;
 
-	/** Pins an item in the outliner. */
-	virtual void PinItem(const FSceneOutlinerTreeItemPtr& InItem) = 0;
+	/** Pins an item list in the outliner. */
+	virtual void PinItems(const TArray<FSceneOutlinerTreeItemPtr>& InItems) = 0;
 
-	/**  Unpins an item in the outliner. */
-	virtual void UnpinItem(const FSceneOutlinerTreeItemPtr& InItem) = 0;
+	/**  Unpins an item list in the outliner. */
+	virtual void UnpinItems(const TArray<FSceneOutlinerTreeItemPtr>& InItems) = 0;
 	
 	/** Pin selected items */
 	virtual void PinSelectedItems() = 0;
diff --git a/Engine/Source/Editor/SceneOutliner/Public/ISceneOutlinerMode.h b/Engine/Source/Editor/SceneOutliner/Public/ISceneOutlinerMode.h
--- a/Engine/Source/Editor/SceneOutliner/Public/ISceneOutlinerMode.h
+++ b/Engine/Source/Editor/SceneOutliner/Public/ISceneOutlinerMode.h
@@ -145,10 +145,10 @@ public:
 	/** Select all descendants of a folder. Optionally select only immediate descendants. */
 	virtual void SelectFoldersDescendants(const TArray<FFolderTreeItem*>& FolderItems, bool bSelectImmediateChildrenOnly) {}
 public:
-	/** Pins an item in the outliner */
-	virtual void PinItem(const FSceneOutlinerTreeItemPtr& InItem) {}
-	/** Unpins an item in the outliner */
-	virtual void UnpinItem(const FSceneOutlinerTreeItemPtr& InItem) {}
+	/** Pins an item list in the outliner */
+	virtual void PinItems(const TArray<FSceneOutlinerTreeItemPtr>& InItems) {}
+	/** Unpins an item list in the outliner */
+	virtual void UnpinItems(const TArray<FSceneOutlinerTreeItemPtr>& InItems) {}
 	/** Pins all selected items */
 	virtual void PinSelectedItems() {}
 	/** Unpins all selected items */
diff --git a/Engine/Source/Editor/SceneOutliner/Public/SSceneOutliner.h b/Engine/Source/Editor/SceneOutliner/Public/SSceneOutliner.h
--- a/Engine/Source/Editor/SceneOutliner/Public/SSceneOutliner.h
+++ b/Engine/Source/Editor/SceneOutliner/Public/SSceneOutliner.h
@@ -402,14 +402,14 @@ public:
 	virtual FSceneOutlinerItemSelection GetSelection() const { return FSceneOutlinerItemSelection(*OutlinerTreeView); }
 
 	/**
-	 * Pins an item in the outliner.
+	 * Pins an item list in the outliner.
 	 */
-	virtual void PinItem(const FSceneOutlinerTreeItemPtr& InItem) override;
+	virtual void PinItems(const TArray<FSceneOutlinerTreeItemPtr>& InItems) override;
 
 	/**
-	 * Unpins an item in the outliner.
+	 * Unpins an item list in the outliner.
 	 */
-	virtual void UnpinItem(const FSceneOutlinerTreeItemPtr& InItem) override;
+	virtual void UnpinItems(const TArray<FSceneOutlinerTreeItemPtr>& InItems) override;
 
 	/**
 	 * Pin selected items
diff --git a/Engine/Source/Runtime/Engine/Private/WorldPartition/WorldPartition.cpp b/Engine/Source/Runtime/Engine/Private/WorldPartition/WorldPartition.cpp
--- a/Engine/Source/Runtime/Engine/Private/WorldPartition/WorldPartition.cpp
+++ b/Engine/Source/Runtime/Engine/Private/WorldPartition/WorldPartition.cpp
@@ -216,38 +216,22 @@
 public:
 	FLoaderAdapterPinnedActors(UWorld* InWorld)
 		: FLoaderAdapterList(InWorld)
-		, bStateChanged(false)
 	{
 		Load();
 	}
 
-	void AddActor(FWorldPartitionHandle& ActorHandle, bool bImmediate)
+	void AddActors(const TArray<FWorldPartitionHandle>& ActorHandles)
 	{
-		if (!ContainsActor(ActorHandle))
-		{
-			Actors.Add(ActorHandle);
-
-			if (bImmediate)
-			{
-				RefreshLoadedState();
-				bStateChanged = false;
-			}
-			else
-			{
-				bStateChanged = true;
-			}
-		}
+		Actors.Append(ActorHandles);
+		RefreshLoadedState();
 	}
 
-	void RemoveActor(FWorldPartitionHandle& ActorHandle)
+	void RemoveActors(const TArray<FWorldPartitionHandle>& ActorHandles)
 	{
-		if (ContainsActor(ActorHandle))
-		{
-			ActorToRemove = ActorHandle;
-			RefreshLoadedState();
-			Actors.Remove(ActorHandle);
-			ActorToRemove.Reset();
-		}
+		ActorsToRemove = TSet<FWorldPartitionHandle>(ActorHandles);
+		RefreshLoadedState();
+		Actors = Actors.Difference(ActorsToRemove);
+		ActorsToRemove.Empty();
 	}
 
 	bool ContainsActor(const FWorldPartitionHandle& ActorHandle) const
@@ -255,26 +239,16 @@
 		return Actors.Contains(ActorHandle);
 	}
 
-	void Tick(float DeltaSeconds)
-	{
-		if (bStateChanged)
-		{
-			RefreshLoadedState();
-			bStateChanged = false;
-		}
-	}
-
 protected:
 	//~ Begin IWorldPartitionActorLoaderInterface::ILoaderAdapterList interface
 	virtual bool ShouldActorBeLoaded(const FWorldPartitionHandle& ActorHandle) const override
 	{
-		return (ActorHandle != ActorToRemove) && PassDataLayersFilter(ActorHandle);
+		return !ActorsToRemove.Contains(ActorHandle) && PassDataLayersFilter(ActorHandle);
 	}
 	//~ End IWorldPartitionActorLoaderInterface::ILoaderAdapterList interface
 
 private:
-	bool bStateChanged;
-	FWorldPartitionHandle ActorToRemove;
+	TSet<FWorldPartitionHandle> ActorsToRemove;
 };
 #endif
 
@@ -359,7 +333,7 @@
 				// If we hold the last reference to that actor, pin it to avoid unloading
 				if (PinnedActors && (ActorHandle->GetHardRefCount() == 1))
 				{
-					PinnedActors->AddActor(ActorHandle, /*bImmediate*/true);
+					PinnedActors->AddActors({ ActorHandle });
 				}
 
 				DirtyActors.Remove(ActorHandle);
@@ -1143,11 +1117,6 @@
 		EditorHash->Tick(DeltaSeconds);
 	}
 
-	if (PinnedActors)
-	{
-		PinnedActors->Tick(DeltaSeconds);
-	}
-
 	if (bForceGarbageCollection)
 	{
 		GEngine->ForceGarbageCollection(bForceGarbageCollectionPurge);
@@ -1368,29 +1337,43 @@
 	}
 }
 
-AActor* UWorldPartition::PinActor(const FGuid& ActorGuid)
+void UWorldPartition::PinActors(const TArray<FGuid>& ActorGuids)
 {
 	if (PinnedActors)
 	{
-		FWorldPartitionHandle ActorHandle(this, ActorGuid);
-		if (ActorHandle.IsValid())
+		TArray<FWorldPartitionHandle> ActorHandles;
+		ActorHandles.Reserve(ActorGuids.Num());
+
+		for (const FGuid& ActorGuid : ActorGuids)
 		{
-			PinnedActors->AddActor(ActorHandle, /*bImmediate*/false);
-			return ActorHandle->GetActor();
-		}		
+			FWorldPartitionHandle ActorHandle(this, ActorGuid);
+			if (ActorHandle.IsValid())
+			{
+				ActorHandles.Add(ActorHandle);
+			}
+		}
+
+		PinnedActors->AddActors(ActorHandles);
 	}
-	return nullptr;
 }
 
-void UWorldPartition::UnpinActor(const FGuid& ActorGuid)
+void UWorldPartition::UnpinActors(const TArray<FGuid>& ActorGuids)
 {
 	if (PinnedActors)
 	{
-		FWorldPartitionHandle ActorHandle(this, ActorGuid);
-		if (ActorHandle.IsValid())
+		TArray<FWorldPartitionHandle> ActorHandles;
+		ActorHandles.Reserve(ActorGuids.Num());
+
+		for (const FGuid& ActorGuid : ActorGuids)
 		{
-			PinnedActors->RemoveActor(ActorHandle);
+			FWorldPartitionHandle ActorHandle(this, ActorGuid);
+			if (ActorHandle.IsValid())
+			{
+				ActorHandles.Add(ActorHandle);
+			}
 		}
+
+		PinnedActors->RemoveActors(ActorHandles);
 	}
 }
 
diff --git a/Engine/Source/Runtime/Engine/Public/WorldPartition/WorldPartition.h b/Engine/Source/Runtime/Engine/Public/WorldPartition/WorldPartition.h
--- a/Engine/Source/Runtime/Engine/Public/WorldPartition/WorldPartition.h
+++ b/Engine/Source/Runtime/Engine/Public/WorldPartition/WorldPartition.h
@@ -170,8 +170,8 @@
 	void AppendAssetRegistryTags(TArray<FAssetRegistryTag>& OutTags) const;
 
 	// Actors pinning
-	AActor* PinActor(const FGuid& ActorGuid);
-	void UnpinActor(const FGuid& ActorGuid);
+	void PinActors(const TArray<FGuid>& ActorGuids);
+	void UnpinActors(const TArray<FGuid>& ActorGuids);
 	bool IsActorPinned(const FGuid& ActorGuid) const;
 
 	void LoadLastLoadedRegions(const TArray<FBox>& EditorLastLoadedRegions);
