Scene Outliner: Remove forcibly showing actors with the SequencerActor tag. Spawnables are still shown in the World Outliner still because they are RF_Transient and the World Outliner is also by default set to show transient actors.

This fixes an issue where in an actor picker menu, you could pick a spawnable actor and it's not desirable to make that reference to a transient actor.

If you need to reference a spawnable from sequencer, you need to make an explicit actor reference property track and choose the actor from the picker in the track. You will no longer be able to pick a spwanble from the details panel for genuine operations such as autokeying a property that already has an actor reference track. You will need to choose it from the picker in the track.

diff --git a/Engine/Source/Editor/SceneOutliner/Private/ActorMode.cpp b/Engine/Source/Editor/SceneOutliner/Private/ActorMode.cpp
--- a/Engine/Source/Editor/SceneOutliner/Private/ActorMode.cpp
+++ b/Engine/Source/Editor/SceneOutliner/Private/ActorMode.cpp
@@ -330,15 +330,12 @@ FFolder::FRootObject FActorMode::GetPasteTargetRootObject() const
 
 bool FActorMode::IsActorDisplayable(const SSceneOutliner* SceneOutliner, const AActor* Actor)
 {
-	static const FName SequencerActorTag(TEXT("SequencerActor"));
-
 	return Actor &&
 		!SceneOutliner->GetSharedData().bOnlyShowFolders && 									// Don't show actors if we're only showing folders
 		Actor->IsEditable() &&																	// Only show actors that are allowed to be selected and drawn in editor
 		Actor->IsListedInSceneOutliner() &&
 		(((Actor->GetWorld() && Actor->GetWorld()->IsPlayInEditor()) || !Actor->HasAnyFlags(RF_Transient)) ||
-		(SceneOutliner->GetSharedData().bShowTransient && Actor->HasAnyFlags(RF_Transient)) ||	// Don't show transient actors in non-play worlds
-		(Actor->ActorHasTag(SequencerActorTag))) &&
+		(SceneOutliner->GetSharedData().bShowTransient && Actor->HasAnyFlags(RF_Transient))) &&	// Don't show transient actors in non-play worlds
 		!Actor->IsTemplate() &&																	// Should never happen, but we never want CDOs displayed
 		!FActorEditorUtils::IsABuilderBrush(Actor) &&											// Don't show the builder brush
 		!Actor->IsA(AWorldSettings::StaticClass()) &&											// Don't show the WorldSettings actor, even though it is technically editable
