Fixed crash in FActorTreeItem access null Level owning world (happens when level is streamed-out).
#rnx

diff --git a/Engine/Source/Editor/SceneOutliner/Private/ActorTreeItem.cpp b/Engine/Source/Editor/SceneOutliner/Private/ActorTreeItem.cpp
--- a/Engine/Source/Editor/SceneOutliner/Private/ActorTreeItem.cpp
+++ b/Engine/Source/Editor/SceneOutliner/Private/ActorTreeItem.cpp
@@ -392,7 +392,7 @@ bool FActorTreeItem::ShouldShowPinnedState() const
 	{
 		const ULevel* Level = ActorPtr->GetLevel();
 		const UWorld* World = Level->GetWorld();
-		return !World->IsGameWorld() && !!Level->GetWorldPartition();
+		return World && !World->IsGameWorld() && !!Level->GetWorldPartition();
 	}
 
 	return false;
