SceneOutliner: Hide empty folders option.

diff --git a/Engine/Source/Editor/SceneOutliner/Private/ActorBrowsingMode.cpp b/Engine/Source/Editor/SceneOutliner/Private/ActorBrowsingMode.cpp
--- a/Engine/Source/Editor/SceneOutliner/Private/ActorBrowsingMode.cpp
+++ b/Engine/Source/Editor/SceneOutliner/Private/ActorBrowsingMode.cpp
@@ -42,9 +42,10 @@ DEFINE_LOG_CATEGORY_STATIC(LogActorBrowser, Log, All);
 
 using FActorFilter = TSceneOutlinerPredicateFilter<FActorTreeItem>;
 using FActorDescFilter = TSceneOutlinerPredicateFilter<FActorDescTreeItem>;
+using FFolderFilter = TSceneOutlinerPredicateFilter<FFolderTreeItem>;
 
 FActorBrowsingMode::FActorBrowsingMode(SSceneOutliner* InSceneOutliner, TWeakObjectPtr<UWorld> InSpecifiedWorldToDisplay)
-	: FActorModeInteractive(FActorModeParams(InSceneOutliner, InSpecifiedWorldToDisplay,  /* bHideComponents */ true, /* bHideLevelInstanceHierarchy */ false, /* bHideUnloadedActors */ false))
+	: FActorModeInteractive(FActorModeParams(InSceneOutliner, InSpecifiedWorldToDisplay,  /* bHideComponents */ true, /* bHideLevelInstanceHierarchy */ false, /* bHideUnloadedActors */ false, /* bHideEmptyFolders */ false))
 	, FilteredActorCount(0)
 {
 	// Capture selection changes of bones from mesh selection in fracture tools
@@ -169,6 +170,23 @@ FActorBrowsingMode::FActorBrowsingMode(SSceneOutliner* InSceneOutliner, TWeakObj
 		});
 	FilterInfoMap.Add(TEXT("HideUnloadedActorsFilter"), HideUnloadedActorsInfo);
 
+	FSceneOutlinerFilterInfo HideEmptyFoldersInfo(LOCTEXT("ToggleHideEmptyFolders", "Hide Empty Folders"), LOCTEXT("ToggleHideEmptyFoldersToolTip", "When enabled, hides all empty folders."), LocalSettings.bHideEmptyFolders, FCreateSceneOutlinerFilter::CreateStatic(&FActorBrowsingMode::CreateHideEmptyFoldersFilter));
+	HideEmptyFoldersInfo.OnToggle().AddLambda([this](bool bIsActive)
+		{
+			FActorBrowsingModeConfig* Settings = GetMutableConfig();
+			if (Settings)
+			{
+				Settings->bHideEmptyFolders = bHideEmptyFolders = bIsActive;
+				SaveConfig();
+			}
+
+			if (auto ActorHierarchy = StaticCast<FActorHierarchy*>(Hierarchy.Get()))
+			{
+				ActorHierarchy->SetShowingEmptyFolders(!bIsActive);
+			}
+		});
+	FilterInfoMap.Add(TEXT("HideEmptyFoldersFilter"), HideEmptyFoldersInfo);
+
 	// Add a filter which sets the interactive mode of LevelInstance items and their children
 	SceneOutliner->AddFilter(MakeShared<FActorFilter>(FActorTreeItem::FFilterPredicate::CreateStatic([](const AActor* Actor) {return true; }), FSceneOutlinerFilter::EDefaultBehaviour::Pass, FActorTreeItem::FFilterPredicate::CreateLambda([this](const AActor* Actor)
 		{
@@ -420,6 +438,12 @@ TSharedRef<FSceneOutlinerFilter> FActorBrowsingMode::CreateHideUnloadedActorsFil
 		[](const FWorldPartitionActorDesc* ActorDesc) { return false; }), FSceneOutlinerFilter::EDefaultBehaviour::Pass));
 }
 
+TSharedRef<FSceneOutlinerFilter> FActorBrowsingMode::CreateHideEmptyFoldersFilter()
+{
+	return MakeShareable(new FFolderFilter(FFolderTreeItem::FFilterPredicate::CreateStatic(
+		[](const FFolder& Folder) { return false; }), FSceneOutlinerFilter::EDefaultBehaviour::Pass));
+}
+
 static const FName DefaultContextBaseMenuName("SceneOutliner.DefaultContextMenuBase");
 static const FName DefaultContextMenuName("SceneOutliner.DefaultContextMenu");
 
diff --git a/Engine/Source/Editor/SceneOutliner/Private/ActorBrowsingModeSettings.h b/Engine/Source/Editor/SceneOutliner/Private/ActorBrowsingModeSettings.h
--- a/Engine/Source/Editor/SceneOutliner/Private/ActorBrowsingModeSettings.h
+++ b/Engine/Source/Editor/SceneOutliner/Private/ActorBrowsingModeSettings.h
@@ -40,6 +40,10 @@ public:
 	UPROPERTY()
 	bool bHideUnloadedActors = false;
 
+	/** True when the Scene Outliner is not displaying empty folders */
+	UPROPERTY()
+	bool bHideEmptyFolders = false;
+
 	/** True when the Scene Outliner updates when an actor is selected in the viewport */
 	UPROPERTY()
 	bool bAlwaysFrameSelection = true;
diff --git a/Engine/Source/Editor/SceneOutliner/Private/ActorHierarchy.cpp b/Engine/Source/Editor/SceneOutliner/Private/ActorHierarchy.cpp
--- a/Engine/Source/Editor/SceneOutliner/Private/ActorHierarchy.cpp
+++ b/Engine/Source/Editor/SceneOutliner/Private/ActorHierarchy.cpp
@@ -390,7 +390,7 @@ void FActorHierarchy::CreateWorldChildren(UWorld* World, TArray<FSceneOutlinerTr
 {
 	check(World);
 
-	if (Mode->ShouldShowFolders())
+	if (Mode->ShouldShowFolders() && bShowingEmptyFolders)
 	{
 		// Add any folders which might match the current search terms
 		FActorFolders::Get().ForEachFolder(*World, [this, &World, &OutItems](const FFolder& Folder)
diff --git a/Engine/Source/Editor/SceneOutliner/Private/ActorMode.cpp b/Engine/Source/Editor/SceneOutliner/Private/ActorMode.cpp
--- a/Engine/Source/Editor/SceneOutliner/Private/ActorMode.cpp
+++ b/Engine/Source/Editor/SceneOutliner/Private/ActorMode.cpp
@@ -87,6 +87,7 @@ FActorMode::FActorMode(const FActorModeParams& Params)
 	, bHideActorWithNoComponent(Params.bHideActorWithNoComponent)
 	, bHideLevelInstanceHierarchy(Params.bHideLevelInstanceHierarchy)
 	, bHideUnloadedActors(Params.bHideUnloadedActors)
+	, bHideEmptyFolders(Params.bHideEmptyFolders)
 {
 	SceneOutliner->AddFilter(MakeShared<FActorFilter>(FActorTreeItem::FFilterPredicate::CreateLambda([this](const AActor* Actor)
 	{
@@ -135,6 +136,7 @@ TUniquePtr<ISceneOutlinerHierarchy> FActorMode::CreateHierarchy()
 	ActorHierarchy->SetShowingOnlyActorWithValidComponents(!bHideComponents && bHideActorWithNoComponent);
 	ActorHierarchy->SetShowingLevelInstances(!bHideLevelInstanceHierarchy);
 	ActorHierarchy->SetShowingUnloadedActors(!bHideUnloadedActors);
+	ActorHierarchy->SetShowingEmptyFolders(!bHideEmptyFolders);
 
 	return ActorHierarchy;
 }
diff --git a/Engine/Source/Editor/SceneOutliner/Private/SceneOutlinerModule.cpp b/Engine/Source/Editor/SceneOutliner/Private/SceneOutlinerModule.cpp
--- a/Engine/Source/Editor/SceneOutliner/Private/SceneOutlinerModule.cpp
+++ b/Engine/Source/Editor/SceneOutliner/Private/SceneOutlinerModule.cpp
@@ -88,11 +88,10 @@ TSharedRef<ISceneOutliner> FSceneOutlinerModule::CreateActorPicker(const FSceneO
 			Params.bHideComponents = true;
 			Params.bHideLevelInstanceHierarchy = true;
 			Params.bHideUnloadedActors = true;
+			Params.bHideEmptyFolders = true;
 			return new FActorPickingMode(Params, OnItemPicked);
 		});
 
-	
-
 	FSceneOutlinerInitializationOptions InitOptions(InInitOptions);
 	InitOptions.ModeFactory = ModeFactory;
 	if (InitOptions.ColumnMap.Num() == 0)
@@ -127,6 +126,7 @@ TSharedRef<ISceneOutliner> FSceneOutlinerModule::CreateComponentPicker(const FSc
 			Params.bHideActorWithNoComponent = true;
 			Params.bHideLevelInstanceHierarchy = true;
 			Params.bHideUnloadedActors = true;
+			Params.bHideEmptyFolders = true;
 			return new FActorPickingMode(Params, OnItemPicked);
 		});
 
diff --git a/Engine/Source/Editor/SceneOutliner/Public/ActorBrowsingMode.h b/Engine/Source/Editor/SceneOutliner/Public/ActorBrowsingMode.h
--- a/Engine/Source/Editor/SceneOutliner/Public/ActorBrowsingMode.h
+++ b/Engine/Source/Editor/SceneOutliner/Public/ActorBrowsingMode.h
@@ -115,6 +115,7 @@ private:
 	static TSharedRef<FSceneOutlinerFilter> CreateHideComponentsFilter();
 	static TSharedRef<FSceneOutlinerFilter> CreateHideLevelInstancesFilter();
 	static TSharedRef<FSceneOutlinerFilter> CreateHideUnloadedActorsFilter();
+	static TSharedRef<FSceneOutlinerFilter> CreateHideEmptyFoldersFilter();
 
 	/** Functions to expose selection framing to the UI */
 	void OnToggleAlwaysFrameSelection();
diff --git a/Engine/Source/Editor/SceneOutliner/Public/ActorHierarchy.h b/Engine/Source/Editor/SceneOutliner/Public/ActorHierarchy.h
--- a/Engine/Source/Editor/SceneOutliner/Public/ActorHierarchy.h
+++ b/Engine/Source/Editor/SceneOutliner/Public/ActorHierarchy.h
@@ -28,11 +28,14 @@ public:
 	void SetShowingOnlyActorWithValidComponents(bool bInShowingOnlyActorWithValidComponents) { bShowingOnlyActorWithValidComponents = bInShowingOnlyActorWithValidComponents; }
 	void SetShowingLevelInstances(bool bInShowingLevelInstances) { bShowingLevelInstances = bInShowingLevelInstances; }
 	void SetShowingUnloadedActors(bool bInShowingUnloadedActors) { bShowingUnloadedActors = bInShowingUnloadedActors; }
+	void SetShowingEmptyFolders(bool bInShowingEmptyFolders) { bShowingEmptyFolders = bInShowingEmptyFolders; }
+
 private:
 	/** Adds all the direct and indirect children of a world to OutItems */
 	void CreateWorldChildren(UWorld* World, TArray<FSceneOutlinerTreeItemPtr>& OutItems) const;
 	/** Create all component items for an actor if we are showing components and place them in OutItems */
 	void CreateComponentItems(const AActor* Actor, TArray<FSceneOutlinerTreeItemPtr>& OutItems) const;
+
 private:
 	// Update the hierarchy when actor or world changing events occur
 	void OnWorldPartitionCreated(UWorld* InWorld);
@@ -83,6 +86,8 @@ private:
 	bool bShowingOnlyActorWithValidComponents = false;
 	bool bShowingLevelInstances = false;
 	bool bShowingUnloadedActors = false;
+	bool bShowingEmptyFolders = false;
+
 private:
 	FActorHierarchy(ISceneOutlinerMode* Mode, const TWeakObjectPtr<UWorld>& Worlds);
 
diff --git a/Engine/Source/Editor/SceneOutliner/Public/ActorMode.h b/Engine/Source/Editor/SceneOutliner/Public/ActorMode.h
--- a/Engine/Source/Editor/SceneOutliner/Public/ActorMode.h
+++ b/Engine/Source/Editor/SceneOutliner/Public/ActorMode.h
@@ -23,12 +23,13 @@ struct SCENEOUTLINER_API FActorModeParams
 {
 	FActorModeParams() {}
 
-	FActorModeParams(SSceneOutliner* InSceneOutliner, const TWeakObjectPtr<UWorld>& InSpecifiedWorldToDisplay = nullptr, bool bInHideComponents = true, bool bInHideLevelInstanceHierarchy = true, bool bInHideUnloadedActors = true)
+	FActorModeParams(SSceneOutliner* InSceneOutliner, const TWeakObjectPtr<UWorld>& InSpecifiedWorldToDisplay = nullptr, bool bInHideComponents = true, bool bInHideLevelInstanceHierarchy = true, bool bInHideUnloadedActors = true, bool bInHideEmptyFolders = true)
 		: SpecifiedWorldToDisplay(InSpecifiedWorldToDisplay)
 		, SceneOutliner(InSceneOutliner)
 		, bHideComponents(bInHideComponents)
 		, bHideLevelInstanceHierarchy(bInHideLevelInstanceHierarchy)
 		, bHideUnloadedActors(bInHideUnloadedActors)
+		, bHideEmptyFolders(bInHideEmptyFolders)
 	{}
 
 	TWeakObjectPtr<UWorld> SpecifiedWorldToDisplay = nullptr;
@@ -37,6 +38,7 @@ struct SCENEOUTLINER_API FActorModeParams
 	bool bHideActorWithNoComponent = false;
 	bool bHideLevelInstanceHierarchy = true;
 	bool bHideUnloadedActors = true;
+	bool bHideEmptyFolders = true;
 };
 
 class SCENEOUTLINER_API FActorMode : public ISceneOutlinerMode
@@ -99,6 +101,8 @@ protected:
 	bool bHideLevelInstanceHierarchy;
 	/** Should unloaded actors be hidden */
 	bool bHideUnloadedActors;
+	/** Should empty folders be hidden */
+	bool bHideEmptyFolders;
 	/** Should the outliner scroll to the item on selection */
 	bool bAlwaysFrameSelection;
 };
