Actor Folders: don't allow option to make current folder as current.

#rnx

diff --git a/Engine/Source/Editor/SceneOutliner/Private/ActorBrowsingMode.cpp b/Engine/Source/Editor/SceneOutliner/Private/ActorBrowsingMode.cpp
--- a/Engine/Source/Editor/SceneOutliner/Private/ActorBrowsingMode.cpp
+++ b/Engine/Source/Editor/SceneOutliner/Private/ActorBrowsingMode.cpp
@@ -592,7 +592,9 @@ void FActorBrowsingMode::RegisterContextMenu()
 										FSceneOutlinerTreeItemPtr Item = Selection.SelectedItems[0].Pin();
 										if (FActorFolderTreeItem* FolderItem = Item->CastTo<FActorFolderTreeItem>())
 										{
-											return FolderItem->World.IsValid() && (FolderItem->World->GetCurrentLevel() == FolderItem->GetFolder().GetRootObjectAssociatedLevel());
+											return FolderItem->World.IsValid() && 
+												  (FolderItem->World->GetCurrentLevel() == FolderItem->GetFolder().GetRootObjectAssociatedLevel()) &&
+												  (FActorFolders::Get().GetActorEditorContextFolder(*FolderItem->World) != FolderItem->GetFolder());
 										}
 									}
 									return false;
