Fixed world outliner Clear Current Folder disabled menu item.

diff --git a/Engine/Source/Editor/SceneOutliner/Private/ActorBrowsingMode.cpp b/Engine/Source/Editor/SceneOutliner/Private/ActorBrowsingMode.cpp
--- a/Engine/Source/Editor/SceneOutliner/Private/ActorBrowsingMode.cpp
+++ b/Engine/Source/Editor/SceneOutliner/Private/ActorBrowsingMode.cpp
@@ -550,7 +550,7 @@ void FActorBrowsingMode::RegisterContextMenu()
 				}
 			}));
 
-		Menu->AddDynamicSection("DynamicActorEditorContext", FNewToolMenuDelegate::CreateLambda([this](UToolMenu* InMenu)
+		Menu->AddDynamicSection("DynamicActorEditorContext", FNewToolMenuDelegate::CreateLambda([](UToolMenu* InMenu)
 			{
 				USceneOutlinerMenuContext* Context = InMenu->FindContext<USceneOutlinerMenuContext>();
 				if (Context && Context->SceneOutliner.IsValid() && Context->bShowParentTree)
@@ -603,23 +603,26 @@ void FActorBrowsingMode::RegisterContextMenu()
 						);
 					}
 
+					const FActorBrowsingMode* Mode = static_cast<const FActorBrowsingMode*>(SceneOutliner->GetMode());
+					check(Mode);
+
 					Section.AddMenuEntry(
 						"ClearCurrentFolder",
 						LOCTEXT("ClearCurrentFolder", "Clear Current Folder"),
 						FText(),
 						FSlateIcon(),
 						FUIAction(
-							FExecuteAction::CreateLambda([this, SceneOutliner]()
+							FExecuteAction::CreateLambda([Mode]()
 							{
-								if (RepresentingWorld.IsValid())
+								if (Mode->RepresentingWorld.IsValid())
 								{
 									const FScopedTransaction Transaction(NSLOCTEXT("UnrealEd", "ClearCurrentActorFolder", "Clear Current Actor Folder"));
-									FActorFolders::Get().SetActorEditorContextFolder(*RepresentingWorld.Get(), FFolder::GetWorldRootFolder(RepresentingWorld.Get()));
+									FActorFolders::Get().SetActorEditorContextFolder(*Mode->RepresentingWorld.Get(), FFolder::GetWorldRootFolder(Mode->RepresentingWorld.Get()));
 								}
 							}),
-							FCanExecuteAction::CreateLambda([this]
-							{ 
-								return RepresentingWorld.IsValid() && !FActorFolders::Get().GetActorEditorContextFolder(*RepresentingWorld.Get()).IsNone(); 
+							FCanExecuteAction::CreateLambda([Mode]
+							{
+								return Mode->RepresentingWorld.IsValid() && !FActorFolders::Get().GetActorEditorContextFolder(*Mode->RepresentingWorld.Get()).IsNone();
 							})
 						)
 					);
@@ -635,7 +638,7 @@ void FActorBrowsingMode::RegisterContextMenu()
 
 TSharedPtr<SWidget> FActorBrowsingMode::BuildContextMenu()
 {
-	RegisterContextMenu();
+	FActorBrowsingMode::RegisterContextMenu();
 
 	const FSceneOutlinerItemSelection ItemSelection(SceneOutliner->GetSelection());
 
diff --git a/Engine/Source/Editor/SceneOutliner/Public/ActorBrowsingMode.h b/Engine/Source/Editor/SceneOutliner/Public/ActorBrowsingMode.h
--- a/Engine/Source/Editor/SceneOutliner/Public/ActorBrowsingMode.h
+++ b/Engine/Source/Editor/SceneOutliner/Public/ActorBrowsingMode.h
@@ -102,7 +102,7 @@ private:
 	/** Build and up the context menu */
 	TSharedPtr<SWidget> BuildContextMenu();
 	/** Register the context menu with the engine */
-	void RegisterContextMenu();
+	static void RegisterContextMenu();
 	bool CanPasteFoldersOnlyFromClipboard() const;
 
 	bool GetFolderNamesFromPayload(const FSceneOutlinerDragDropPayload& InPayload, TArray<FName>& OutFolders, FFolder::FRootObject& OutCommonRootObject) const;
