Outliner: Preserve item selection when the parent is collapsed

diff --git a/Engine/Source/Editor/SceneOutliner/Private/SSceneOutliner.cpp b/Engine/Source/Editor/SceneOutliner/Private/SSceneOutliner.cpp
--- a/Engine/Source/Editor/SceneOutliner/Private/SSceneOutliner.cpp
+++ b/Engine/Source/Editor/SceneOutliner/Private/SSceneOutliner.cpp
@@ -261,6 +261,9 @@ void SSceneOutliner::Construct(const FArguments& InArgs, const FSceneOutlinerIni
 
 			// Show the Hierarchy of actors pinned at the top of the tree view
 			.ShouldStackHierarchyHeaders(this, &SSceneOutliner::ShouldStackHierarchyHeaders)
+
+			// Preserve the selection when the selected item is hidden due to a parent collapsing
+			.AllowInvisibleItemSelection(true)
 		]
 	];
 
