Fix all folders being uninteractable in scene outliner while hiding empty folders.

diff --git a/Engine/Source/Editor/SceneOutliner/Private/ActorBrowsingMode.cpp b/Engine/Source/Editor/SceneOutliner/Private/ActorBrowsingMode.cpp
--- a/Engine/Source/Editor/SceneOutliner/Private/ActorBrowsingMode.cpp
+++ b/Engine/Source/Editor/SceneOutliner/Private/ActorBrowsingMode.cpp
@@ -441,7 +441,7 @@ TSharedRef<FSceneOutlinerFilter> FActorBrowsingMode::CreateHideUnloadedActorsFil
 TSharedRef<FSceneOutlinerFilter> FActorBrowsingMode::CreateHideEmptyFoldersFilter()
 {
 	return MakeShareable(new FFolderFilter(FFolderTreeItem::FFilterPredicate::CreateStatic(
-		[](const FFolder& Folder) { return false; }), FSceneOutlinerFilter::EDefaultBehaviour::Pass));
+		[](const FFolder& Folder) { return true; }), FSceneOutlinerFilter::EDefaultBehaviour::Pass));
 }
 
 static const FName DefaultContextBaseMenuName("SceneOutliner.DefaultContextMenuBase");
