Outliner: Add functionality to scroll to the current selection on pressing the 'F' key

diff --git a/Engine/Source/Editor/SceneOutliner/Private/ActorBrowsingMode.cpp b/Engine/Source/Editor/SceneOutliner/Private/ActorBrowsingMode.cpp
--- a/Engine/Source/Editor/SceneOutliner/Private/ActorBrowsingMode.cpp
+++ b/Engine/Source/Editor/SceneOutliner/Private/ActorBrowsingMode.cpp
@@ -35,6 +35,7 @@
 #include "LevelInstance/LevelInstanceSubsystem.h"
 #include "LevelInstance/LevelInstanceEditorInstanceActor.h"
 #include "EditorLevelUtils.h"
+#include "EditorViewportCommands.h"
 
 DEFINE_LOG_CATEGORY_STATIC(LogActorBrowser, Log, All);
 
@@ -995,6 +996,33 @@ FReply FActorBrowsingMode::OnKeyDown(const FKeyEvent& InKeyEvent)
 		return FReply::Handled();
 			
 	}
+
+	/* Allow the user to scroll to the current selection (and expand if needed) by pressing the key bound to
+	 * FEditorViewportCommands::Get().FocusViewportToSelection (Default: 'F')
+	 */
+	else
+	{
+		const FModifierKeysState ModifierKeys = FSlateApplication::Get().GetModifierKeys();
+		const FInputChord CheckChord( InKeyEvent.GetKey(), EModifierKey::FromBools(ModifierKeys.IsControlDown(), ModifierKeys.IsAltDown(), ModifierKeys.IsShiftDown(), ModifierKeys.IsCommandDown()) );
+
+		// Use the keyboard shortcut bound to 'Focus Viewport To Selection'
+		if (FEditorViewportCommands::Get().FocusViewportToSelection->HasActiveChord(CheckChord))
+		{
+			if (Selection.Num() == 1)
+			{
+				FSceneOutlinerTreeItemPtr ItemToFocus = Selection.SelectedItems[0].Pin();
+
+				if (ItemToFocus.IsValid())
+				{
+					SceneOutliner->ScrollItemIntoView(ItemToFocus);
+				}
+
+				// Return Unhandled here so that the level editor viewport can handle this event and focus the selected item
+				return FReply::Unhandled();
+			}
+		}
+	}
+	
 	return FReply::Unhandled();
 }
 
