PR #9310: Add source control actions for OFPA actors in Level Editor Outliner (Contributed by [REDACTED_PERSON])

fyi patrick.laflamme

diff --git a/Engine/Source/Editor/SceneOutliner/Private/ActorBrowsingMode.cpp b/Engine/Source/Editor/SceneOutliner/Private/ActorBrowsingMode.cpp
--- a/Engine/Source/Editor/SceneOutliner/Private/ActorBrowsingMode.cpp
+++ b/Engine/Source/Editor/SceneOutliner/Private/ActorBrowsingMode.cpp
@@ -36,6 +36,7 @@
 #include "LevelInstance/LevelInstanceEditorInstanceActor.h"
 #include "EditorLevelUtils.h"
 #include "EditorViewportCommands.h"
+#include "SceneOutlinerActorSCCColumn.h"
 
 DEFINE_LOG_CATEGORY_STATIC(LogActorBrowser, Log, All);
 
@@ -153,7 +154,7 @@ FActorBrowsingMode::FActorBrowsingMode(SSceneOutliner* InSceneOutliner, TWeakObj
 			}
 		});
 	FilterInfoMap.Add(TEXT("HideLevelInstancesFilter"), HideLevelInstancesInfo);
-	
+
 	FSceneOutlinerFilterInfo HideUnloadedActorsInfo(LOCTEXT("ToggleHideUnloadedActors", "Hide Unloaded Actors"), LOCTEXT("ToggleHideUnloadedActorsToolTip", "When enabled, hides all unloaded world partition actors."), LocalSettings.bHideUnloadedActors, FCreateSceneOutlinerFilter::CreateStatic(&FActorBrowsingMode::CreateHideUnloadedActorsFilter));
 	HideUnloadedActorsInfo.OnToggle().AddLambda([this] (bool bIsActive)
 		{
@@ -228,7 +229,7 @@ FActorBrowsingMode::~FActorBrowsingMode()
 	GEngine->OnLevelActorDeleted().RemoveAll(this);
 
 	GEditor->OnSelectUnloadedActorsEvent().RemoveAll(this);
-		
+
 	FEditorDelegates::OnEditCutActorsBegin.RemoveAll(this);
 	FEditorDelegates::OnEditCutActorsEnd.RemoveAll(this);
 	FEditorDelegates::OnEditCopyActorsBegin.RemoveAll(this);
@@ -251,7 +252,7 @@ void FActorBrowsingMode::Rebuild()
 			WorldPartition->OnActorDescRemovedEvent.RemoveAll(this);
 		}
 	}
-	
+
 	FActorMode::Rebuild();
 
 	FilteredActorCount = 0;
@@ -269,7 +270,7 @@ void FActorBrowsingMode::Rebuild()
 	}
 }
 
-FText FActorBrowsingMode::GetStatusText() const 
+FText FActorBrowsingMode::GetStatusText() const
 {
 	if (!RepresentingWorld.IsValid())
 	{
@@ -409,7 +410,7 @@ TSharedRef<FSceneOutlinerFilter> FActorBrowsingMode::CreateIsInCurrentLevelFilte
 TSharedRef<FSceneOutlinerFilter> FActorBrowsingMode::CreateHideComponentsFilter()
 {
 	return MakeShared<TSceneOutlinerPredicateFilter<FComponentTreeItem>>(TSceneOutlinerPredicateFilter<FComponentTreeItem>(
-		FComponentTreeItem::FFilterPredicate::CreateStatic([](const UActorComponent*) { return false; }), 
+		FComponentTreeItem::FFilterPredicate::CreateStatic([](const UActorComponent*) { return false; }),
 		FSceneOutlinerFilter::EDefaultBehaviour::Pass));
 }
 
@@ -548,6 +549,19 @@ void FActorBrowsingMode::RegisterContextMenu()
 							}
 						}
 					}
+
+					if (Context->NumSelectedItems > 0 && Context->SceneOutliner.IsValid())
+					{
+						SSceneOutliner* SceneOutliner = Context->SceneOutliner.Pin().Get();
+						TMap<FName, TSharedPtr<ISceneOutlinerColumn>> OutlinerColumns = SceneOutliner->GetColumns();
+						TSharedPtr<ISceneOutlinerColumn> * ColumnLookup = OutlinerColumns.Find("Source Control");
+						if (ColumnLookup != nullptr && ColumnLookup->IsValid())
+						{
+							TSharedPtr<FSceneOutlinerActorSCCColumn> SCCColumn = StaticCastSharedPtr<FSceneOutlinerActorSCCColumn>(*ColumnLookup);
+							TArray<FSceneOutlinerTreeItemPtr> SelectedItems = SceneOutliner->GetTree().GetSelectedItems();
+							SCCColumn->AddSourceControlMenuOptions(InMenu, SelectedItems);
+						}
+					}
 				}
 			}));
 
@@ -593,7 +607,7 @@ void FActorBrowsingMode::RegisterContextMenu()
 										FSceneOutlinerTreeItemPtr Item = Selection.SelectedItems[0].Pin();
 										if (FActorFolderTreeItem* FolderItem = Item->CastTo<FActorFolderTreeItem>())
 										{
-											return FolderItem->World.IsValid() && 
+											return FolderItem->World.IsValid() &&
 												  (FolderItem->World->GetCurrentLevel() == FolderItem->GetFolder().GetRootObjectAssociatedLevel()) &&
 												  (FActorFolders::Get().GetActorEditorContextFolder(*FolderItem->World) != FolderItem->GetFolder());
 										}
@@ -779,7 +793,7 @@ void FActorBrowsingMode::OnSelectUnloadedActors(const TArray<FGuid>& ActorGuids)
 			ItemsToSelect.Add(ItemPtr);
 		}
 	}
-	
+
 	if (ItemsToSelect.Num())
 	{
 		SceneOutliner->SetItemSelection(ItemsToSelect, true);
@@ -979,7 +993,7 @@ FReply FActorBrowsingMode::OnKeyDown(const FKeyEvent& InKeyEvent)
 	}
 
 	// Delete key: Delete selected actors (not rebindable, because it doesn't make much sense to bind.)
-	// Use Delete and Backspace instead of Platform_Delete because the LevelEditor default Edit Delete is bound to both 
+	// Use Delete and Backspace instead of Platform_Delete because the LevelEditor default Edit Delete is bound to both
 	else if (InKeyEvent.GetKey() == EKeys::Delete || InKeyEvent.GetKey() == EKeys::BackSpace)
 	{
 		if (SceneOutliner->GetSharedData().CustomDelete.IsBound())
@@ -994,7 +1008,7 @@ FReply FActorBrowsingMode::OnKeyDown(const FKeyEvent& InKeyEvent)
 			}
 		}
 		return FReply::Handled();
-			
+
 	}
 
 	/* Allow the user to scroll to the current selection (and expand if needed) by pressing the key bound to
@@ -1022,7 +1036,7 @@ FReply FActorBrowsingMode::OnKeyDown(const FKeyEvent& InKeyEvent)
 			}
 		}
 	}
-	
+
 	return FReply::Unhandled();
 }
 
@@ -1112,7 +1126,7 @@ TSharedPtr<FDragDropOperation> FActorBrowsingMode::CreateDragDropOperation(const
 	OutlinerOp->Construct();
 	return OutlinerOp;
 }
-	
+
 bool FActorBrowsingMode::ParseDragDrop(FSceneOutlinerDragDropPayload& OutPayload, const FDragDropOperation& Operation) const
 {
 	if (Operation.IsOfType<FSceneOutlinerDragDropOp>())
@@ -1142,7 +1156,7 @@ bool FActorBrowsingMode::ParseDragDrop(FSceneOutlinerDragDropPayload& OutPayload
 		}
 		return true;
 	}
-		
+
 	return false;
 }
 
@@ -1392,7 +1406,7 @@ FSceneOutlinerDragValidationInfo FActorBrowsingMode::ValidateDrop(const ISceneOu
 						}
 					}
 				}
-				
+
 				if (Actor->IsChildActor())
 				{
 					return FSceneOutlinerDragValidationInfo(ESceneOutlinerDropCompatibility::IncompatibleGeneric, FText::Format(LOCTEXT("Error_AttachChildActor", "Cannot move {0} as it is a child actor."), FText::FromString(Actor->GetActorLabel())));
@@ -1555,7 +1569,7 @@ void FActorBrowsingMode::OnDrop(ISceneOutlinerTreeItem& DropTarget, const FScene
 					PerformAttachment(NAME_None, DropActor, MoveTemp(DraggedActors));
 				}
 			}
-			
+
 		}
 		// Report errors
 		EditorErrors.Notify(NSLOCTEXT("ActorAttachmentError", "AttachmentsFailed", "Attachments Failed!"));
@@ -1595,14 +1609,14 @@ void FActorBrowsingMode::OnDrop(ISceneOutlinerTreeItem& DropTarget, const FScene
 
 					// If the folder root object changes, 1st pass will put actors at root. 2nd pass will set the destination path.
 					FName NewPath = (SrcFolder.GetRootObject() == DestinationPath.GetRootObject()) ? DestinationPath.GetPath() : NAME_None;
-					
+
 					Actor->SetFolderPath(NewPath);
 					FActorEditorUtils::TraverseActorTree_ParentFirst(Actor, [&](AActor* InActor) {
 						ChildActors.Add(InActor);
 						InActor->SetFolderPath(NewPath);
 						return true;
 						}, false);
-					
+
 					if ((Actor->GetFolderRootObject() != DestinationPath.GetRootObject()) && SrcFolder.IsRootObjectPersistentLevel() && (DestinationPath.IsRootObjectValid() && !DestinationPath.IsRootObjectPersistentLevel()))
 					{
 						MovingActorsToValidRootObject.Add(Actor);
@@ -1934,7 +1948,7 @@ void FActorBrowsingMode::UnpinItems(const TArray<FSceneOutlinerTreeItemPtr>& InI
 		WorldPartition->UnpinActors(ActorsToUnpin);
 
 		GEditor->GetSelectedActors()->EndBatchSelectOperation(/*bNotify=*/true);
-	}	
+	}
 }
 
 void FActorBrowsingMode::PinSelectedItems()
diff --git a/Engine/Source/Editor/SceneOutliner/Private/SSourceControlWidget.cpp b/Engine/Source/Editor/SceneOutliner/Private/SSourceControlWidget.cpp
new file mode 100644
--- /dev/null
+++ b/Engine/Source/Editor/SceneOutliner/Private/SSourceControlWidget.cpp
@@ -0,0 +1,169 @@
+// Copyright Epic Games, Inc. All Rights Reserved.
+
+#include "SSourceControlWidget.h"
+
+#include "ActorTreeItem.h"
+#include "ActorDescTreeItem.h"
+#include "ActorFolderTreeItem.h"
+#include "ISourceControlProvider.h"
+#include "ISourceControlModule.h"
+#include "WorldPartition/WorldPartitionActorDesc.h"
+
+void SSourceControlWidget::Construct(const FArguments& InArgs, TWeakPtr<ISceneOutliner> InWeakOutliner, TWeakPtr<ISceneOutlinerTreeItem> InWeakTreeItem)
+{
+	WeakTreeItem = InWeakTreeItem;
+	WeakOutliner = InWeakOutliner;
+
+	SImage::Construct(
+		SImage::FArguments()
+		.ColorAndOpacity(this, &SSourceControlWidget::GetForegroundColor)
+		.Image(FStyleDefaults::GetNoBrush()));
+
+	FSceneOutlinerTreeItemPtr TreeItemPtr = WeakTreeItem.Pin();
+	if (TreeItemPtr.IsValid())
+	{
+		if (FActorTreeItem* ActorItem = TreeItemPtr->CastTo<FActorTreeItem>())
+		{
+			if (AActor* Actor = ActorItem->Actor.Get())
+			{
+				if (Actor->IsPackageExternal())
+				{
+					ExternalPackageName = USourceControlHelpers::PackageFilename(Actor->GetExternalPackage());
+					ExternalPackage = Actor->GetExternalPackage();
+				}
+
+				ActorPackingModeChangedDelegateHandle = Actor->OnPackagingModeChanged.AddLambda([this](AActor* InActor, bool bExternal)
+				{
+					if (bExternal)
+					{
+						ExternalPackageName = USourceControlHelpers::PackageFilename(InActor->GetExternalPackage());
+						ExternalPackage = InActor->GetExternalPackage();
+						ConnectSourceControl();
+					}
+					else
+					{
+						ExternalPackageName = FString();
+						ExternalPackage = nullptr;
+						DisconnectSourceControl();
+					}
+				});
+			}
+		}
+		else if (FActorFolderTreeItem* ActorFolderItem = TreeItemPtr->CastTo<FActorFolderTreeItem>())
+		{
+			if (const UActorFolder* ActorFolder = ActorFolderItem->GetActorFolder())
+			{
+				if (ActorFolder->IsPackageExternal())
+				{
+					ExternalPackageName = USourceControlHelpers::PackageFilename(ActorFolder->GetExternalPackage());
+					ExternalPackage = ActorFolder->GetExternalPackage();
+				}
+			}
+		}
+		else if (FActorDescTreeItem* ActorDescItem = TreeItemPtr->CastTo<FActorDescTreeItem>())
+		{
+			if (const FWorldPartitionActorDesc* ActorDesc = ActorDescItem->ActorDescHandle.Get())
+			{
+				ExternalPackageName = USourceControlHelpers::PackageFilename(ActorDesc->GetActorPackage().ToString());
+				ExternalPackage = FindPackage(nullptr, *ActorDesc->GetActorPackage().ToString());
+			}
+		}
+	}
+
+	if (!ExternalPackageName.IsEmpty())
+	{
+		ConnectSourceControl();
+	}
+}
+
+SSourceControlWidget::~SSourceControlWidget()
+{
+	DisconnectSourceControl();
+}
+
+FReply SSourceControlWidget::OnMouseButtonDoubleClick(const FGeometry& InMyGeometry, const FPointerEvent& InMouseEvent)
+{
+	FSourceControlStatePtr SourceControlState = ISourceControlModule::Get().GetProvider().GetState(ExternalPackageName, EStateCacheUsage::ForceUpdate);
+	if (SourceControlState.IsValid())
+	{
+		UpdateSourceControlStateIcon(SourceControlState);
+
+	}
+	return FReply::Handled();
+}
+
+FSourceControlStatePtr SSourceControlWidget::GetSourceControlState()
+{
+	return ISourceControlModule::Get().GetProvider().GetState(ExternalPackageName, EStateCacheUsage::Use);
+}
+
+void SSourceControlWidget::ConnectSourceControl()
+{
+	check(!ExternalPackageName.IsEmpty());
+
+	ISourceControlModule& SCCModule = ISourceControlModule::Get();
+	SourceControlProviderChangedDelegateHandle = SCCModule.RegisterProviderChanged(FSourceControlProviderChanged::FDelegate::CreateSP(this, &SSourceControlWidget::HandleSourceControlProviderChanged));
+	SourceControlStateChangedDelegateHandle = SCCModule.GetProvider().RegisterSourceControlStateChanged_Handle(FSourceControlStateChanged::FDelegate::CreateSP(this, &SSourceControlWidget::HandleSourceControlStateChanged, EStateCacheUsage::Use));
+
+	// Check if there is already a cached state for this item
+	FSourceControlStatePtr SourceControlState = ISourceControlModule::Get().GetProvider().GetState(ExternalPackageName, EStateCacheUsage::Use);
+	if (SourceControlState.IsValid() && !SourceControlState->IsUnknown())
+	{
+		UpdateSourceControlStateIcon(SourceControlState);
+	}
+	else
+	{
+		SCCModule.QueueStatusUpdate(ExternalPackageName);
+	}
+}
+
+void SSourceControlWidget::DisconnectSourceControl()
+{
+	FSceneOutlinerTreeItemPtr TreeItemPtr = WeakTreeItem.Pin();
+	if (TreeItemPtr.IsValid())
+	{
+		if (FActorTreeItem* ActorItem = TreeItemPtr->CastTo<FActorTreeItem>())
+		{
+			if (AActor* Actor = ActorItem->Actor.Get())
+			{
+				Actor->OnPackagingModeChanged.Remove(ActorPackingModeChangedDelegateHandle);
+			}
+		}
+	}
+	ISourceControlModule::Get().GetProvider().UnregisterSourceControlStateChanged_Handle(SourceControlStateChangedDelegateHandle);
+	ISourceControlModule::Get().UnregisterProviderChanged(SourceControlProviderChangedDelegateHandle);
+}
+
+void SSourceControlWidget::HandleSourceControlStateChanged(EStateCacheUsage::Type CacheUsage)
+{
+	FSourceControlStatePtr SourceControlState = ISourceControlModule::Get().GetProvider().GetState(ExternalPackageName, CacheUsage);
+	if (SourceControlState.IsValid())
+	{
+		UpdateSourceControlStateIcon(SourceControlState);
+	}
+}
+
+void SSourceControlWidget::HandleSourceControlProviderChanged(ISourceControlProvider& OldProvider, ISourceControlProvider& NewProvider)
+{
+	OldProvider.UnregisterSourceControlStateChanged_Handle(SourceControlStateChangedDelegateHandle);
+	SourceControlStateChangedDelegateHandle = NewProvider.RegisterSourceControlStateChanged_Handle(FSourceControlStateChanged::FDelegate::CreateSP(this, &SSourceControlWidget::HandleSourceControlStateChanged, EStateCacheUsage::Use));
+	
+	UpdateSourceControlStateIcon(nullptr);
+
+	ISourceControlModule::Get().QueueStatusUpdate(ExternalPackageName);
+}
+
+void SSourceControlWidget::UpdateSourceControlStateIcon(FSourceControlStatePtr SourceControlState)
+{
+	if(SourceControlState.IsValid())
+	{
+		FSlateIcon Icon = SourceControlState->GetIcon();
+		
+		SetFromSlateIcon(Icon);
+	}
+	else
+	{
+		SetImage(nullptr);
+		RemoveAllLayers();
+	}
+}
\ No newline at end of file
diff --git a/Engine/Source/Editor/SceneOutliner/Private/SSourceControlWidget.h b/Engine/Source/Editor/SceneOutliner/Private/SSourceControlWidget.h
new file mode 100644
--- /dev/null
+++ b/Engine/Source/Editor/SceneOutliner/Private/SSourceControlWidget.h
@@ -0,0 +1,59 @@
+// Copyright Epic Games, Inc. All Rights Reserved.
+
+#pragma once
+
+#include "ISceneOutliner.h"
+#include "SceneOutlinerPublicTypes.h"
+#include "Widgets/Images/SLayeredImage.h"
+#include "SourceControlHelpers.h"
+
+class SSourceControlWidget : public SLayeredImage
+{
+public:
+	SLATE_BEGIN_ARGS(SSourceControlWidget) {}
+	SLATE_END_ARGS()
+
+	/** Construct this widget */
+	void Construct(const FArguments& InArgs, TWeakPtr<ISceneOutliner> InWeakOutliner, TWeakPtr<ISceneOutlinerTreeItem> InWeakTreeItem);
+
+	~SSourceControlWidget();
+
+	FSourceControlStatePtr GetSourceControlState();
+
+	FString GetPackageName() { return ExternalPackageName; }
+
+	UPackage* GetPackage() { return ExternalPackage; }
+
+private:
+
+	virtual FReply OnMouseButtonDoubleClick(const FGeometry& InMyGeometry, const FPointerEvent& InMouseEvent) override;
+
+	void ConnectSourceControl();
+
+	void DisconnectSourceControl();
+
+	void HandleSourceControlStateChanged(EStateCacheUsage::Type CacheUsage);
+
+	void HandleSourceControlProviderChanged(ISourceControlProvider& OldProvider, ISourceControlProvider& NewProvider);
+
+	void UpdateSourceControlStateIcon(FSourceControlStatePtr SourceControlState);
+
+	/** The tree item we relate to */
+	TWeakPtr<ISceneOutlinerTreeItem> WeakTreeItem;
+
+	/** Reference back to the outliner so we can set visibility of a whole selection */
+	TWeakPtr<ISceneOutliner> WeakOutliner;
+
+	/** Cache the items external package name */
+	FString ExternalPackageName;
+	UPackage* ExternalPackage = nullptr;
+
+	/** Source control state changed delegate handle */
+	FDelegateHandle SourceControlStateChangedDelegateHandle;
+
+	/** Source control provider changed delegate handle */
+	FDelegateHandle SourceControlProviderChangedDelegateHandle;
+
+	/** Actor packaging mode changed delegate handle */
+	FDelegateHandle ActorPackingModeChangedDelegateHandle;
+};
diff --git a/Engine/Source/Editor/SceneOutliner/Private/SceneOutlinerActorSCCColumn.cpp b/Engine/Source/Editor/SceneOutliner/Private/SceneOutlinerActorSCCColumn.cpp
--- a/Engine/Source/Editor/SceneOutliner/Private/SceneOutlinerActorSCCColumn.cpp
+++ b/Engine/Source/Editor/SceneOutliner/Private/SceneOutlinerActorSCCColumn.cpp
@@ -7,238 +7,361 @@
 #include "ActorDescTreeItem.h"
 #include "ActorFolderTreeItem.h"
 #include "WorldPartition/WorldPartitionActorDesc.h"
-#include "ISourceControlProvider.h"
 #include "ISourceControlModule.h"
-#include "SourceControlHelpers.h"
-#include "Widgets/Images/SLayeredImage.h"
+#include "SourceControlOperations.h"
+#include "FileHelpers.h"
+#include "SourceControlWindows.h"
+#include "AssetViewUtils.h"
+#include "Misc/MessageDialog.h"
 
 #define LOCTEXT_NAMESPACE "SceneOutlinerSourceControlColumn"
 
-class SSourceControlWidget : public SLayeredImage
+FName FSceneOutlinerActorSCCColumn::GetColumnID()
+{
+	return GetID();
+}
+
+SHeaderRow::FColumn::FArguments FSceneOutlinerActorSCCColumn::ConstructHeaderRowColumn()
+{
+	return SHeaderRow::Column(GetColumnID())
+		.FixedWidth(24.f)
+		.HAlignHeader(HAlign_Center)
+		.VAlignHeader(VAlign_Center)
+		.HAlignCell(HAlign_Center)
+		.VAlignCell(VAlign_Center)
+		.DefaultTooltip(FText::FromName(GetColumnID()))
+		[
+			SNew(SImage)
+			.ColorAndOpacity(FSlateColor::UseForeground())
+			.Image(this, &FSceneOutlinerActorSCCColumn::GetHeaderIcon)
+		];
+}
+
+const TSharedRef<SWidget> FSceneOutlinerActorSCCColumn::ConstructRowWidget(FSceneOutlinerTreeItemRef TreeItem, const STableRow<FSceneOutlinerTreeItemPtr>& Row)
+{
+	if (TreeItem->IsA<FActorTreeItem>() || 
+		TreeItem->IsA<FActorDescTreeItem>() || 
+		(TreeItem->IsA<FActorFolderTreeItem>() && TreeItem->CastTo<FActorFolderTreeItem>()->GetActorFolder()))
+	{
+		TSharedRef<SSourceControlWidget> Widget = SNew(SSourceControlWidget, WeakSceneOutliner, TreeItem);
+
+		ItemWidgets.Add(FSceneOutlinerTreeItemPtr(TreeItem), Widget);
+		
+		return SNew(SHorizontalBox)
+			+ SHorizontalBox::Slot()
+			.AutoWidth()
+			.VAlign(VAlign_Center)
+			[
+				Widget
+			];
+	}
+	return SNullWidget::NullWidget;
+}
+
+const FSlateBrush* FSceneOutlinerActorSCCColumn::GetHeaderIcon() const
+{
+	if (ISourceControlModule::Get().IsEnabled())
+	{
+		return FAppStyle::GetBrush("SourceControl.StatusIcon.On");
+	}
+	else
+	{
+		return FAppStyle::GetBrush("SourceControl.StatusIcon.Off");
+	}
+}
+
+bool FSceneOutlinerActorSCCColumn::AddSourceControlMenuOptions(UToolMenu* Menu, TArray<FSceneOutlinerTreeItemPtr> InSelectedItems)
 {
-public:
-	SLATE_BEGIN_ARGS(SSourceControlWidget) {}
-	SLATE_END_ARGS()
+	SelectedItems = InSelectedItems;
+
+	CacheCanExecuteVars();
 
-	/** Construct this widget */
-	void Construct(const FArguments& InArgs, TWeakPtr<ISceneOutliner> InWeakOutliner, TWeakPtr<ISceneOutlinerTreeItem> InWeakTreeItem)
+	FToolMenuSection& Section = Menu->AddSection("AssetContextSourceControl");
+	
+	if ( ISourceControlModule::Get().IsEnabled() )
 	{
-		WeakTreeItem = InWeakTreeItem;
-		WeakOutliner = InWeakOutliner;
+		// SCC sub menu
+		Section.AddSubMenu(
+			"SourceControlSubMenu",
+			LOCTEXT("SourceControlSubMenuLabel", "Source Control"),
+			LOCTEXT("SourceControlSubMenuToolTip", "Source control actions."),
+			FNewToolMenuDelegate::CreateSP(this, &FSceneOutlinerActorSCCColumn::FillSourceControlSubMenu),
+			FUIAction(
+				FExecuteAction(),
+				FCanExecuteAction::CreateSP( this, &FSceneOutlinerActorSCCColumn::CanExecuteSourceControlActions )
+				),
+			EUserInterfaceActionType::Button,
+			false,
+			FSlateIcon(FAppStyle::GetAppStyleSetName(), "SourceControl.StatusIcon.On")
+			);
+	}
+
+	return true;
+}
 
-		SImage::Construct(
-			SImage::FArguments()
-			.ColorAndOpacity(this, &SSourceControlWidget::GetForegroundColor)
-			.Image(FStyleDefaults::GetNoBrush()));
 
-		FSceneOutlinerTreeItemPtr TreeItemPtr = WeakTreeItem.Pin();
-		if (TreeItemPtr.IsValid())
+bool FSceneOutlinerActorSCCColumn::CanExecuteSourceControlActions() const
+{
+	return SelectedItems.Num() > 0;
+}
+
+void FSceneOutlinerActorSCCColumn::CacheCanExecuteVars()
+{
+	bCanExecuteSCCCheckOut = false;
+	bCanExecuteSCCCheckIn = false;
+	bCanExecuteSCCHistory = false;
+	bCanExecuteSCCRevert = false;
+
+	ISourceControlProvider& SourceControlProvider = ISourceControlModule::Get().GetProvider();
+	if ( ISourceControlModule::Get().IsEnabled() )
+	{
+		for (FSceneOutlinerTreeItemPtr SelectedItem : SelectedItems)
 		{
-			if (FActorTreeItem* ActorItem = TreeItemPtr->CastTo<FActorTreeItem>())
+			TSharedRef<SSourceControlWidget> * WidgetLookup = ItemWidgets.Find(SelectedItem);
+			if (WidgetLookup == nullptr)
 			{
-				if (AActor* Actor = ActorItem->Actor.Get())
-				{
-					if (Actor->IsPackageExternal())
-					{
-						ExternalPackageName = USourceControlHelpers::PackageFilename(Actor->GetExternalPackage());
-					}
-
-					ActorPackingModeChangedDelegateHandle = Actor->OnPackagingModeChanged.AddLambda([this](AActor* InActor, bool bExternal)
-						{
-							if (bExternal)
-							{
-								ExternalPackageName = USourceControlHelpers::PackageFilename(InActor->GetExternalPackage());
-								ConnectSourceControl();
-							}
-							else
-							{
-								ExternalPackageName = FString();
-								DisconnectSourceControl();
-							}
-						});
-				}
+				continue;
 			}
-			else if (FActorFolderTreeItem* ActorFolderItem = TreeItemPtr->CastTo<FActorFolderTreeItem>())
+
+			TSharedRef<SSourceControlWidget> Widget = *WidgetLookup;
+
+			// Check the SCC state for each package in the selected paths
+			FSourceControlStatePtr SourceControlState = Widget->GetSourceControlState();
+			if(SourceControlState.IsValid())
 			{
-				if (const UActorFolder* ActorFolder = ActorFolderItem->GetActorFolder())
+				if ( SourceControlState->CanCheckout() )
 				{
-					if (ActorFolder->IsPackageExternal())
-					{
-						ExternalPackageName = USourceControlHelpers::PackageFilename(ActorFolder->GetExternalPackage());
-					}
+					bCanExecuteSCCCheckOut = true;
 				}
-			}
-			else if (FActorDescTreeItem* ActorDescItem = TreeItemPtr->CastTo<FActorDescTreeItem>())
-			{
-				if (const FWorldPartitionActorDesc* ActorDesc = ActorDescItem->ActorDescHandle.Get())
+
+				if( SourceControlState->IsSourceControlled() && !SourceControlState->IsAdded() )
+				{
+					bCanExecuteSCCHistory = true;
+				}
+
+				if ( SourceControlState->CanCheckIn() )
 				{
-					ExternalPackageName =  USourceControlHelpers::PackageFilename(ActorDesc->GetActorPackage().ToString());
+					bCanExecuteSCCCheckIn = true;
+				}
+
+				if (SourceControlState->CanRevert())
+				{
+					bCanExecuteSCCRevert = true;
 				}
 			}
-		}
 
-		if (!ExternalPackageName.IsEmpty())
-		{
-			ConnectSourceControl();
+			if ( bCanExecuteSCCCheckOut
+				&& bCanExecuteSCCCheckIn
+				&& bCanExecuteSCCHistory
+				&& bCanExecuteSCCRevert
+				)
+			{
+				// All options are available, no need to keep iterating
+				break;
+			}
 		}
 	}
+}
 
-	~SSourceControlWidget()
-	{
-		DisconnectSourceControl();
-	}
+bool FSceneOutlinerActorSCCColumn::CanExecuteSCCCheckOut() const
+{
+	return bCanExecuteSCCCheckOut;
+}
 
-private:
+bool FSceneOutlinerActorSCCColumn::CanExecuteSCCCheckIn() const
+{
+	return bCanExecuteSCCCheckIn;
+}
 
-	virtual FReply OnMouseButtonDoubleClick(const FGeometry& InMyGeometry, const FPointerEvent& InMouseEvent) override
-	{
-		FSourceControlStatePtr SourceControlState = ISourceControlModule::Get().GetProvider().GetState(ExternalPackageName, EStateCacheUsage::ForceUpdate);
-		if (SourceControlState.IsValid())
-		{
-			UpdateSourceControlStateIcon(SourceControlState);
+bool FSceneOutlinerActorSCCColumn::CanExecuteSCCHistory() const
+{
+	return bCanExecuteSCCHistory;
+}
 
-		}
-		return FReply::Handled();
-	}
+bool FSceneOutlinerActorSCCColumn::CanExecuteSCCRevert() const
+{
+	return bCanExecuteSCCRevert;
+}
 
-	void ConnectSourceControl()
-	{
-		check(!ExternalPackageName.IsEmpty());
+bool FSceneOutlinerActorSCCColumn::CanExecuteSCCRefresh() const
+{
+	return ISourceControlModule::Get().IsEnabled();
+}
 
-		ISourceControlModule& SCCModule = ISourceControlModule::Get();
-		SourceControlProviderChangedDelegateHandle = SCCModule.RegisterProviderChanged(FSourceControlProviderChanged::FDelegate::CreateSP(this, &SSourceControlWidget::HandleSourceControlProviderChanged));
-		SourceControlStateChangedDelegateHandle = SCCModule.GetProvider().RegisterSourceControlStateChanged_Handle(FSourceControlStateChanged::FDelegate::CreateSP(this, &SSourceControlWidget::HandleSourceControlStateChanged, EStateCacheUsage::Use));
+void FSceneOutlinerActorSCCColumn::FillSourceControlSubMenu(UToolMenu* Menu)
+{
+	FToolMenuSection& Section = Menu->AddSection("AssetSourceControlActions", LOCTEXT("AssetSourceControlActionsMenuHeading", "Source Control"));
 
-		// Check if there is already a cached state for this item
-		FSourceControlStatePtr SourceControlState = ISourceControlModule::Get().GetProvider().GetState(ExternalPackageName, EStateCacheUsage::Use);
-		if (SourceControlState.IsValid() && !SourceControlState->IsUnknown())
-		{
-			UpdateSourceControlStateIcon(SourceControlState);
-		}
-		else
-		{
-			SCCModule.QueueStatusUpdate(ExternalPackageName);
-		}
+	if ( CanExecuteSCCCheckOut() )
+	{
+		Section.AddMenuEntry(
+			"SCCCheckOut",
+			LOCTEXT("SCCCheckOut", "Check Out"),
+			LOCTEXT("SCCCheckOutTooltip", "Checks out the selected asset from source control."),
+			FSlateIcon(FAppStyle::GetAppStyleSetName(), "SourceControl.Actions.CheckOut"),
+			FUIAction(
+				FExecuteAction::CreateSP( this, &FSceneOutlinerActorSCCColumn::ExecuteSCCCheckOut ),
+				FCanExecuteAction::CreateSP( this, &FSceneOutlinerActorSCCColumn::CanExecuteSCCCheckOut )
+			)
+		);
 	}
 
-	void DisconnectSourceControl()
+	if ( CanExecuteSCCCheckIn() )
 	{
-		FSceneOutlinerTreeItemPtr TreeItemPtr = WeakTreeItem.Pin();
-		if (TreeItemPtr.IsValid())
-		{
-			if (FActorTreeItem* ActorItem = TreeItemPtr->CastTo<FActorTreeItem>())
-			{
-				if (AActor* Actor = ActorItem->Actor.Get())
-				{
-					Actor->OnPackagingModeChanged.Remove(ActorPackingModeChangedDelegateHandle);
-				}
-			}
-		}
-		ISourceControlModule::Get().GetProvider().UnregisterSourceControlStateChanged_Handle(SourceControlStateChangedDelegateHandle);
-		ISourceControlModule::Get().UnregisterProviderChanged(SourceControlProviderChangedDelegateHandle);
+		Section.AddMenuEntry(
+			"SCCCheckIn",
+			LOCTEXT("SCCCheckIn", "Check In"),
+			LOCTEXT("SCCCheckInTooltip", "Checks in the selected asset to source control."),
+			FSlateIcon(FAppStyle::GetAppStyleSetName(), "SourceControl.Actions.Submit"),
+			FUIAction(
+				FExecuteAction::CreateSP( this, &FSceneOutlinerActorSCCColumn::ExecuteSCCCheckIn ),
+				FCanExecuteAction::CreateSP( this, &FSceneOutlinerActorSCCColumn::CanExecuteSCCCheckIn )
+			)
+		);
 	}
 
-	void HandleSourceControlStateChanged(EStateCacheUsage::Type CacheUsage)
+	Section.AddMenuEntry(
+		"SCCRefresh",
+		LOCTEXT("SCCRefresh", "Refresh"),
+		LOCTEXT("SCCRefreshTooltip", "Updates the source control status of the asset."),
+		FSlateIcon(FAppStyle::GetAppStyleSetName(), "SourceControl.Actions.Refresh"),
+		FUIAction(
+			FExecuteAction::CreateSP( this, &FSceneOutlinerActorSCCColumn::ExecuteSCCRefresh ),
+			FCanExecuteAction::CreateSP( this, &FSceneOutlinerActorSCCColumn::CanExecuteSCCRefresh )
+			)
+		);
+
+	if( CanExecuteSCCHistory() )
 	{
-		FSourceControlStatePtr SourceControlState = ISourceControlModule::Get().GetProvider().GetState(ExternalPackageName, CacheUsage);
-		if (SourceControlState.IsValid())
-		{
-			UpdateSourceControlStateIcon(SourceControlState);
-		}
+		Section.AddMenuEntry(
+			"SCCHistory",
+			LOCTEXT("SCCHistory", "History"),
+			LOCTEXT("SCCHistoryTooltip", "Displays the source control revision history of the selected asset."),
+			FSlateIcon(FAppStyle::GetAppStyleSetName(), "SourceControl.Actions.History"),
+			FUIAction(
+				FExecuteAction::CreateSP( this, &FSceneOutlinerActorSCCColumn::ExecuteSCCHistory ),
+				FCanExecuteAction::CreateSP( this, &FSceneOutlinerActorSCCColumn::CanExecuteSCCHistory )
+			)
+		);
 	}
 
-	void HandleSourceControlProviderChanged(ISourceControlProvider& OldProvider, ISourceControlProvider& NewProvider)
+	if( CanExecuteSCCRevert() )
 	{
-		OldProvider.UnregisterSourceControlStateChanged_Handle(SourceControlStateChangedDelegateHandle);
-		SourceControlStateChangedDelegateHandle = NewProvider.RegisterSourceControlStateChanged_Handle(FSourceControlStateChanged::FDelegate::CreateSP(this, &SSourceControlWidget::HandleSourceControlStateChanged, EStateCacheUsage::Use));
-		
-		UpdateSourceControlStateIcon(nullptr);
-
-		ISourceControlModule::Get().QueueStatusUpdate(ExternalPackageName);
+		Section.AddMenuEntry(
+			"SCCRevert",
+			LOCTEXT("SCCRevert", "Revert"),
+			LOCTEXT("SCCRevertTooltip", "Reverts the asset to the state it was before it was checked out."),
+			FSlateIcon(FAppStyle::GetAppStyleSetName(), "SourceControl.Actions.Revert"),
+			FUIAction(
+				FExecuteAction::CreateSP( this, &FSceneOutlinerActorSCCColumn::ExecuteSCCRevert ),
+				FCanExecuteAction::CreateSP( this, &FSceneOutlinerActorSCCColumn::CanExecuteSCCRevert )
+			)
+		);
 	}
+}
 
-	void UpdateSourceControlStateIcon(FSourceControlStatePtr SourceControlState)
+void FSceneOutlinerActorSCCColumn::GetSelectedPackageNames(TArray<FString>& OutPackageNames) const
+{
+	for (FSceneOutlinerTreeItemPtr SelectedItem : SelectedItems)
 	{
-		if(SourceControlState.IsValid())
+		const TSharedRef<SSourceControlWidget> * WidgetLookup = ItemWidgets.Find(SelectedItem);
+		if (WidgetLookup == nullptr)
 		{
-			FSlateIcon Icon = SourceControlState->GetIcon();
-			
-			SetFromSlateIcon(Icon);
+			continue;
 		}
-		else
-		{
-			SetImage(nullptr);
-			RemoveAllLayers();
+
+		TSharedRef<SSourceControlWidget> Widget = *WidgetLookup;
+		FString PackageName = Widget->GetPackageName();
+		if (!PackageName.IsEmpty()) {
+			OutPackageNames.Add(PackageName);
 		}
 	}
+}
 
-	/** The tree item we relate to */
-	TWeakPtr<ISceneOutlinerTreeItem> WeakTreeItem;
-
-	/** Reference back to the outliner so we can set visibility of a whole selection */
-	TWeakPtr<ISceneOutliner> WeakOutliner;
-
-	/** Cache the items external package name */
-	FString ExternalPackageName;
-
-	/** Source control state changed delegate handle */
-	FDelegateHandle SourceControlStateChangedDelegateHandle;
-
-	/** Source control provider changed delegate handle */
-	FDelegateHandle SourceControlProviderChangedDelegateHandle;
-
-	/** Actor packaging mode changed delegate handle */
-	FDelegateHandle ActorPackingModeChangedDelegateHandle;
-};
-
-FName FSceneOutlinerActorSCCColumn::GetColumnID()
+void FSceneOutlinerActorSCCColumn::GetSelectedPackages(TArray<UPackage*>& OutPackages) const
 {
-	return GetID();
+	for (FSceneOutlinerTreeItemPtr SelectedItem : SelectedItems)
+	{
+		const TSharedRef<SSourceControlWidget> * WidgetLookup = ItemWidgets.Find(SelectedItem);
+		if (WidgetLookup == nullptr)
+		{
+			continue;
+		}
+
+		TSharedRef<SSourceControlWidget> Widget = *WidgetLookup;
+		UPackage* Package = Widget->GetPackage();
+		if (Package != nullptr) {
+			OutPackages.Add(Package);
+		}
+	}
 }
 
-SHeaderRow::FColumn::FArguments FSceneOutlinerActorSCCColumn::ConstructHeaderRowColumn()
+void FSceneOutlinerActorSCCColumn::ExecuteSCCRefresh()
 {
-	return SHeaderRow::Column(GetColumnID())
-		.FixedWidth(24.f)
-		.HAlignHeader(HAlign_Center)
-		.VAlignHeader(VAlign_Center)
-		.HAlignCell(HAlign_Center)
-		.VAlignCell(VAlign_Center)
-		.DefaultTooltip(FText::FromName(GetColumnID()))
-		[
-			SNew(SImage)
-			.ColorAndOpacity(FSlateColor::UseForeground())
-			.Image(this, &FSceneOutlinerActorSCCColumn::GetHeaderIcon)
-		];
+	TArray<FString> PackageNames;
+	GetSelectedPackageNames(PackageNames);
+
+	ISourceControlModule::Get().GetProvider().Execute(ISourceControlOperation::Create<FUpdateStatus>(), SourceControlHelpers::PackageFilenames(PackageNames), EConcurrency::Asynchronous);
 }
 
-const TSharedRef<SWidget> FSceneOutlinerActorSCCColumn::ConstructRowWidget(FSceneOutlinerTreeItemRef TreeItem, const STableRow<FSceneOutlinerTreeItemPtr>& Row)
+void FSceneOutlinerActorSCCColumn::ExecuteSCCCheckOut()
 {
-	if (TreeItem->IsA<FActorTreeItem>() || 
-		TreeItem->IsA<FActorDescTreeItem>() || 
-		(TreeItem->IsA<FActorFolderTreeItem>() && TreeItem->CastTo<FActorFolderTreeItem>()->GetActorFolder()))
+	TArray<UPackage*> PackagesToCheckOut;
+	GetSelectedPackages(PackagesToCheckOut);
+
+	if ( PackagesToCheckOut.Num() > 0 )
 	{
-		return SNew(SHorizontalBox)
-			+ SHorizontalBox::Slot()
-			.AutoWidth()
-			.VAlign(VAlign_Center)
-			[
-				SNew(SSourceControlWidget, WeakSceneOutliner, TreeItem)
-			];
+		FEditorFileUtils::CheckoutPackages(PackagesToCheckOut);
 	}
-	return SNullWidget::NullWidget;
 }
 
-const FSlateBrush* FSceneOutlinerActorSCCColumn::GetHeaderIcon() const
+void FSceneOutlinerActorSCCColumn::ExecuteSCCCheckIn()
 {
-	if (ISourceControlModule::Get().IsEnabled())
+	TArray<UPackage*> Packages;
+	GetSelectedPackages(Packages);
+
+	// Prompt the user to ask if they would like to first save any dirty packages they are trying to check-in
+	const FEditorFileUtils::EPromptReturnCode UserResponse = FEditorFileUtils::PromptForCheckoutAndSave( Packages, true, true );
+
+	// If the user elected to save dirty packages, but one or more of the packages failed to save properly OR if the user
+	// canceled out of the prompt, don't follow through on the check-in process
+	const bool bShouldProceed = ( UserResponse == FEditorFileUtils::EPromptReturnCode::PR_Success || UserResponse == FEditorFileUtils::EPromptReturnCode::PR_Declined );
+	if ( bShouldProceed )
 	{
-		return FAppStyle::GetBrush("SourceControl.StatusIcon.On");
+		TArray<FString> PackageNames;
+		GetSelectedPackageNames(PackageNames);
+
+		const bool bUseSourceControlStateCache = true;
+		const bool bCheckinGood = FSourceControlWindows::PromptForCheckin(bUseSourceControlStateCache, PackageNames);
+
+		if (!bCheckinGood)
+		{
+			FMessageDialog::Open(EAppMsgType::Ok, NSLOCTEXT("UnrealEd", "SCC_Checkin_Failed", "Check-in failed as a result of save failure."));
+		}
 	}
 	else
 	{
-		return FAppStyle::GetBrush("SourceControl.StatusIcon.Off");
+		// If a failure occurred, alert the user that the check-in was aborted. This warning shouldn't be necessary if the user cancelled
+		// from the dialog, because they obviously intended to cancel the whole operation.
+		if ( UserResponse == FEditorFileUtils::EPromptReturnCode::PR_Failure )
+		{
+			FMessageDialog::Open( EAppMsgType::Ok, NSLOCTEXT("UnrealEd", "SCC_Checkin_Aborted", "Check-in aborted as a result of save failure.") );
+		}
 	}
 }
 
+void FSceneOutlinerActorSCCColumn::ExecuteSCCHistory()
+{
+	TArray<FString> PackageNames;
+	GetSelectedPackageNames(PackageNames);
+	FSourceControlWindows::DisplayRevisionHistory(SourceControlHelpers::PackageFilenames(PackageNames));
+}
+
+void FSceneOutlinerActorSCCColumn::ExecuteSCCRevert()
+{
+	TArray<FString> PackageNames;
+	GetSelectedPackageNames(PackageNames);
+	FSourceControlWindows::PromptForRevert(PackageNames);
+}
+
 #undef LOCTEXT_NAMESPACE
\ No newline at end of file
diff --git a/Engine/Source/Editor/SceneOutliner/Private/SceneOutlinerActorSCCColumn.h b/Engine/Source/Editor/SceneOutliner/Private/SceneOutlinerActorSCCColumn.h
--- a/Engine/Source/Editor/SceneOutliner/Private/SceneOutlinerActorSCCColumn.h
+++ b/Engine/Source/Editor/SceneOutliner/Private/SceneOutlinerActorSCCColumn.h
@@ -8,6 +8,8 @@
 #include "ISceneOutliner.h"
 #include "SceneOutlinerPublicTypes.h"
 #include "ISceneOutlinerColumn.h"
+#include "SSourceControlWidget.h"
+#include "ToolMenu.h"
 
 template<typename ItemType> class STableRow;
 
@@ -22,6 +24,8 @@ public:
 
 	static FName GetID() { return FSceneOutlinerBuiltInColumnTypes::SourceControl(); }
 
+	bool AddSourceControlMenuOptions(UToolMenu* Menu, TArray<FSceneOutlinerTreeItemPtr> InSelectedItems);
+
 	//////////////////////////////////////////////////////////////////////////
 	// Begin ISceneOutlinerColumn Implementation
 
@@ -38,5 +42,30 @@ public:
 private:
 	const FSlateBrush* GetHeaderIcon() const;
 
+	bool CanExecuteSourceControlActions() const;
+	void CacheCanExecuteVars();
+	bool CanExecuteSCCCheckOut() const;
+	bool CanExecuteSCCCheckIn() const;
+	bool CanExecuteSCCHistory() const;
+	bool CanExecuteSCCRevert() const;
+	bool CanExecuteSCCRefresh() const;
+	void FillSourceControlSubMenu(UToolMenu* Menu);
+	void GetSelectedPackageNames(TArray<FString>& OutPackageNames) const;
+	void GetSelectedPackages(TArray<UPackage*>& OutPackages) const;
+	void ExecuteSCCRefresh();
+	void ExecuteSCCCheckOut();
+	void ExecuteSCCCheckIn();
+	void ExecuteSCCHistory();
+	void ExecuteSCCRevert();
+
 	TWeakPtr<ISceneOutliner> WeakSceneOutliner;
+
+	TArray<FSceneOutlinerTreeItemPtr> SelectedItems;
+
+	TMap<FSceneOutlinerTreeItemPtr, TSharedRef<SSourceControlWidget>> ItemWidgets;
+
+	bool bCanExecuteSCCCheckOut = false;
+	bool bCanExecuteSCCCheckIn = false;
+	bool bCanExecuteSCCHistory = false;
+	bool bCanExecuteSCCRevert = false;
 };
diff --git a/Engine/Source/Editor/SceneOutliner/SceneOutliner.Build.cs b/Engine/Source/Editor/SceneOutliner/SceneOutliner.Build.cs
--- a/Engine/Source/Editor/SceneOutliner/SceneOutliner.Build.cs
+++ b/Engine/Source/Editor/SceneOutliner/SceneOutliner.Build.cs
@@ -24,7 +24,9 @@ public class SceneOutliner : ModuleRules
 				"UnrealEd",
 				"ToolMenus",
 				"SourceControl",
-				"EditorConfig"
+				"EditorConfig",
+				"SourceControl",
+				"SourceControlWindows",
 			}
 		);
 	}
