Outliner: Add option to Auto Expand Hierarchy when Always Frame Selection is on, and enable parent highlighting for collapsed items

diff --git a/Engine/Source/Editor/SceneOutliner/Private/ActorBrowsingMode.cpp b/Engine/Source/Editor/SceneOutliner/Private/ActorBrowsingMode.cpp
--- a/Engine/Source/Editor/SceneOutliner/Private/ActorBrowsingMode.cpp
+++ b/Engine/Source/Editor/SceneOutliner/Private/ActorBrowsingMode.cpp
@@ -212,6 +212,16 @@ FActorBrowsingMode::FActorBrowsingMode(SSceneOutliner* InSceneOutliner, TWeakObj
 
 	bAlwaysFrameSelection = LocalSettings.bAlwaysFrameSelection;
 
+	// Auto Expand Hierarchy is auto turned off when Always Frame Selection is off
+	if(!bAlwaysFrameSelection)
+	{
+		bAutoExpandHierarchy = false;
+	}
+	else
+	{
+		bAutoExpandHierarchy = LocalSettings.bAutoExpandHierarchy;
+	}
+
 	Rebuild();
 }
 
@@ -333,6 +343,12 @@ void FActorBrowsingMode::OnToggleAlwaysFrameSelection()
 {
 	bAlwaysFrameSelection = !bAlwaysFrameSelection;
 
+	// Auto Expand Hierarchy is auto turned off when Always Frame Selection is off
+	if(!bAlwaysFrameSelection)
+	{
+		bAutoExpandHierarchy = false;
+	}
+
 	FActorBrowsingModeConfig* Settings = GetMutableConfig();
 	if(Settings)
 	{
@@ -341,6 +357,23 @@ void FActorBrowsingMode::OnToggleAlwaysFrameSelection()
 	}
 }
 
+void FActorBrowsingMode::OnToggleAutoExpandHierarchy()
+{
+	bAutoExpandHierarchy = !bAutoExpandHierarchy;
+
+	FActorBrowsingModeConfig* Settings = GetMutableConfig();
+	if(Settings)
+	{
+		Settings->bAutoExpandHierarchy = bAutoExpandHierarchy;
+		SaveConfig();
+	}
+}
+
+bool FActorBrowsingMode::ShouldAutoExpandHierarchy()
+{
+	return bAutoExpandHierarchy;
+}
+
 bool FActorBrowsingMode::ShouldAlwaysFrameSelection()
 {
 	return bAlwaysFrameSelection;
@@ -363,6 +396,21 @@ void FActorBrowsingMode::CreateViewContent(FMenuBuilder& MenuBuilder)
 		EUserInterfaceActionType::ToggleButton
 	);
 
+	MenuBuilder.AddMenuEntry(
+		LOCTEXT("AutoExpandHierarchy", "Auto Expand Hierarchy"),
+		LOCTEXT("AlwaysFrameSelectionTooltip", "Automatically expand a collapsed item when it is selected through the level editor viewport. Disabled when Always Frame Selection is off"),
+		FSlateIcon(),
+		FUIAction(
+			FExecuteAction::CreateRaw(this, &FActorBrowsingMode::OnToggleAutoExpandHierarchy),
+			// Don't allow the user to toggle this when Always Frame Selection is off
+			FCanExecuteAction::CreateRaw(this, &FActorBrowsingMode::ShouldAlwaysFrameSelection),
+			FIsActionChecked::CreateRaw(this, &FActorBrowsingMode::ShouldAutoExpandHierarchy)
+		),
+		NAME_None,
+		EUserInterfaceActionType::ToggleButton
+	);
+
+
 	MenuBuilder.EndSection();
 
 
diff --git a/Engine/Source/Editor/SceneOutliner/Private/ActorBrowsingModeSettings.h b/Engine/Source/Editor/SceneOutliner/Private/ActorBrowsingModeSettings.h
--- a/Engine/Source/Editor/SceneOutliner/Private/ActorBrowsingModeSettings.h
+++ b/Engine/Source/Editor/SceneOutliner/Private/ActorBrowsingModeSettings.h
@@ -47,6 +47,13 @@ public:
 	/** True when the Scene Outliner updates when an actor is selected in the viewport */
 	UPROPERTY()
 	bool bAlwaysFrameSelection = true;
+
+	/** True when the Scene Outliner expands an actor when it is selected in the viewport
+	 *  False when the Scene Outliner scrolls to the first visible parent instead.
+	 *  Automatically false when bAlwaysFrameSelection is false
+	 */
+	UPROPERTY()
+	bool bAutoExpandHierarchy = true;
 };
 
 UCLASS(EditorConfig="ActorBrowsingMode")
diff --git a/Engine/Source/Editor/SceneOutliner/Private/ActorModeInteractive.cpp b/Engine/Source/Editor/SceneOutliner/Private/ActorModeInteractive.cpp
--- a/Engine/Source/Editor/SceneOutliner/Private/ActorModeInteractive.cpp
+++ b/Engine/Source/Editor/SceneOutliner/Private/ActorModeInteractive.cpp
@@ -63,14 +63,25 @@ void FActorModeInteractive::OnLevelSelectionChanged(UObject* Obj)
 		SceneOutliner->ClearSelection();
 		SceneOutliner->RefreshSelection();
 
-		// Scroll last item into view if Selection Framing is enabled - this means if we are multi-selecting, we show newest selection. @TODO Not perfect though
+		// Scroll last item into view  - this means if we are multi-selecting, we show newest selection. @TODO Not perfect though
 		if (const AActor* LastSelectedActor = GEditor->GetSelectedActors()->GetBottom<AActor>())
 		{
 			if (FSceneOutlinerTreeItemPtr TreeItem = SceneOutliner->GetTreeItem(LastSelectedActor, false))
 			{
+				// Only scroll if selection framing is enabled
 				if(bAlwaysFrameSelection)
 				{
-					SceneOutliner->ScrollItemIntoView(TreeItem);
+					// If Auto Expand Hierarchy is also enabled (default behavior), expand the selected item and scroll to it
+					if(bAutoExpandHierarchy)
+					{
+						SceneOutliner->ScrollItemIntoView(TreeItem);
+					}
+					// Otherwise, don't expand the selected item, just scroll to the first visible parent (which will be highlighted)
+					else
+					{
+						SceneOutliner->ScrollToFirstVisibleParent(TreeItem);
+					}
+					
 				}
 			}
 			else
diff --git a/Engine/Source/Editor/SceneOutliner/Private/SOutlinerTreeView.cpp b/Engine/Source/Editor/SceneOutliner/Private/SOutlinerTreeView.cpp
--- a/Engine/Source/Editor/SceneOutliner/Private/SOutlinerTreeView.cpp
+++ b/Engine/Source/Editor/SceneOutliner/Private/SOutlinerTreeView.cpp
@@ -152,6 +152,36 @@ FReply SSceneOutlinerTreeView::OnDrop(const FGeometry& MyGeometry, const FDragDr
 	return HandleDropFromWeak(SceneOutlinerWeak, DragDropEvent, ValidationInfo, true);
 }
 
+void SSceneOutlinerTreeView::ScrollToFirstVisibleParent(FSceneOutlinerTreeItemPtr InItemPtr)
+{
+	while(InItemPtr.IsValid())
+	{
+		// The first parent that is in the LinearizedItems is visible in the tree
+		if(LinearizedItems.Find(InItemPtr) != INDEX_NONE)
+		{
+			RequestScrollIntoView(InItemPtr);
+			return;
+		}
+		InItemPtr = InItemPtr->GetParent();
+	}
+}
+
+void SSceneOutlinerTreeView::Private_UpdateParentHighlights()
+{
+	this->ClearHighlightedItems();
+
+	// For the Outliner, we want to highlight parent items even if the current selection is not visible (i.e collapsed)
+	for( typename TItemSet::TConstIterator SelectedItemIt(SelectedItems); SelectedItemIt; ++SelectedItemIt )
+	{
+		auto Parent = (*SelectedItemIt)->GetParent();
+		while (Parent.IsValid())
+		{
+			Private_SetItemHighlighted(Parent, true);
+			Parent = Parent->GetParent();
+		}
+	}
+}
+
 FReply SSceneOutlinerTreeRow::OnDrop( const FGeometry& MyGeometry, const FDragDropEvent& DragDropEvent )
 {
 	auto ItemPtr = Item.Pin();
diff --git a/Engine/Source/Editor/SceneOutliner/Private/SSceneOutliner.cpp b/Engine/Source/Editor/SceneOutliner/Private/SSceneOutliner.cpp
--- a/Engine/Source/Editor/SceneOutliner/Private/SSceneOutliner.cpp
+++ b/Engine/Source/Editor/SceneOutliner/Private/SSceneOutliner.cpp
@@ -1745,6 +1745,11 @@ void SSceneOutliner::ScrollItemIntoView(const FSceneOutlinerTreeItemPtr& Item)
 	OutlinerTreeView->RequestScrollIntoView(Item);
 }
 
+void SSceneOutliner::ScrollToFirstVisibleParent(const FSceneOutlinerTreeItemPtr& Item)
+{
+	OutlinerTreeView->ScrollToFirstVisibleParent(Item);
+}
+
 void SSceneOutliner::SetItemExpansion(const FSceneOutlinerTreeItemPtr& Item, bool bIsExpanded)
 {
 	OutlinerTreeView->SetItemExpansion(Item, bIsExpanded);
diff --git a/Engine/Source/Editor/SceneOutliner/Public/ActorBrowsingMode.h b/Engine/Source/Editor/SceneOutliner/Public/ActorBrowsingMode.h
--- a/Engine/Source/Editor/SceneOutliner/Public/ActorBrowsingMode.h
+++ b/Engine/Source/Editor/SceneOutliner/Public/ActorBrowsingMode.h
@@ -121,6 +121,10 @@ private:
 	void OnToggleAlwaysFrameSelection();
 	bool ShouldAlwaysFrameSelection();
 
+	/** Functions to expose automatic hierarchy expansion to the UI */
+	void OnToggleAutoExpandHierarchy();
+	bool ShouldAutoExpandHierarchy();
+
 	/**
 	 * Get a mutable version of the ActorBrowser config for setting values.
 	 * @returns		The config for this ActorBrowser.
diff --git a/Engine/Source/Editor/SceneOutliner/Public/ActorMode.h b/Engine/Source/Editor/SceneOutliner/Public/ActorMode.h
--- a/Engine/Source/Editor/SceneOutliner/Public/ActorMode.h
+++ b/Engine/Source/Editor/SceneOutliner/Public/ActorMode.h
@@ -105,4 +105,10 @@ protected:
 	bool bHideEmptyFolders;
 	/** Should the outliner scroll to the item on selection */
 	bool bAlwaysFrameSelection;
+	
+	/** True when the Scene Outliner expands an actor when it is selected in the viewport
+	 *  False when the Scene Outliner scrolls to the first visible parent instead.
+	 *  Automatically false when bAlwaysFrameSelection is false
+	 */
+	bool bAutoExpandHierarchy;
 };
diff --git a/Engine/Source/Editor/SceneOutliner/Public/SOutlinerTreeView.h b/Engine/Source/Editor/SceneOutliner/Public/SOutlinerTreeView.h
--- a/Engine/Source/Editor/SceneOutliner/Public/SOutlinerTreeView.h
+++ b/Engine/Source/Editor/SceneOutliner/Public/SOutlinerTreeView.h
@@ -19,6 +19,8 @@ public:
 
 	void FlashHighlightOnItem( FSceneOutlinerTreeItemPtr FlashHighlightOnItem );
 
+	void ScrollToFirstVisibleParent(FSceneOutlinerTreeItemPtr InItemPtr);
+
 	const TWeakPtr<SSceneOutliner>& GetOutlinerPtr() { return SceneOutlinerWeak; }
 protected:
 
@@ -28,6 +30,9 @@ protected:
 
 	virtual FReply OnDrop( const FGeometry& MyGeometry, const FDragDropEvent& DragDropEvent ) override;
 
+private:
+	virtual void Private_UpdateParentHighlights() override;
+
 	/** Weak reference to the outliner widget that owns this list */
 	TWeakPtr<SSceneOutliner> SceneOutlinerWeak;
 };
diff --git a/Engine/Source/Editor/SceneOutliner/Public/SSceneOutliner.h b/Engine/Source/Editor/SceneOutliner/Public/SSceneOutliner.h
--- a/Engine/Source/Editor/SceneOutliner/Public/SSceneOutliner.h
+++ b/Engine/Source/Editor/SceneOutliner/Public/SSceneOutliner.h
@@ -590,6 +590,9 @@ public:
 		
 	bool IsItemExpanded(const FSceneOutlinerTreeItemPtr& Item) const;
 
+	/** Scroll to the first visible (non-collapsed) parent of the specified item */
+	void ScrollToFirstVisibleParent(const FSceneOutlinerTreeItemPtr& Item);
+
 	void ExpandAll();
 	void CollapseAll();
 
