ActorBrowsingMode: Fix duplicate localization key

diff --git a/Engine/Source/Editor/SceneOutliner/Private/ActorBrowsingMode.cpp b/Engine/Source/Editor/SceneOutliner/Private/ActorBrowsingMode.cpp
--- a/Engine/Source/Editor/SceneOutliner/Private/ActorBrowsingMode.cpp
+++ b/Engine/Source/Editor/SceneOutliner/Private/ActorBrowsingMode.cpp
@@ -384,7 +384,7 @@ void FActorBrowsingMode::CreateViewContent(FMenuBuilder& MenuBuilder)
 	MenuBuilder.BeginSection("OutlinerSelectionOptions", LOCTEXT("OptionsHeading", "Options"));
 
 	MenuBuilder.AddMenuEntry(
-		LOCTEXT("AlwaysFrameSelection", "Always Frame Selection"),
+		LOCTEXT("AlwaysFrameSelectionLabel", "Always Frame Selection"),
 		LOCTEXT("AlwaysFrameSelectionTooltip", "Select an item in the outliner when it is selected through the level editor viewport."),
 		FSlateIcon(),
 		FUIAction(
@@ -397,8 +397,8 @@ void FActorBrowsingMode::CreateViewContent(FMenuBuilder& MenuBuilder)
 	);
 
 	MenuBuilder.AddMenuEntry(
-		LOCTEXT("AutoExpandHierarchy", "Auto Expand Hierarchy"),
-		LOCTEXT("AlwaysFrameSelectionTooltip", "Automatically expand a collapsed item when it is selected through the level editor viewport. Disabled when Always Frame Selection is off"),
+		LOCTEXT("AutoExpandHierarchyLabel", "Auto Expand Hierarchy"),
+		LOCTEXT("AutoExpandHierarchyTooltip", "Automatically expand a collapsed item when it is selected through the level editor viewport. Disabled when Always Frame Selection is off"),
 		FSlateIcon(),
 		FUIAction(
 			FExecuteAction::CreateRaw(this, &FActorBrowsingMode::OnToggleAutoExpandHierarchy),
