Outliner: Improve tooltip for "Always Frame Selection" and "Auto Expand Hierarchy"

diff --git a/Engine/Source/Editor/SceneOutliner/Private/ActorBrowsingMode.cpp b/Engine/Source/Editor/SceneOutliner/Private/ActorBrowsingMode.cpp
--- a/Engine/Source/Editor/SceneOutliner/Private/ActorBrowsingMode.cpp
+++ b/Engine/Source/Editor/SceneOutliner/Private/ActorBrowsingMode.cpp
@@ -385,7 +385,7 @@ void FActorBrowsingMode::CreateViewContent(FMenuBuilder& MenuBuilder)
 
 	MenuBuilder.AddMenuEntry(
 		LOCTEXT("AlwaysFrameSelectionLabel", "Always Frame Selection"),
-		LOCTEXT("AlwaysFrameSelectionTooltip", "Select an item in the outliner when it is selected through the level editor viewport."),
+		LOCTEXT("AlwaysFrameSelectionTooltip", "Automatically scroll to show the currently selected actor on selection."),
 		FSlateIcon(),
 		FUIAction(
 			FExecuteAction::CreateRaw(this, &FActorBrowsingMode::OnToggleAlwaysFrameSelection),
@@ -398,7 +398,7 @@ void FActorBrowsingMode::CreateViewContent(FMenuBuilder& MenuBuilder)
 
 	MenuBuilder.AddMenuEntry(
 		LOCTEXT("AutoExpandHierarchyLabel", "Auto Expand Hierarchy"),
-		LOCTEXT("AutoExpandHierarchyTooltip", "Automatically expand a collapsed item when it is selected through the level editor viewport. Disabled when Always Frame Selection is off"),
+		LOCTEXT("AutoExpandHierarchyTooltip", "Automatically expand any containing hierarchy when an actor is selected. Disabled when Always Frame Selection is off."),
 		FSlateIcon(),
 		FUIAction(
 			FExecuteAction::CreateRaw(this, &FActorBrowsingMode::OnToggleAutoExpandHierarchy),
