In the source outliner, displayed a tooltip over the source control icon. The tooltip is taken from the source control implementation.

diff --git a/Engine/Source/Editor/SceneOutliner/Private/SSourceControlWidget.cpp b/Engine/Source/Editor/SceneOutliner/Private/SSourceControlWidget.cpp
--- a/Engine/Source/Editor/SceneOutliner/Private/SSourceControlWidget.cpp
+++ b/Engine/Source/Editor/SceneOutliner/Private/SSourceControlWidget.cpp
@@ -87,7 +87,6 @@ FReply SSourceControlWidget::OnMouseButtonDoubleClick(const FGeometry& InMyGeome
 	if (SourceControlState.IsValid())
 	{
 		UpdateSourceControlStateIcon(SourceControlState);
-
 	}
 	return FReply::Handled();
 }
@@ -158,12 +157,13 @@ void SSourceControlWidget::UpdateSourceControlStateIcon(FSourceControlStatePtr S
 	if(SourceControlState.IsValid())
 	{
 		FSlateIcon Icon = SourceControlState->GetIcon();
-		
 		SetFromSlateIcon(Icon);
+		SetToolTipText(SourceControlState->GetDisplayTooltip());
 	}
 	else
 	{
 		SetImage(nullptr);
+		SetToolTipText(FText::GetEmpty());
 		RemoveAllLayers();
 	}
-}
\ No newline at end of file
+}
