Added "Copy Actor File Path" to clipboard option on unloaded actors.

diff --git a/Engine/Source/Editor/SceneOutliner/Private/ActorDescTreeItem.cpp b/Engine/Source/Editor/SceneOutliner/Private/ActorDescTreeItem.cpp
--- a/Engine/Source/Editor/SceneOutliner/Private/ActorDescTreeItem.cpp
+++ b/Engine/Source/Editor/SceneOutliner/Private/ActorDescTreeItem.cpp
@@ -21,6 +21,7 @@
 #include "SSocketChooser.h"
 #include "ToolMenus.h"
 #include "LevelEditorViewport.h"
+#include "HAL/PlatformApplicationMisc.h"
 
 #define LOCTEXT_NAMESPACE "SceneOutliner_ActorDescTreeItem"
 
@@ -267,10 +268,24 @@ void FActorDescTreeItem::FocusActorBounds() const
 	}
 }
 
+void FActorDescTreeItem::CopyActorFilePathtoClipboard() const
+{
+	if (FWorldPartitionActorDesc const* ActorDesc = ActorDescHandle.Get())
+	{
+		FString PackageFilename;
+		if (FPackageName::TryConvertLongPackageNameToFilename(ActorDesc->GetActorPackage().ToString(), PackageFilename, FPackageName::GetAssetPackageExtension()))
+		{
+			FString Result = FPaths::ConvertRelativePathToFull(PackageFilename);
+			FPlatformApplicationMisc::ClipboardCopy(*Result);
+		}
+	}
+}
+
 void FActorDescTreeItem::GenerateContextMenu(UToolMenu* Menu, SSceneOutliner&)
 {
 	FToolMenuSection& Section = Menu->AddSection("Section");
 	Section.AddMenuEntry("FocusActorBounds", LOCTEXT("FocusActorBounds", "Focus Actor Bounds"), FText(), FSlateIcon(), FUIAction(FExecuteAction::CreateSP(this, &FActorDescTreeItem::FocusActorBounds)));
+	Section.AddMenuEntry("CopyActorFilePath", LOCTEXT("CopyActorFilePath", "Copy Actor File Path"), LOCTEXT("CopyActorFilePathTooltip", "Copy the file path where this actor is saved"), FSlateIcon(FAppStyle::GetAppStyleSetName(), "GenericCommands.Copy"), FUIAction(FExecuteAction::CreateSP(this, &FActorDescTreeItem::CopyActorFilePathtoClipboard)));
 }
 
 bool FActorDescTreeItem::GetVisibility() const
diff --git a/Engine/Source/Editor/SceneOutliner/Public/ActorDescTreeItem.h b/Engine/Source/Editor/SceneOutliner/Public/ActorDescTreeItem.h
--- a/Engine/Source/Editor/SceneOutliner/Public/ActorDescTreeItem.h
+++ b/Engine/Source/Editor/SceneOutliner/Public/ActorDescTreeItem.h
@@ -64,5 +64,6 @@ protected:
 
 private:
 	void FocusActorBounds() const;
+	void CopyActorFilePathtoClipboard() const;
 	FGuid ActorGuid;
 };
