PR #9372: Fix issues with initial integration of OFPA outliner source control actions (Contributed by [REDACTED_PERSON])

#rnx

diff --git a/Engine/Source/Editor/SceneOutliner/Private/ActorBrowsingMode.cpp b/Engine/Source/Editor/SceneOutliner/Private/ActorBrowsingMode.cpp
--- a/Engine/Source/Editor/SceneOutliner/Private/ActorBrowsingMode.cpp
+++ b/Engine/Source/Editor/SceneOutliner/Private/ActorBrowsingMode.cpp
@@ -548,14 +548,7 @@ void FActorBrowsingMode::RegisterContextMenu()
 					if (Context->NumSelectedItems > 0 && Context->SceneOutliner.IsValid())
 					{
 						SSceneOutliner* SceneOutliner = Context->SceneOutliner.Pin().Get();
-						TMap<FName, TSharedPtr<ISceneOutlinerColumn>> OutlinerColumns = SceneOutliner->GetColumns();
-						TSharedPtr<ISceneOutlinerColumn> * ColumnLookup = OutlinerColumns.Find("Source Control");
-						if (ColumnLookup != nullptr && ColumnLookup->IsValid())
-						{
-							TSharedPtr<FSceneOutlinerActorSCCColumn> SCCColumn = StaticCastSharedPtr<FSceneOutlinerActorSCCColumn>(*ColumnLookup);
-							TArray<FSceneOutlinerTreeItemPtr> SelectedItems = SceneOutliner->GetTree().GetSelectedItems();
-							SCCColumn->AddSourceControlMenuOptions(InMenu, SelectedItems);
-						}
+						SceneOutliner->AddSourceControlMenuOptions(InMenu);
 					}
 				}
 			}));
diff --git a/Engine/Source/Editor/SceneOutliner/Private/SSceneOutliner.cpp b/Engine/Source/Editor/SceneOutliner/Private/SSceneOutliner.cpp
--- a/Engine/Source/Editor/SceneOutliner/Private/SSceneOutliner.cpp
+++ b/Engine/Source/Editor/SceneOutliner/Private/SSceneOutliner.cpp
@@ -298,6 +298,8 @@ void SSceneOutliner::Construct(const FArguments& InArgs, const FSceneOutlinerIni
 
 	// Register to be notified when properties are edited
 	FCoreUObjectDelegates::OnPackageReloaded.AddRaw(this, &SSceneOutliner::OnAssetReloaded);
+
+	SourceControlHandler = TSharedPtr<FSceneOutlinerSCCHandler>(new FSceneOutlinerSCCHandler());
 }
 
 void SSceneOutliner::HandleHiddenColumnsChanged()
@@ -1757,6 +1759,8 @@ bool SSceneOutliner::IsItemExpanded(const FSceneOutlinerTreeItemPtr& Item) const
 
 TSharedRef< ITableRow > SSceneOutliner::OnGenerateRowForOutlinerTree( FSceneOutlinerTreeItemPtr Item, const TSharedRef< STableViewBase >& OwnerTable )
 {
+	SourceControlHandler->AddItem(Item);
+
 	return SNew( SSceneOutlinerTreeRow, OutlinerTreeView.ToSharedRef(), SharedThis(this) ).Item( Item );
 }
 
@@ -2374,4 +2378,15 @@ void FSceneOutlinerMenuHelper::AddMenuEntryCleanupFolders(FToolMenuSection& InSe
 	}
 }
 
+TSharedPtr<FSceneOutlinerTreeItemSCC> SSceneOutliner::GetItemSourceControl(const FSceneOutlinerTreeItemPtr& InItem)
+{
+	return SourceControlHandler->GetItemSourceControl(InItem);
+}
+
+void SSceneOutliner::AddSourceControlMenuOptions(UToolMenu* Menu)
+{
+	TArray<FSceneOutlinerTreeItemPtr> SelectedItems = GetTree().GetSelectedItems();
+	SourceControlHandler->AddSourceControlMenuOptions(Menu, SelectedItems);
+}
+
 #undef LOCTEXT_NAMESPACE
\ No newline at end of file
diff --git a/Engine/Source/Editor/SceneOutliner/Private/SSourceControlWidget.cpp b/Engine/Source/Editor/SceneOutliner/Private/SSourceControlWidget.cpp
--- a/Engine/Source/Editor/SceneOutliner/Private/SSourceControlWidget.cpp
+++ b/Engine/Source/Editor/SceneOutliner/Private/SSourceControlWidget.cpp
@@ -2,154 +2,29 @@
 
 #include "SSourceControlWidget.h"
 
-#include "ActorTreeItem.h"
-#include "ActorDescTreeItem.h"
-#include "ActorFolderTreeItem.h"
-#include "ISourceControlProvider.h"
-#include "ISourceControlModule.h"
-#include "WorldPartition/WorldPartitionActorDesc.h"
-
-void SSourceControlWidget::Construct(const FArguments& InArgs, TWeakPtr<ISceneOutliner> InWeakOutliner, TWeakPtr<ISceneOutlinerTreeItem> InWeakTreeItem)
+void SSourceControlWidget::Construct(const FArguments& InArgs, TSharedPtr<FSceneOutlinerTreeItemSCC> InItemSourceControl)
 {
-	WeakTreeItem = InWeakTreeItem;
-	WeakOutliner = InWeakOutliner;
+	check(InItemSourceControl.IsValid());
+
+	ItemSourceControl = InItemSourceControl;
+
+	ItemSourceControl->OnSourceControlStateChanged.BindSP(this, &SSourceControlWidget::UpdateSourceControlStateIcon);
 
 	SImage::Construct(
 		SImage::FArguments()
 		.ColorAndOpacity(this, &SSourceControlWidget::GetForegroundColor)
 		.Image(FStyleDefaults::GetNoBrush()));
 
-	FSceneOutlinerTreeItemPtr TreeItemPtr = WeakTreeItem.Pin();
-	if (TreeItemPtr.IsValid())
-	{
-		if (FActorTreeItem* ActorItem = TreeItemPtr->CastTo<FActorTreeItem>())
-		{
-			if (AActor* Actor = ActorItem->Actor.Get())
-			{
-				if (Actor->IsPackageExternal())
-				{
-					ExternalPackageName = USourceControlHelpers::PackageFilename(Actor->GetExternalPackage());
-					ExternalPackage = Actor->GetExternalPackage();
-				}
-
-				ActorPackingModeChangedDelegateHandle = Actor->OnPackagingModeChanged.AddLambda([this](AActor* InActor, bool bExternal)
-				{
-					if (bExternal)
-					{
-						ExternalPackageName = USourceControlHelpers::PackageFilename(InActor->GetExternalPackage());
-						ExternalPackage = InActor->GetExternalPackage();
-						ConnectSourceControl();
-					}
-					else
-					{
-						ExternalPackageName = FString();
-						ExternalPackage = nullptr;
-						DisconnectSourceControl();
-					}
-				});
-			}
-		}
-		else if (FActorFolderTreeItem* ActorFolderItem = TreeItemPtr->CastTo<FActorFolderTreeItem>())
-		{
-			if (const UActorFolder* ActorFolder = ActorFolderItem->GetActorFolder())
-			{
-				if (ActorFolder->IsPackageExternal())
-				{
-					ExternalPackageName = USourceControlHelpers::PackageFilename(ActorFolder->GetExternalPackage());
-					ExternalPackage = ActorFolder->GetExternalPackage();
-				}
-			}
-		}
-		else if (FActorDescTreeItem* ActorDescItem = TreeItemPtr->CastTo<FActorDescTreeItem>())
-		{
-			if (const FWorldPartitionActorDesc* ActorDesc = ActorDescItem->ActorDescHandle.Get())
-			{
-				ExternalPackageName = USourceControlHelpers::PackageFilename(ActorDesc->GetActorPackage().ToString());
-				ExternalPackage = FindPackage(nullptr, *ActorDesc->GetActorPackage().ToString());
-			}
-		}
-	}
-
-	if (!ExternalPackageName.IsEmpty())
-	{
-		ConnectSourceControl();
-	}
-}
-
-SSourceControlWidget::~SSourceControlWidget()
-{
-	DisconnectSourceControl();
+	UpdateSourceControlStateIcon(ItemSourceControl->GetSourceControlState());
 }
 
 FReply SSourceControlWidget::OnMouseButtonDoubleClick(const FGeometry& InMyGeometry, const FPointerEvent& InMouseEvent)
 {
-	FSourceControlStatePtr SourceControlState = ISourceControlModule::Get().GetProvider().GetState(ExternalPackageName, EStateCacheUsage::ForceUpdate);
-	if (SourceControlState.IsValid())
-	{
-		UpdateSourceControlStateIcon(SourceControlState);
-	}
-	return FReply::Handled();
-}
-
-FSourceControlStatePtr SSourceControlWidget::GetSourceControlState()
-{
-	return ISourceControlModule::Get().GetProvider().GetState(ExternalPackageName, EStateCacheUsage::Use);
-}
-
-void SSourceControlWidget::ConnectSourceControl()
-{
-	check(!ExternalPackageName.IsEmpty());
+	check(ItemSourceControl);
 
-	ISourceControlModule& SCCModule = ISourceControlModule::Get();
-	SourceControlProviderChangedDelegateHandle = SCCModule.RegisterProviderChanged(FSourceControlProviderChanged::FDelegate::CreateSP(this, &SSourceControlWidget::HandleSourceControlProviderChanged));
-	SourceControlStateChangedDelegateHandle = SCCModule.GetProvider().RegisterSourceControlStateChanged_Handle(FSourceControlStateChanged::FDelegate::CreateSP(this, &SSourceControlWidget::HandleSourceControlStateChanged, EStateCacheUsage::Use));
-
-	// Check if there is already a cached state for this item
-	FSourceControlStatePtr SourceControlState = ISourceControlModule::Get().GetProvider().GetState(ExternalPackageName, EStateCacheUsage::Use);
-	if (SourceControlState.IsValid() && !SourceControlState->IsUnknown())
-	{
-		UpdateSourceControlStateIcon(SourceControlState);
-	}
-	else
-	{
-		SCCModule.QueueStatusUpdate(ExternalPackageName);
-	}
-}
-
-void SSourceControlWidget::DisconnectSourceControl()
-{
-	FSceneOutlinerTreeItemPtr TreeItemPtr = WeakTreeItem.Pin();
-	if (TreeItemPtr.IsValid())
-	{
-		if (FActorTreeItem* ActorItem = TreeItemPtr->CastTo<FActorTreeItem>())
-		{
-			if (AActor* Actor = ActorItem->Actor.Get())
-			{
-				Actor->OnPackagingModeChanged.Remove(ActorPackingModeChangedDelegateHandle);
-			}
-		}
-	}
-	ISourceControlModule::Get().GetProvider().UnregisterSourceControlStateChanged_Handle(SourceControlStateChangedDelegateHandle);
-	ISourceControlModule::Get().UnregisterProviderChanged(SourceControlProviderChangedDelegateHandle);
-}
-
-void SSourceControlWidget::HandleSourceControlStateChanged(EStateCacheUsage::Type CacheUsage)
-{
-	FSourceControlStatePtr SourceControlState = ISourceControlModule::Get().GetProvider().GetState(ExternalPackageName, CacheUsage);
-	if (SourceControlState.IsValid())
-	{
-		UpdateSourceControlStateIcon(SourceControlState);
-	}
-}
-
-void SSourceControlWidget::HandleSourceControlProviderChanged(ISourceControlProvider& OldProvider, ISourceControlProvider& NewProvider)
-{
-	OldProvider.UnregisterSourceControlStateChanged_Handle(SourceControlStateChangedDelegateHandle);
-	SourceControlStateChangedDelegateHandle = NewProvider.RegisterSourceControlStateChanged_Handle(FSourceControlStateChanged::FDelegate::CreateSP(this, &SSourceControlWidget::HandleSourceControlStateChanged, EStateCacheUsage::Use));
-	
-	UpdateSourceControlStateIcon(nullptr);
-
-	ISourceControlModule::Get().QueueStatusUpdate(ExternalPackageName);
+	FSourceControlStatePtr SourceControlState = ItemSourceControl->RefreshSourceControlState();
+	UpdateSourceControlStateIcon(SourceControlState);
+	return FReply::Handled();
 }
 
 void SSourceControlWidget::UpdateSourceControlStateIcon(FSourceControlStatePtr SourceControlState)
@@ -157,13 +32,12 @@ void SSourceControlWidget::UpdateSourceControlStateIcon(FSourceControlStatePtr S
 	if(SourceControlState.IsValid())
 	{
 		FSlateIcon Icon = SourceControlState->GetIcon();
+		
 		SetFromSlateIcon(Icon);
-		SetToolTipText(SourceControlState->GetDisplayTooltip());
 	}
 	else
 	{
 		SetImage(nullptr);
-		SetToolTipText(FText::GetEmpty());
 		RemoveAllLayers();
 	}
 }
diff --git a/Engine/Source/Editor/SceneOutliner/Private/SSourceControlWidget.h b/Engine/Source/Editor/SceneOutliner/Private/SSourceControlWidget.h
--- a/Engine/Source/Editor/SceneOutliner/Private/SSourceControlWidget.h
+++ b/Engine/Source/Editor/SceneOutliner/Private/SSourceControlWidget.h
@@ -2,10 +2,9 @@
 
 #pragma once
 
-#include "ISceneOutliner.h"
-#include "SceneOutlinerPublicTypes.h"
 #include "Widgets/Images/SLayeredImage.h"
 #include "SourceControlHelpers.h"
+#include "SceneOutlinerTreeItemSCC.h"
 
 class SSourceControlWidget : public SLayeredImage
 {
@@ -14,46 +13,14 @@ public:
 	SLATE_END_ARGS()
 
 	/** Construct this widget */
-	void Construct(const FArguments& InArgs, TWeakPtr<ISceneOutliner> InWeakOutliner, TWeakPtr<ISceneOutlinerTreeItem> InWeakTreeItem);
-
-	~SSourceControlWidget();
-
-	FSourceControlStatePtr GetSourceControlState();
-
-	FString GetPackageName() { return ExternalPackageName; }
-
-	UPackage* GetPackage() { return ExternalPackage; }
+	void Construct(const FArguments& InArgs, TSharedPtr<FSceneOutlinerTreeItemSCC> InItemSourceControl);
 
 private:
 
 	virtual FReply OnMouseButtonDoubleClick(const FGeometry& InMyGeometry, const FPointerEvent& InMouseEvent) override;
 
-	void ConnectSourceControl();
-
-	void DisconnectSourceControl();
-
-	void HandleSourceControlStateChanged(EStateCacheUsage::Type CacheUsage);
-
-	void HandleSourceControlProviderChanged(ISourceControlProvider& OldProvider, ISourceControlProvider& NewProvider);
-
 	void UpdateSourceControlStateIcon(FSourceControlStatePtr SourceControlState);
 
-	/** The tree item we relate to */
-	TWeakPtr<ISceneOutlinerTreeItem> WeakTreeItem;
-
-	/** Reference back to the outliner so we can set visibility of a whole selection */
-	TWeakPtr<ISceneOutliner> WeakOutliner;
-
-	/** Cache the items external package name */
-	FString ExternalPackageName;
-	UPackage* ExternalPackage = nullptr;
-
-	/** Source control state changed delegate handle */
-	FDelegateHandle SourceControlStateChangedDelegateHandle;
-
-	/** Source control provider changed delegate handle */
-	FDelegateHandle SourceControlProviderChangedDelegateHandle;
-
-	/** Actor packaging mode changed delegate handle */
-	FDelegateHandle ActorPackingModeChangedDelegateHandle;
+	/** The object that keeps the source control state for the outliner */
+	TSharedPtr<FSceneOutlinerTreeItemSCC> ItemSourceControl;
 };
diff --git a/Engine/Source/Editor/SceneOutliner/Private/SceneOutlinerActorSCCColumn.cpp b/Engine/Source/Editor/SceneOutliner/Private/SceneOutlinerActorSCCColumn.cpp
--- a/Engine/Source/Editor/SceneOutliner/Private/SceneOutlinerActorSCCColumn.cpp
+++ b/Engine/Source/Editor/SceneOutliner/Private/SceneOutlinerActorSCCColumn.cpp
@@ -8,10 +8,6 @@
 #include "ActorFolderTreeItem.h"
 #include "WorldPartition/WorldPartitionActorDesc.h"
 #include "ISourceControlModule.h"
-#include "SourceControlOperations.h"
-#include "FileHelpers.h"
-#include "SourceControlWindows.h"
-#include "AssetViewUtils.h"
 #include "Misc/MessageDialog.h"
 
 #define LOCTEXT_NAMESPACE "SceneOutlinerSourceControlColumn"
@@ -43,17 +39,19 @@ const TSharedRef<SWidget> FSceneOutlinerActorSCCColumn::ConstructRowWidget(FScen
 		TreeItem->IsA<FActorDescTreeItem>() || 
 		(TreeItem->IsA<FActorFolderTreeItem>() && TreeItem->CastTo<FActorFolderTreeItem>()->GetActorFolder()))
 	{
-		TSharedRef<SSourceControlWidget> Widget = SNew(SSourceControlWidget, WeakSceneOutliner, TreeItem);
-
-		ItemWidgets.Add(FSceneOutlinerTreeItemPtr(TreeItem), Widget);
-		
-		return SNew(SHorizontalBox)
-			+ SHorizontalBox::Slot()
-			.AutoWidth()
-			.VAlign(VAlign_Center)
-			[
-				Widget
-			];
+		TSharedPtr<FSceneOutlinerTreeItemSCC> SourceControl = WeakSceneOutliner.Pin()->GetItemSourceControl(TreeItem);
+		if (SourceControl.IsValid())
+		{
+			TSharedRef<SSourceControlWidget> Widget = SNew(SSourceControlWidget, SourceControl);
+			
+			return SNew(SHorizontalBox)
+				+ SHorizontalBox::Slot()
+				.AutoWidth()
+				.VAlign(VAlign_Center)
+				[
+					Widget
+				];
+		}
 	}
 	return SNullWidget::NullWidget;
 }
@@ -70,298 +68,4 @@ const FSlateBrush* FSceneOutlinerActorSCCColumn::GetHeaderIcon() const
 	}
 }
 
-bool FSceneOutlinerActorSCCColumn::AddSourceControlMenuOptions(UToolMenu* Menu, TArray<FSceneOutlinerTreeItemPtr> InSelectedItems)
-{
-	SelectedItems = InSelectedItems;
-
-	CacheCanExecuteVars();
-
-	FToolMenuSection& Section = Menu->AddSection("AssetContextSourceControl");
-	
-	if ( ISourceControlModule::Get().IsEnabled() )
-	{
-		// SCC sub menu
-		Section.AddSubMenu(
-			"SourceControlSubMenu",
-			LOCTEXT("SourceControlSubMenuLabel", "Source Control"),
-			LOCTEXT("SourceControlSubMenuToolTip", "Source control actions."),
-			FNewToolMenuDelegate::CreateSP(this, &FSceneOutlinerActorSCCColumn::FillSourceControlSubMenu),
-			FUIAction(
-				FExecuteAction(),
-				FCanExecuteAction::CreateSP( this, &FSceneOutlinerActorSCCColumn::CanExecuteSourceControlActions )
-				),
-			EUserInterfaceActionType::Button,
-			false,
-			FSlateIcon(FAppStyle::GetAppStyleSetName(), "SourceControl.StatusIcon.On")
-			);
-	}
-
-	return true;
-}
-
-
-bool FSceneOutlinerActorSCCColumn::CanExecuteSourceControlActions() const
-{
-	return SelectedItems.Num() > 0;
-}
-
-void FSceneOutlinerActorSCCColumn::CacheCanExecuteVars()
-{
-	bCanExecuteSCCCheckOut = false;
-	bCanExecuteSCCCheckIn = false;
-	bCanExecuteSCCHistory = false;
-	bCanExecuteSCCRevert = false;
-
-	ISourceControlProvider& SourceControlProvider = ISourceControlModule::Get().GetProvider();
-	if ( ISourceControlModule::Get().IsEnabled() )
-	{
-		for (FSceneOutlinerTreeItemPtr SelectedItem : SelectedItems)
-		{
-			TSharedRef<SSourceControlWidget> * WidgetLookup = ItemWidgets.Find(SelectedItem);
-			if (WidgetLookup == nullptr)
-			{
-				continue;
-			}
-
-			TSharedRef<SSourceControlWidget> Widget = *WidgetLookup;
-
-			// Check the SCC state for each package in the selected paths
-			FSourceControlStatePtr SourceControlState = Widget->GetSourceControlState();
-			if(SourceControlState.IsValid())
-			{
-				if ( SourceControlState->CanCheckout() )
-				{
-					bCanExecuteSCCCheckOut = true;
-				}
-
-				if( SourceControlState->IsSourceControlled() && !SourceControlState->IsAdded() )
-				{
-					bCanExecuteSCCHistory = true;
-				}
-
-				if ( SourceControlState->CanCheckIn() )
-				{
-					bCanExecuteSCCCheckIn = true;
-				}
-
-				if (SourceControlState->CanRevert())
-				{
-					bCanExecuteSCCRevert = true;
-				}
-			}
-
-			if ( bCanExecuteSCCCheckOut
-				&& bCanExecuteSCCCheckIn
-				&& bCanExecuteSCCHistory
-				&& bCanExecuteSCCRevert
-				)
-			{
-				// All options are available, no need to keep iterating
-				break;
-			}
-		}
-	}
-}
-
-bool FSceneOutlinerActorSCCColumn::CanExecuteSCCCheckOut() const
-{
-	return bCanExecuteSCCCheckOut;
-}
-
-bool FSceneOutlinerActorSCCColumn::CanExecuteSCCCheckIn() const
-{
-	return bCanExecuteSCCCheckIn;
-}
-
-bool FSceneOutlinerActorSCCColumn::CanExecuteSCCHistory() const
-{
-	return bCanExecuteSCCHistory;
-}
-
-bool FSceneOutlinerActorSCCColumn::CanExecuteSCCRevert() const
-{
-	return bCanExecuteSCCRevert;
-}
-
-bool FSceneOutlinerActorSCCColumn::CanExecuteSCCRefresh() const
-{
-	return ISourceControlModule::Get().IsEnabled();
-}
-
-void FSceneOutlinerActorSCCColumn::FillSourceControlSubMenu(UToolMenu* Menu)
-{
-	FToolMenuSection& Section = Menu->AddSection("AssetSourceControlActions", LOCTEXT("AssetSourceControlActionsMenuHeading", "Source Control"));
-
-	if ( CanExecuteSCCCheckOut() )
-	{
-		Section.AddMenuEntry(
-			"SCCCheckOut",
-			LOCTEXT("SCCCheckOut", "Check Out"),
-			LOCTEXT("SCCCheckOutTooltip", "Checks out the selected asset from source control."),
-			FSlateIcon(FAppStyle::GetAppStyleSetName(), "SourceControl.Actions.CheckOut"),
-			FUIAction(
-				FExecuteAction::CreateSP( this, &FSceneOutlinerActorSCCColumn::ExecuteSCCCheckOut ),
-				FCanExecuteAction::CreateSP( this, &FSceneOutlinerActorSCCColumn::CanExecuteSCCCheckOut )
-			)
-		);
-	}
-
-	if ( CanExecuteSCCCheckIn() )
-	{
-		Section.AddMenuEntry(
-			"SCCCheckIn",
-			LOCTEXT("SCCCheckIn", "Check In"),
-			LOCTEXT("SCCCheckInTooltip", "Checks in the selected asset to source control."),
-			FSlateIcon(FAppStyle::GetAppStyleSetName(), "SourceControl.Actions.Submit"),
-			FUIAction(
-				FExecuteAction::CreateSP( this, &FSceneOutlinerActorSCCColumn::ExecuteSCCCheckIn ),
-				FCanExecuteAction::CreateSP( this, &FSceneOutlinerActorSCCColumn::CanExecuteSCCCheckIn )
-			)
-		);
-	}
-
-	Section.AddMenuEntry(
-		"SCCRefresh",
-		LOCTEXT("SCCRefresh", "Refresh"),
-		LOCTEXT("SCCRefreshTooltip", "Updates the source control status of the asset."),
-		FSlateIcon(FAppStyle::GetAppStyleSetName(), "SourceControl.Actions.Refresh"),
-		FUIAction(
-			FExecuteAction::CreateSP( this, &FSceneOutlinerActorSCCColumn::ExecuteSCCRefresh ),
-			FCanExecuteAction::CreateSP( this, &FSceneOutlinerActorSCCColumn::CanExecuteSCCRefresh )
-			)
-		);
-
-	if( CanExecuteSCCHistory() )
-	{
-		Section.AddMenuEntry(
-			"SCCHistory",
-			LOCTEXT("SCCHistory", "History"),
-			LOCTEXT("SCCHistoryTooltip", "Displays the source control revision history of the selected asset."),
-			FSlateIcon(FAppStyle::GetAppStyleSetName(), "SourceControl.Actions.History"),
-			FUIAction(
-				FExecuteAction::CreateSP( this, &FSceneOutlinerActorSCCColumn::ExecuteSCCHistory ),
-				FCanExecuteAction::CreateSP( this, &FSceneOutlinerActorSCCColumn::CanExecuteSCCHistory )
-			)
-		);
-	}
-
-	if( CanExecuteSCCRevert() )
-	{
-		Section.AddMenuEntry(
-			"SCCRevert",
-			LOCTEXT("SCCRevert", "Revert"),
-			LOCTEXT("SCCRevertTooltip", "Reverts the asset to the state it was before it was checked out."),
-			FSlateIcon(FAppStyle::GetAppStyleSetName(), "SourceControl.Actions.Revert"),
-			FUIAction(
-				FExecuteAction::CreateSP( this, &FSceneOutlinerActorSCCColumn::ExecuteSCCRevert ),
-				FCanExecuteAction::CreateSP( this, &FSceneOutlinerActorSCCColumn::CanExecuteSCCRevert )
-			)
-		);
-	}
-}
-
-void FSceneOutlinerActorSCCColumn::GetSelectedPackageNames(TArray<FString>& OutPackageNames) const
-{
-	for (FSceneOutlinerTreeItemPtr SelectedItem : SelectedItems)
-	{
-		const TSharedRef<SSourceControlWidget> * WidgetLookup = ItemWidgets.Find(SelectedItem);
-		if (WidgetLookup == nullptr)
-		{
-			continue;
-		}
-
-		TSharedRef<SSourceControlWidget> Widget = *WidgetLookup;
-		FString PackageName = Widget->GetPackageName();
-		if (!PackageName.IsEmpty()) {
-			OutPackageNames.Add(PackageName);
-		}
-	}
-}
-
-void FSceneOutlinerActorSCCColumn::GetSelectedPackages(TArray<UPackage*>& OutPackages) const
-{
-	for (FSceneOutlinerTreeItemPtr SelectedItem : SelectedItems)
-	{
-		const TSharedRef<SSourceControlWidget> * WidgetLookup = ItemWidgets.Find(SelectedItem);
-		if (WidgetLookup == nullptr)
-		{
-			continue;
-		}
-
-		TSharedRef<SSourceControlWidget> Widget = *WidgetLookup;
-		UPackage* Package = Widget->GetPackage();
-		if (Package != nullptr) {
-			OutPackages.Add(Package);
-		}
-	}
-}
-
-void FSceneOutlinerActorSCCColumn::ExecuteSCCRefresh()
-{
-	TArray<FString> PackageNames;
-	GetSelectedPackageNames(PackageNames);
-
-	ISourceControlModule::Get().GetProvider().Execute(ISourceControlOperation::Create<FUpdateStatus>(), SourceControlHelpers::PackageFilenames(PackageNames), EConcurrency::Asynchronous);
-}
-
-void FSceneOutlinerActorSCCColumn::ExecuteSCCCheckOut()
-{
-	TArray<UPackage*> PackagesToCheckOut;
-	GetSelectedPackages(PackagesToCheckOut);
-
-	if ( PackagesToCheckOut.Num() > 0 )
-	{
-		FEditorFileUtils::CheckoutPackages(PackagesToCheckOut);
-	}
-}
-
-void FSceneOutlinerActorSCCColumn::ExecuteSCCCheckIn()
-{
-	TArray<UPackage*> Packages;
-	GetSelectedPackages(Packages);
-
-	// Prompt the user to ask if they would like to first save any dirty packages they are trying to check-in
-	const FEditorFileUtils::EPromptReturnCode UserResponse = FEditorFileUtils::PromptForCheckoutAndSave( Packages, true, true );
-
-	// If the user elected to save dirty packages, but one or more of the packages failed to save properly OR if the user
-	// canceled out of the prompt, don't follow through on the check-in process
-	const bool bShouldProceed = ( UserResponse == FEditorFileUtils::EPromptReturnCode::PR_Success || UserResponse == FEditorFileUtils::EPromptReturnCode::PR_Declined );
-	if ( bShouldProceed )
-	{
-		TArray<FString> PackageNames;
-		GetSelectedPackageNames(PackageNames);
-
-		const bool bUseSourceControlStateCache = true;
-		const bool bCheckinGood = FSourceControlWindows::PromptForCheckin(bUseSourceControlStateCache, PackageNames);
-
-		if (!bCheckinGood)
-		{
-			FMessageDialog::Open(EAppMsgType::Ok, NSLOCTEXT("UnrealEd", "SCC_Checkin_Failed", "Check-in failed as a result of save failure."));
-		}
-	}
-	else
-	{
-		// If a failure occurred, alert the user that the check-in was aborted. This warning shouldn't be necessary if the user cancelled
-		// from the dialog, because they obviously intended to cancel the whole operation.
-		if ( UserResponse == FEditorFileUtils::EPromptReturnCode::PR_Failure )
-		{
-			FMessageDialog::Open( EAppMsgType::Ok, NSLOCTEXT("UnrealEd", "SCC_Checkin_Aborted", "Check-in aborted as a result of save failure.") );
-		}
-	}
-}
-
-void FSceneOutlinerActorSCCColumn::ExecuteSCCHistory()
-{
-	TArray<FString> PackageNames;
-	GetSelectedPackageNames(PackageNames);
-	FSourceControlWindows::DisplayRevisionHistory(SourceControlHelpers::PackageFilenames(PackageNames));
-}
-
-void FSceneOutlinerActorSCCColumn::ExecuteSCCRevert()
-{
-	TArray<FString> PackageNames;
-	GetSelectedPackageNames(PackageNames);
-	FSourceControlWindows::PromptForRevert(PackageNames);
-}
-
 #undef LOCTEXT_NAMESPACE
\ No newline at end of file
diff --git a/Engine/Source/Editor/SceneOutliner/Private/SceneOutlinerActorSCCColumn.h b/Engine/Source/Editor/SceneOutliner/Private/SceneOutlinerActorSCCColumn.h
--- a/Engine/Source/Editor/SceneOutliner/Private/SceneOutlinerActorSCCColumn.h
+++ b/Engine/Source/Editor/SceneOutliner/Private/SceneOutlinerActorSCCColumn.h
@@ -9,7 +9,6 @@
 #include "SceneOutlinerPublicTypes.h"
 #include "ISceneOutlinerColumn.h"
 #include "SSourceControlWidget.h"
-#include "ToolMenu.h"
 
 template<typename ItemType> class STableRow;
 
@@ -24,8 +23,6 @@ public:
 
 	static FName GetID() { return FSceneOutlinerBuiltInColumnTypes::SourceControl(); }
 
-	bool AddSourceControlMenuOptions(UToolMenu* Menu, TArray<FSceneOutlinerTreeItemPtr> InSelectedItems);
-
 	//////////////////////////////////////////////////////////////////////////
 	// Begin ISceneOutlinerColumn Implementation
 
@@ -42,30 +39,5 @@ public:
 private:
 	const FSlateBrush* GetHeaderIcon() const;
 
-	bool CanExecuteSourceControlActions() const;
-	void CacheCanExecuteVars();
-	bool CanExecuteSCCCheckOut() const;
-	bool CanExecuteSCCCheckIn() const;
-	bool CanExecuteSCCHistory() const;
-	bool CanExecuteSCCRevert() const;
-	bool CanExecuteSCCRefresh() const;
-	void FillSourceControlSubMenu(UToolMenu* Menu);
-	void GetSelectedPackageNames(TArray<FString>& OutPackageNames) const;
-	void GetSelectedPackages(TArray<UPackage*>& OutPackages) const;
-	void ExecuteSCCRefresh();
-	void ExecuteSCCCheckOut();
-	void ExecuteSCCCheckIn();
-	void ExecuteSCCHistory();
-	void ExecuteSCCRevert();
-
 	TWeakPtr<ISceneOutliner> WeakSceneOutliner;
-
-	TArray<FSceneOutlinerTreeItemPtr> SelectedItems;
-
-	TMap<FSceneOutlinerTreeItemPtr, TSharedRef<SSourceControlWidget>> ItemWidgets;
-
-	bool bCanExecuteSCCCheckOut = false;
-	bool bCanExecuteSCCCheckIn = false;
-	bool bCanExecuteSCCHistory = false;
-	bool bCanExecuteSCCRevert = false;
 };
diff --git a/Engine/Source/Editor/SceneOutliner/Private/SceneOutlinerSCCHandler.cpp b/Engine/Source/Editor/SceneOutliner/Private/SceneOutlinerSCCHandler.cpp
new file mode 100644
--- /dev/null
+++ b/Engine/Source/Editor/SceneOutliner/Private/SceneOutlinerSCCHandler.cpp
@@ -0,0 +1,314 @@
+// Copyright Epic Games, Inc. All Rights Reserved.
+
+#include "SceneOutlinerSCCHandler.h"
+#include "ISourceControlModule.h"
+#include "SourceControlOperations.h"
+#include "Misc/MessageDialog.h"
+#include "FileHelpers.h"
+#include "SourceControlWindows.h"
+#include "AssetViewUtils.h"
+
+#define LOCTEXT_NAMESPACE "FSceneOutlinerSCCHandler"
+
+void FSceneOutlinerSCCHandler::AddItem(FSceneOutlinerTreeItemPtr Item)
+{
+	TSharedPtr<FSceneOutlinerTreeItemSCC> SourceControl(new FSceneOutlinerTreeItemSCC(Item));
+	ItemSourceControls.Add(Item, SourceControl);
+}
+
+TSharedPtr<FSceneOutlinerTreeItemSCC> FSceneOutlinerSCCHandler::GetItemSourceControl(const FSceneOutlinerTreeItemPtr& InItem)
+{
+	TSharedPtr<FSceneOutlinerTreeItemSCC>* Result = ItemSourceControls.Find(InItem);
+	return Result == nullptr ? nullptr : *Result;
+}
+
+bool FSceneOutlinerSCCHandler::AddSourceControlMenuOptions(UToolMenu* Menu, TArray<FSceneOutlinerTreeItemPtr> InSelectedItems)
+{
+	SelectedItems = InSelectedItems;
+
+	CacheCanExecuteVars();
+
+	if (ISourceControlModule::Get().IsEnabled())
+	{
+		FToolMenuSection& Section = Menu->AddSection("AssetContextSourceControl");
+
+		if (bCanExecuteSCC)
+		{
+			// SCC sub menu
+			Section.AddSubMenu(
+				"SourceControlSubMenu",
+				LOCTEXT("SourceControlSubMenuLabel", "Source Control"),
+				LOCTEXT("SourceControlSubMenuToolTip", "Source control actions."),
+				FNewToolMenuDelegate::CreateSP(this, &FSceneOutlinerSCCHandler::FillSourceControlSubMenu),
+				FUIAction(
+					FExecuteAction(),
+					FCanExecuteAction::CreateSP( this, &FSceneOutlinerSCCHandler::CanExecuteSourceControlActions )
+					),
+				EUserInterfaceActionType::Button,
+				false,
+				FSlateIcon(FAppStyle::GetAppStyleSetName(), "SourceControl.StatusIcon.On")
+				);
+		}
+		else
+		{
+			// SCC sub menu
+			Section.AddMenuEntry(
+				"SourceControlSubMenu",
+				LOCTEXT("SourceControlSubMenuLabel", "Source Control"),
+				LOCTEXT("SourceControlSubMenuToolTip", "Disabled because one or more selected items are not external packages."),
+				FSlateIcon(FAppStyle::GetAppStyleSetName(), "SourceControl.StatusIcon.Off"),
+				FUIAction(
+					FExecuteAction(),
+					FCanExecuteAction::CreateSP( this, &FSceneOutlinerSCCHandler::CanExecuteSCC )
+				)
+				);
+		}
+
+		return true;
+	}
+
+	return false;
+}
+
+
+bool FSceneOutlinerSCCHandler::CanExecuteSourceControlActions() const
+{
+	return SelectedItems.Num() > 0;
+}
+
+void FSceneOutlinerSCCHandler::CacheCanExecuteVars()
+{
+	bCanExecuteSCC = true;
+	bCanExecuteSCCCheckOut = false;
+	bCanExecuteSCCCheckIn = false;
+	bCanExecuteSCCHistory = false;
+
+	ISourceControlProvider& SourceControlProvider = ISourceControlModule::Get().GetProvider();
+	if ( ISourceControlModule::Get().IsEnabled() )
+	{
+		for (FSceneOutlinerTreeItemPtr SelectedItem : SelectedItems)
+		{
+			TSharedPtr<FSceneOutlinerTreeItemSCC>* SourceControlLookup = ItemSourceControls.Find(SelectedItem);
+			if (SourceControlLookup == nullptr)
+			{
+				continue;
+			}
+
+			TSharedPtr<FSceneOutlinerTreeItemSCC> SourceControl = *SourceControlLookup;
+
+			if (!SourceControl->IsExternalPackage())
+			{
+				// this isn't an external package so we can't do anything with source control
+				bCanExecuteSCC = false;
+
+				// skip computing the rest of this package, but don't break
+				// since getting a mixed external/internal selection is possible
+				// and we change the UI state to disabled to show the user why
+				// they aren't getting SCC actions
+				continue;
+			}
+
+			// Check the SCC state for each package in the selected paths
+			FSourceControlStatePtr SourceControlState = SourceControl->GetSourceControlState();
+			if(SourceControlState.IsValid())
+			{
+				if ( SourceControlState->CanCheckout() )
+				{
+					bCanExecuteSCCCheckOut = true;
+				}
+
+				if( SourceControlState->IsSourceControlled() && !SourceControlState->IsAdded() )
+				{
+					bCanExecuteSCCHistory = true;
+				}
+
+				if ( SourceControlState->CanCheckIn() )
+				{
+					bCanExecuteSCCCheckIn = true;
+				}
+			}
+		}
+	}
+}
+
+bool FSceneOutlinerSCCHandler::CanExecuteSCC() const
+{
+	return bCanExecuteSCC;
+}
+
+bool FSceneOutlinerSCCHandler::CanExecuteSCCCheckOut() const
+{
+	return bCanExecuteSCCCheckOut;
+}
+
+bool FSceneOutlinerSCCHandler::CanExecuteSCCCheckIn() const
+{
+	return bCanExecuteSCCCheckIn;
+}
+
+bool FSceneOutlinerSCCHandler::CanExecuteSCCHistory() const
+{
+	return bCanExecuteSCCHistory;
+}
+
+bool FSceneOutlinerSCCHandler::CanExecuteSCCRefresh() const
+{
+	return ISourceControlModule::Get().IsEnabled();
+}
+
+void FSceneOutlinerSCCHandler::FillSourceControlSubMenu(UToolMenu* Menu)
+{
+	FToolMenuSection& Section = Menu->AddSection("AssetSourceControlActions", LOCTEXT("AssetSourceControlActionsMenuHeading", "Source Control"));
+
+	if ( CanExecuteSCCCheckOut() )
+	{
+		Section.AddMenuEntry(
+			"SCCCheckOut",
+			LOCTEXT("SCCCheckOut", "Check Out"),
+			LOCTEXT("SCCCheckOutTooltip", "Checks out the selected asset from source control."),
+			FSlateIcon(FAppStyle::GetAppStyleSetName(), "SourceControl.Actions.CheckOut"),
+			FUIAction(
+				FExecuteAction::CreateSP( this, &FSceneOutlinerSCCHandler::ExecuteSCCCheckOut ),
+				FCanExecuteAction::CreateSP( this, &FSceneOutlinerSCCHandler::CanExecuteSCCCheckOut )
+			)
+		);
+	}
+
+	if ( CanExecuteSCCCheckIn() )
+	{
+		Section.AddMenuEntry(
+			"SCCCheckIn",
+			LOCTEXT("SCCCheckIn", "Check In"),
+			LOCTEXT("SCCCheckInTooltip", "Checks in the selected asset to source control."),
+			FSlateIcon(FAppStyle::GetAppStyleSetName(), "SourceControl.Actions.Submit"),
+			FUIAction(
+				FExecuteAction::CreateSP( this, &FSceneOutlinerSCCHandler::ExecuteSCCCheckIn ),
+				FCanExecuteAction::CreateSP( this, &FSceneOutlinerSCCHandler::CanExecuteSCCCheckIn )
+			)
+		);
+	}
+
+	Section.AddMenuEntry(
+		"SCCRefresh",
+		LOCTEXT("SCCRefresh", "Refresh"),
+		LOCTEXT("SCCRefreshTooltip", "Updates the source control status of the asset."),
+		FSlateIcon(FAppStyle::GetAppStyleSetName(), "SourceControl.Actions.Refresh"),
+		FUIAction(
+			FExecuteAction::CreateSP( this, &FSceneOutlinerSCCHandler::ExecuteSCCRefresh ),
+			FCanExecuteAction::CreateSP( this, &FSceneOutlinerSCCHandler::CanExecuteSCCRefresh )
+			)
+		);
+
+	if( CanExecuteSCCHistory() )
+	{
+		Section.AddMenuEntry(
+			"SCCHistory",
+			LOCTEXT("SCCHistory", "History"),
+			LOCTEXT("SCCHistoryTooltip", "Displays the source control revision history of the selected asset."),
+			FSlateIcon(FAppStyle::GetAppStyleSetName(), "SourceControl.Actions.History"),
+			FUIAction(
+				FExecuteAction::CreateSP( this, &FSceneOutlinerSCCHandler::ExecuteSCCHistory ),
+				FCanExecuteAction::CreateSP( this, &FSceneOutlinerSCCHandler::CanExecuteSCCHistory )
+			)
+		);
+	}
+}
+
+void FSceneOutlinerSCCHandler::GetSelectedPackageNames(TArray<FString>& OutPackageNames) const
+{
+	for (FSceneOutlinerTreeItemPtr SelectedItem : SelectedItems)
+	{
+		const TSharedPtr<FSceneOutlinerTreeItemSCC>* SourceControlLookup = ItemSourceControls.Find(SelectedItem);
+		if (SourceControlLookup == nullptr)
+		{
+			continue;
+		}
+
+		TSharedPtr<FSceneOutlinerTreeItemSCC> SourceControl = *SourceControlLookup;
+		FString PackageName = SourceControl->GetPackageName();
+		if (!PackageName.IsEmpty()) {
+			OutPackageNames.Add(PackageName);
+		}
+	}
+}
+
+void FSceneOutlinerSCCHandler::GetSelectedPackages(TArray<UPackage*>& OutPackages) const
+{
+	for (FSceneOutlinerTreeItemPtr SelectedItem : SelectedItems)
+	{
+		const TSharedPtr<FSceneOutlinerTreeItemSCC>* SourceControlLookup = ItemSourceControls.Find(SelectedItem);
+		if (SourceControlLookup == nullptr)
+		{
+			continue;
+		}
+
+		TSharedPtr<FSceneOutlinerTreeItemSCC> SourceControl = *SourceControlLookup;
+		UPackage* Package = SourceControl->GetPackage();
+		if (Package != nullptr) {
+			OutPackages.Add(Package);
+		}
+	}
+}
+
+void FSceneOutlinerSCCHandler::ExecuteSCCRefresh()
+{
+	TArray<FString> PackageNames;
+	GetSelectedPackageNames(PackageNames);
+
+	ISourceControlModule::Get().GetProvider().Execute(ISourceControlOperation::Create<FUpdateStatus>(), SourceControlHelpers::PackageFilenames(PackageNames), EConcurrency::Asynchronous);
+}
+
+void FSceneOutlinerSCCHandler::ExecuteSCCCheckOut()
+{
+	TArray<UPackage*> PackagesToCheckOut;
+	GetSelectedPackages(PackagesToCheckOut);
+
+	if ( PackagesToCheckOut.Num() > 0 )
+	{
+		FEditorFileUtils::CheckoutPackages(PackagesToCheckOut);
+	}
+}
+
+void FSceneOutlinerSCCHandler::ExecuteSCCCheckIn()
+{
+	TArray<UPackage*> Packages;
+	GetSelectedPackages(Packages);
+
+	// Prompt the user to ask if they would like to first save any dirty packages they are trying to check-in
+	const FEditorFileUtils::EPromptReturnCode UserResponse = FEditorFileUtils::PromptForCheckoutAndSave( Packages, true, true );
+
+	// If the user elected to save dirty packages, but one or more of the packages failed to save properly OR if the user
+	// canceled out of the prompt, don't follow through on the check-in process
+	const bool bShouldProceed = ( UserResponse == FEditorFileUtils::EPromptReturnCode::PR_Success || UserResponse == FEditorFileUtils::EPromptReturnCode::PR_Declined );
+	if ( bShouldProceed )
+	{
+		TArray<FString> PackageNames;
+		GetSelectedPackageNames(PackageNames);
+
+		const bool bUseSourceControlStateCache = true;
+		const bool bCheckinGood = FSourceControlWindows::PromptForCheckin(bUseSourceControlStateCache, PackageNames);
+
+		if (!bCheckinGood)
+		{
+			FMessageDialog::Open(EAppMsgType::Ok, NSLOCTEXT("UnrealEd", "SCC_Checkin_Failed", "Check-in failed as a result of save failure."));
+		}
+	}
+	else
+	{
+		// If a failure occurred, alert the user that the check-in was aborted. This warning shouldn't be necessary if the user cancelled
+		// from the dialog, because they obviously intended to cancel the whole operation.
+		if ( UserResponse == FEditorFileUtils::EPromptReturnCode::PR_Failure )
+		{
+			FMessageDialog::Open( EAppMsgType::Ok, NSLOCTEXT("UnrealEd", "SCC_Checkin_Aborted", "Check-in aborted as a result of save failure.") );
+		}
+	}
+}
+
+void FSceneOutlinerSCCHandler::ExecuteSCCHistory()
+{
+	TArray<FString> PackageNames;
+	GetSelectedPackageNames(PackageNames);
+	FSourceControlWindows::DisplayRevisionHistory(SourceControlHelpers::PackageFilenames(PackageNames));
+}
+
+#undef LOCTEXT_NAMESPACE
diff --git a/Engine/Source/Editor/SceneOutliner/Private/SceneOutlinerTreeItemSCC.cpp b/Engine/Source/Editor/SceneOutliner/Private/SceneOutlinerTreeItemSCC.cpp
new file mode 100644
--- /dev/null
+++ b/Engine/Source/Editor/SceneOutliner/Private/SceneOutlinerTreeItemSCC.cpp
@@ -0,0 +1,145 @@
+// Copyright Epic Games, Inc. All Rights Reserved.
+
+#include "SceneOutlinerTreeItemSCC.h"
+
+#include "ActorTreeItem.h"
+#include "ActorDescTreeItem.h"
+#include "ActorFolderTreeItem.h"
+#include "ISourceControlProvider.h"
+#include "ISourceControlModule.h"
+#include "WorldPartition/WorldPartitionActorDesc.h"
+
+FSceneOutlinerTreeItemSCC::FSceneOutlinerTreeItemSCC(FSceneOutlinerTreeItemPtr InTreeItemPtr)
+{
+	TreeItemPtr = InTreeItemPtr;
+
+	if (TreeItemPtr.IsValid())
+	{
+		if (FActorTreeItem* ActorItem = TreeItemPtr->CastTo<FActorTreeItem>())
+		{
+			if (AActor* Actor = ActorItem->Actor.Get())
+			{
+				if (Actor->IsPackageExternal())
+				{
+					ExternalPackageName = USourceControlHelpers::PackageFilename(Actor->GetExternalPackage());
+					ExternalPackage = Actor->GetExternalPackage();
+				}
+
+				ActorPackingModeChangedDelegateHandle = Actor->OnPackagingModeChanged.AddLambda([this](AActor* InActor, bool bExternal)
+					{
+						if (bExternal)
+						{
+							ExternalPackageName = USourceControlHelpers::PackageFilename(InActor->GetExternalPackage());
+							ExternalPackage = InActor->GetExternalPackage();
+							ConnectSourceControl();
+						}
+						else
+						{
+							ExternalPackageName = FString();
+							ExternalPackage = nullptr;
+							DisconnectSourceControl();
+						}
+					});
+			}
+		}
+		else if (FActorFolderTreeItem* ActorFolderItem = TreeItemPtr->CastTo<FActorFolderTreeItem>())
+		{
+			if (const UActorFolder* ActorFolder = ActorFolderItem->GetActorFolder())
+			{
+				if (ActorFolder->IsPackageExternal())
+				{
+					ExternalPackageName = USourceControlHelpers::PackageFilename(ActorFolder->GetExternalPackage());
+					ExternalPackage = ActorFolder->GetExternalPackage();
+				}
+			}
+		}
+		else if (FActorDescTreeItem* ActorDescItem = TreeItemPtr->CastTo<FActorDescTreeItem>())
+		{
+			if (const FWorldPartitionActorDesc* ActorDesc = ActorDescItem->ActorDescHandle.Get())
+			{
+				ExternalPackageName =  USourceControlHelpers::PackageFilename(ActorDesc->GetActorPackage().ToString());
+				ExternalPackage = FindPackage(nullptr, *ActorDesc->GetActorPackage().ToString());
+			}
+		}
+
+		if (!ExternalPackageName.IsEmpty())
+		{
+			ConnectSourceControl();
+		}
+	}
+}
+
+FSceneOutlinerTreeItemSCC::~FSceneOutlinerTreeItemSCC()
+{
+	DisconnectSourceControl();
+}
+
+FSourceControlStatePtr FSceneOutlinerTreeItemSCC::GetSourceControlState()
+{
+	return ISourceControlModule::Get().GetProvider().GetState(ExternalPackageName, EStateCacheUsage::Use);
+}
+
+FSourceControlStatePtr FSceneOutlinerTreeItemSCC::RefreshSourceControlState()
+{
+	return ISourceControlModule::Get().GetProvider().GetState(ExternalPackageName, EStateCacheUsage::ForceUpdate);
+}
+
+void FSceneOutlinerTreeItemSCC::ConnectSourceControl()
+{
+	check(!ExternalPackageName.IsEmpty());
+
+	ISourceControlModule& SCCModule = ISourceControlModule::Get();
+	SourceControlProviderChangedDelegateHandle = SCCModule.RegisterProviderChanged(FSourceControlProviderChanged::FDelegate::CreateRaw(this, &FSceneOutlinerTreeItemSCC::HandleSourceControlProviderChanged));
+	SourceControlStateChangedDelegateHandle = SCCModule.GetProvider().RegisterSourceControlStateChanged_Handle(FSourceControlStateChanged::FDelegate::CreateRaw(this, &FSceneOutlinerTreeItemSCC::HandleSourceControlStateChanged, EStateCacheUsage::Use));
+
+	// Check if there is already a cached state for this item
+	FSourceControlStatePtr SourceControlState = ISourceControlModule::Get().GetProvider().GetState(ExternalPackageName, EStateCacheUsage::Use);
+	if (SourceControlState.IsValid() && !SourceControlState->IsUnknown())
+	{
+		BroadcastNewState(SourceControlState);
+	}
+	else
+	{
+		SCCModule.QueueStatusUpdate(ExternalPackageName);
+	}
+}
+
+void FSceneOutlinerTreeItemSCC::DisconnectSourceControl()
+{
+	if (TreeItemPtr.IsValid())
+	{
+		if (FActorTreeItem* ActorItem = TreeItemPtr->CastTo<FActorTreeItem>())
+		{
+			if (AActor* Actor = ActorItem->Actor.Get())
+			{
+				Actor->OnPackagingModeChanged.Remove(ActorPackingModeChangedDelegateHandle);
+			}
+		}
+	}
+	ISourceControlModule::Get().GetProvider().UnregisterSourceControlStateChanged_Handle(SourceControlStateChangedDelegateHandle);
+	ISourceControlModule::Get().UnregisterProviderChanged(SourceControlProviderChangedDelegateHandle);
+}
+
+void FSceneOutlinerTreeItemSCC::HandleSourceControlStateChanged(EStateCacheUsage::Type CacheUsage)
+{
+	FSourceControlStatePtr SourceControlState = ISourceControlModule::Get().GetProvider().GetState(ExternalPackageName, CacheUsage);
+	if (SourceControlState.IsValid())
+	{
+		BroadcastNewState(SourceControlState);
+	}
+}
+
+void FSceneOutlinerTreeItemSCC::HandleSourceControlProviderChanged(ISourceControlProvider& OldProvider, ISourceControlProvider& NewProvider)
+{
+	OldProvider.UnregisterSourceControlStateChanged_Handle(SourceControlStateChangedDelegateHandle);
+	SourceControlStateChangedDelegateHandle = NewProvider.RegisterSourceControlStateChanged_Handle(FSourceControlStateChanged::FDelegate::CreateRaw(this, &FSceneOutlinerTreeItemSCC::HandleSourceControlStateChanged, EStateCacheUsage::Use));
+	
+	BroadcastNewState(nullptr);
+
+	ISourceControlModule::Get().QueueStatusUpdate(ExternalPackageName);
+}
+
+void FSceneOutlinerTreeItemSCC::BroadcastNewState(FSourceControlStatePtr SourceControlState)
+{
+	OnSourceControlStateChanged.ExecuteIfBound(SourceControlState);
+}
diff --git a/Engine/Source/Editor/SceneOutliner/Public/ISceneOutliner.h b/Engine/Source/Editor/SceneOutliner/Public/ISceneOutliner.h
--- a/Engine/Source/Editor/SceneOutliner/Public/ISceneOutliner.h
+++ b/Engine/Source/Editor/SceneOutliner/Public/ISceneOutliner.h
@@ -147,6 +147,9 @@ public:
 
 	/** Get the active SceneOutlinerMode */
 	const ISceneOutlinerMode* GetMode() const { return Mode; }
+
+	/** Get the associated source control object for the specified item. */
+	virtual TSharedPtr<FSceneOutlinerTreeItemSCC> GetItemSourceControl(const FSceneOutlinerTreeItemPtr& InItem) = 0;
 protected:
 	ISceneOutlinerMode* Mode;
 };
diff --git a/Engine/Source/Editor/SceneOutliner/Public/SSceneOutliner.h b/Engine/Source/Editor/SceneOutliner/Public/SSceneOutliner.h
--- a/Engine/Source/Editor/SceneOutliner/Public/SSceneOutliner.h
+++ b/Engine/Source/Editor/SceneOutliner/Public/SSceneOutliner.h
@@ -30,6 +30,9 @@
 #include "ISceneOutlinerHierarchy.h"
 #include "SceneOutlinerDragDrop.h"
 
+#include "SceneOutlinerSCCHandler.h"
+#include "SceneOutlinerTreeItemSCC.h"
+
 class FMenuBuilder;
 class UToolMenu;
 class ISceneOutlinerColumn;
@@ -907,6 +910,13 @@ private:
 	const FSceneOutlinerConfig* GetConstConfig() const;
 
 	void SaveConfig();
+
+	TSharedPtr<FSceneOutlinerSCCHandler> SourceControlHandler;
+
+public:
+	virtual TSharedPtr<FSceneOutlinerTreeItemSCC> GetItemSourceControl(const FSceneOutlinerTreeItemPtr& InItem) override;
+
+	void AddSourceControlMenuOptions(UToolMenu* Menu);
 };
 
 struct SCENEOUTLINER_API FSceneOutlinerMenuHelper
diff --git a/Engine/Source/Editor/SceneOutliner/Public/SceneOutlinerFwd.h b/Engine/Source/Editor/SceneOutliner/Public/SceneOutlinerFwd.h
--- a/Engine/Source/Editor/SceneOutliner/Public/SceneOutlinerFwd.h
+++ b/Engine/Source/Editor/SceneOutliner/Public/SceneOutlinerFwd.h
@@ -44,3 +44,5 @@ DECLARE_DELEGATE_RetVal_OneParam( TSharedRef<ISceneOutlinerColumn>, FCreateScene
 
 /** A delegate used to factory a new filter type */
 DECLARE_DELEGATE_RetVal( TSharedRef<FSceneOutlinerFilter>, FCreateSceneOutlinerFilter );
+
+class FSceneOutlinerTreeItemSCC;
diff --git a/Engine/Source/Editor/SceneOutliner/Public/SceneOutlinerSCCHandler.h b/Engine/Source/Editor/SceneOutliner/Public/SceneOutlinerSCCHandler.h
new file mode 100644
--- /dev/null
+++ b/Engine/Source/Editor/SceneOutliner/Public/SceneOutlinerSCCHandler.h
@@ -0,0 +1,49 @@
+// Copyright Epic Games, Inc. All Rights Reserved.
+
+#pragma once
+
+#include "CoreMinimal.h"
+#include "SceneOutlinerPublicTypes.h"
+#include "ToolMenu.h"
+#include "SceneOutlinerTreeItemSCC.h"
+
+class FSceneOutlinerSCCHandler : public TSharedFromThis<class FSceneOutlinerSCCHandler>
+{
+public:
+
+	void AddItem(FSceneOutlinerTreeItemPtr Item);
+
+	bool AddSourceControlMenuOptions(UToolMenu* Menu, TArray<FSceneOutlinerTreeItemPtr> InSelectedItems);
+
+	TSharedPtr<FSceneOutlinerTreeItemSCC> GetItemSourceControl(const FSceneOutlinerTreeItemPtr& InItem);
+
+private:
+
+	bool CanExecuteSourceControlActions() const;
+	void CacheCanExecuteVars();
+	bool CanExecuteSCC() const;
+	bool CanExecuteSCCCheckOut() const;
+	bool CanExecuteSCCCheckIn() const;
+	bool CanExecuteSCCHistory() const;
+	bool CanExecuteSCCRefresh() const;
+	void FillSourceControlSubMenu(UToolMenu* Menu);
+	void GetSelectedPackageNames(TArray<FString>& OutPackageNames) const;
+	void GetSelectedPackages(TArray<UPackage*>& OutPackages) const;
+	void ExecuteSCCRefresh();
+	void ExecuteSCCCheckOut();
+	void ExecuteSCCCheckIn();
+	void ExecuteSCCHistory();
+
+	TArray<FSceneOutlinerTreeItemPtr> SelectedItems;
+
+	// todo: Sync and Revert functionality is not enabled currently
+	// due to needing to unload/reload the actor and having to deal
+	// with any potential dependencies. You can read more about this
+	// decision here: https://github.com/EpicGames/UnrealEngine/pull/9310#issuecomment-1171505628
+	bool bCanExecuteSCC = false;
+	bool bCanExecuteSCCCheckOut = false;
+	bool bCanExecuteSCCCheckIn = false;
+	bool bCanExecuteSCCHistory = false;
+
+	TMap<FSceneOutlinerTreeItemPtr, TSharedPtr<FSceneOutlinerTreeItemSCC>> ItemSourceControls;
+};
diff --git a/Engine/Source/Editor/SceneOutliner/Public/SceneOutlinerTreeItemSCC.h b/Engine/Source/Editor/SceneOutliner/Public/SceneOutlinerTreeItemSCC.h
new file mode 100644
--- /dev/null
+++ b/Engine/Source/Editor/SceneOutliner/Public/SceneOutlinerTreeItemSCC.h
@@ -0,0 +1,60 @@
+// Copyright Epic Games, Inc. All Rights Reserved.
+
+#pragma once
+
+#include "SceneOutlinerPublicTypes.h"
+#include "SourceControlHelpers.h"
+
+#define LOCTEXT_NAMESPACE "SceneOutlinerTreeItemSCC"
+
+DECLARE_DELEGATE_OneParam(FSourceControlStateChangedDelegate, FSourceControlStatePtr);
+
+class FSceneOutlinerTreeItemSCC : public TSharedFromThis<FSceneOutlinerTreeItemSCC>
+{
+public:
+	FSceneOutlinerTreeItemSCC(FSceneOutlinerTreeItemPtr InTreeItemPtr);
+
+	~FSceneOutlinerTreeItemSCC();
+
+	FSourceControlStatePtr GetSourceControlState();
+
+	FSourceControlStatePtr RefreshSourceControlState();
+
+	bool IsExternalPackage() { return ExternalPackage != nullptr; }
+
+	FString GetPackageName() { return ExternalPackageName; }
+
+	UPackage* GetPackage() { return ExternalPackage; }
+
+	FSourceControlStateChangedDelegate OnSourceControlStateChanged;
+
+private:
+
+	void ConnectSourceControl();
+
+	void DisconnectSourceControl();
+
+	void HandleSourceControlStateChanged(EStateCacheUsage::Type CacheUsage);
+
+	void HandleSourceControlProviderChanged(ISourceControlProvider& OldProvider, ISourceControlProvider& NewProvider);
+
+	void BroadcastNewState(FSourceControlStatePtr SourceControlState);
+
+	/** The tree item we relate to */
+	FSceneOutlinerTreeItemPtr TreeItemPtr;
+
+	/** Cache the items external package name */
+	FString ExternalPackageName;
+	UPackage* ExternalPackage = nullptr;
+
+	/** Source control state changed delegate handle */
+	FDelegateHandle SourceControlStateChangedDelegateHandle;
+
+	/** Source control provider changed delegate handle */
+	FDelegateHandle SourceControlProviderChangedDelegateHandle;
+
+	/** Actor packaging mode changed delegate handle */
+	FDelegateHandle ActorPackingModeChangedDelegateHandle;
+};
+
+#undef LOCTEXT_NAMESPACE
diff --git a/Engine/Source/Editor/SceneOutliner/SceneOutliner.Build.cs b/Engine/Source/Editor/SceneOutliner/SceneOutliner.Build.cs
--- a/Engine/Source/Editor/SceneOutliner/SceneOutliner.Build.cs
+++ b/Engine/Source/Editor/SceneOutliner/SceneOutliner.Build.cs
@@ -25,7 +25,6 @@ public class SceneOutliner : ModuleRules
 				"ToolMenus",
 				"SourceControl",
 				"EditorConfig",
-				"SourceControl",
 				"SourceControlWindows",
 			}
 		);
