Fix text conflict from LOCTEXT macro from latest PR.

#rnx

diff --git a/Engine/Source/Editor/SceneOutliner/Private/SceneOutlinerSCCHandler.cpp b/Engine/Source/Editor/SceneOutliner/Private/SceneOutlinerSCCHandler.cpp
--- a/Engine/Source/Editor/SceneOutliner/Private/SceneOutlinerSCCHandler.cpp
+++ b/Engine/Source/Editor/SceneOutliner/Private/SceneOutlinerSCCHandler.cpp
@@ -55,7 +55,7 @@ bool FSceneOutlinerSCCHandler::AddSourceControlMenuOptions(UToolMenu* Menu, TArr
 			Section.AddMenuEntry(
 				"SourceControlSubMenu",
 				LOCTEXT("SourceControlSubMenuLabel", "Source Control"),
-				LOCTEXT("SourceControlSubMenuToolTip", "Disabled because one or more selected items are not external packages."),
+				LOCTEXT("SourceControlSubMenuDisabledToolTip", "Disabled because one or more selected items are not external packages."),
 				FSlateIcon(FAppStyle::GetAppStyleSetName(), "SourceControl.StatusIcon.Off"),
 				FUIAction(
 					FExecuteAction(),
