Level Instance: add sub level column in the world outliner.

#rnx

diff --git a/Engine/Source/Editor/SceneOutliner/Private/SceneOutlinerModule.cpp b/Engine/Source/Editor/SceneOutliner/Private/SceneOutlinerModule.cpp
--- a/Engine/Source/Editor/SceneOutliner/Private/SceneOutlinerModule.cpp
+++ b/Engine/Source/Editor/SceneOutliner/Private/SceneOutlinerModule.cpp
@@ -307,18 +307,42 @@ void FSceneOutlinerModule::CreateActorInfoColumns(FSceneOutlinerInitializationOp
 		return Builder.ToString();
 	});
 
+	FGetTextForItem SubPackageInfoText = FGetTextForItem::CreateLambda([](const ISceneOutlinerTreeItem& Item) -> FString
+	{
+		if (const FActorTreeItem* ActorItem = Item.CastTo<FActorTreeItem>())
+		{
+			if (AActor* Actor = ActorItem->Actor.Get())
+			{
+				if (const ILevelInstanceInterface* ActorAsLevelInstance = Cast<ILevelInstanceInterface>(Actor))
+				{
+					return ActorAsLevelInstance->GetWorldAssetPackage();
+				}
+			}
+		}
+		else if (const FActorDescTreeItem* ActorDescItem = Item.CastTo<FActorDescTreeItem>())
+		{
+			if (const FWorldPartitionActorDesc* ActorDesc = ActorDescItem->ActorDescHandle.Get())
+			{
+				if (FName LevelPackage = ActorDesc->GetLevelPackage(); !LevelPackage.IsNone())
+				{
+					return ActorDesc->GetLevelPackage().ToString();
+				}
+			}
+		}
+
+		return FString();
+	});
+
 	FGetTextForItem SocketInfoText = FGetTextForItem::CreateLambda([](const ISceneOutlinerTreeItem& Item) -> FString
 	{
 		if (const FActorTreeItem* ActorItem = Item.CastTo<FActorTreeItem>())
 		{
 			AActor* Actor = ActorItem->Actor.Get();
 
-			if (!Actor)
+			if (Actor)
 			{
-				return FString();
+				return Actor->GetAttachParentSocketName().ToString();
 			}
-
-			return Actor->GetAttachParentSocketName().ToString();
 		}
 		
 		return FString();
@@ -382,7 +406,6 @@ void FSceneOutlinerModule::CreateActorInfoColumns(FSceneOutlinerInitializationOp
 			}
 		}
 
-
 		return FString();
 	});
 
@@ -391,12 +414,10 @@ void FSceneOutlinerModule::CreateActorInfoColumns(FSceneOutlinerInitializationOp
 		if (const FActorTreeItem* ActorItem = Item.CastTo<FActorTreeItem>())
 		{
 			AActor* Actor = ActorItem->Actor.Get();
-			if (!Actor)
+			if (Actor)
 			{
-				return FString();
+				return FString::Printf(TEXT("%7d"), Actor->GetNumUncachedStaticLightingInteractions());
 			}
-
-			return FString::Printf(TEXT("%7d"), Actor->GetNumUncachedStaticLightingInteractions());
 		}
 		return FString();
 	});
@@ -433,6 +454,7 @@ void FSceneOutlinerModule::CreateActorInfoColumns(FSceneOutlinerInitializationOp
 	AddTextInfoColumn(FSceneOutlinerBuiltInColumnTypes::Level(), LevelColumnName, LevelInfoText);
 	AddTextInfoColumn(FSceneOutlinerBuiltInColumnTypes::Layer(), FSceneOutlinerBuiltInColumnTypes::Layer_Localized(), LayerInfoText);
 	AddTextInfoColumn(FSceneOutlinerBuiltInColumnTypes::DataLayer(), FSceneOutlinerBuiltInColumnTypes::DataLayer_Localized(), DataLayerInfoText);
+	AddTextInfoColumn(FSceneOutlinerBuiltInColumnTypes::SubPackage(), FSceneOutlinerBuiltInColumnTypes::SubPackage_Localized(), SubPackageInfoText);
 	AddTextInfoColumn(FSceneOutlinerBuiltInColumnTypes::Socket(), FSceneOutlinerBuiltInColumnTypes::Socket_Localized(), SocketInfoText);
 	AddTextInfoColumn(FSceneOutlinerBuiltInColumnTypes::IDName(), FSceneOutlinerBuiltInColumnTypes::IDName_Localized(), InternalNameInfoText);
 	AddTextInfoColumn(FSceneOutlinerBuiltInColumnTypes::UncachedLights(), FSceneOutlinerBuiltInColumnTypes::UncachedLights_Localized(), UncachedLightsInfoText);
diff --git a/Engine/Source/Editor/SceneOutliner/Public/SceneOutlinerPublicTypes.h b/Engine/Source/Editor/SceneOutliner/Public/SceneOutlinerPublicTypes.h
--- a/Engine/Source/Editor/SceneOutliner/Public/SceneOutlinerPublicTypes.h
+++ b/Engine/Source/Editor/SceneOutliner/Public/SceneOutlinerPublicTypes.h
@@ -42,6 +42,7 @@ struct FSceneOutlinerBuiltInColumnTypes
 	DEFINE_SCENEOUTLINER_BUILTIN_COLUMN_TYPE(Level, "Level", "SceneOutlinerLevelColumn", "Level");
 	DEFINE_SCENEOUTLINER_BUILTIN_COLUMN_TYPE(Layer, "Layer", "SceneOutlinerLayerColumn", "Layer");
 	DEFINE_SCENEOUTLINER_BUILTIN_COLUMN_TYPE(DataLayer, "Data Layer", "SceneOutlinerDataLayerColumn", "Data Layer");
+	DEFINE_SCENEOUTLINER_BUILTIN_COLUMN_TYPE(SubPackage, "Sub Package", "SceneOutlinerSubPackageColumn", "Sub Package");
 	DEFINE_SCENEOUTLINER_BUILTIN_COLUMN_TYPE(Pinned, "Pinned", "SceneOutlinerPinnedColumn", "Pinned");
 	DEFINE_SCENEOUTLINER_BUILTIN_COLUMN_TYPE(IDName, "ID Name", "SceneOutlinerIDNameColumn", "ID Name");
 	DEFINE_SCENEOUTLINER_BUILTIN_COLUMN_TYPE(PackageShortName, "Package Short Name", "SceneOutlinerPackageShortNameColumn", "Package Short Name");
diff --git a/Engine/Source/Runtime/Engine/Public/WorldPartition/LevelInstance/LevelInstanceActorDesc.h b/Engine/Source/Runtime/Engine/Public/WorldPartition/LevelInstance/LevelInstanceActorDesc.h
--- a/Engine/Source/Runtime/Engine/Public/WorldPartition/LevelInstance/LevelInstanceActorDesc.h
+++ b/Engine/Source/Runtime/Engine/Public/WorldPartition/LevelInstance/LevelInstanceActorDesc.h
@@ -23,9 +23,8 @@ public:
 	FLevelInstanceActorDesc();
 	virtual ~FLevelInstanceActorDesc() override;
 
-	inline FName GetLevelPackage() const { return LevelPackage; }
-
 	virtual bool IsContainerInstance() const override;
+	virtual FName GetLevelPackage() const override { return LevelPackage; }
 	virtual bool GetContainerInstance(const UActorDescContainer*& OutLevelContainer, FTransform& OutLevelTransform, EContainerClusterMode& OutClusterMode) const override;
 
 protected:
diff --git a/Engine/Source/Runtime/Engine/Public/WorldPartition/WorldPartitionActorDesc.h b/Engine/Source/Runtime/Engine/Public/WorldPartition/WorldPartitionActorDesc.h
--- a/Engine/Source/Runtime/Engine/Public/WorldPartition/WorldPartitionActorDesc.h
+++ b/Engine/Source/Runtime/Engine/Public/WorldPartition/WorldPartitionActorDesc.h
@@ -109,6 +109,7 @@ public:
 	FName GetDisplayClassName() const;
 
 	virtual bool IsContainerInstance() const { return false; }
+	virtual FName GetLevelPackage() const { return NAME_None; }
 	virtual bool GetContainerInstance(const UActorDescContainer*& OutLevelContainer, FTransform& OutLevelTransform, EContainerClusterMode& OutClusterMode) const { return false; }
 
 	virtual const FGuid& GetSceneOutlinerParent() const { return GetParentActor(); }
