SceneOutliner SCC Column: readd lost source control state Tooltip

diff --git a/Engine/Source/Editor/SceneOutliner/Private/SSourceControlWidget.cpp b/Engine/Source/Editor/SceneOutliner/Private/SSourceControlWidget.cpp
--- a/Engine/Source/Editor/SceneOutliner/Private/SSourceControlWidget.cpp
+++ b/Engine/Source/Editor/SceneOutliner/Private/SSourceControlWidget.cpp
@@ -8,14 +8,14 @@ void SSourceControlWidget::Construct(const FArguments& InArgs, TSharedPtr<FScene
 
 	ItemSourceControl = InItemSourceControl;
 
-	ItemSourceControl->OnSourceControlStateChanged.BindSP(this, &SSourceControlWidget::UpdateSourceControlStateIcon);
+	ItemSourceControl->OnSourceControlStateChanged.BindSP(this, &SSourceControlWidget::UpdateSourceControlState);
 
 	SImage::Construct(
 		SImage::FArguments()
 		.ColorAndOpacity(this, &SSourceControlWidget::GetForegroundColor)
 		.Image(FStyleDefaults::GetNoBrush()));
 
-	UpdateSourceControlStateIcon(ItemSourceControl->GetSourceControlState());
+	UpdateSourceControlState(ItemSourceControl->GetSourceControlState());
 }
 
 FReply SSourceControlWidget::OnMouseButtonDoubleClick(const FGeometry& InMyGeometry, const FPointerEvent& InMouseEvent)
@@ -23,21 +23,21 @@ FReply SSourceControlWidget::OnMouseButtonDoubleClick(const FGeometry& InMyGeome
 	check(ItemSourceControl);
 
 	FSourceControlStatePtr SourceControlState = ItemSourceControl->RefreshSourceControlState();
-	UpdateSourceControlStateIcon(SourceControlState);
+	UpdateSourceControlState(SourceControlState);
 	return FReply::Handled();
 }
 
-void SSourceControlWidget::UpdateSourceControlStateIcon(FSourceControlStatePtr SourceControlState)
+void SSourceControlWidget::UpdateSourceControlState(FSourceControlStatePtr SourceControlState)
 {
 	if(SourceControlState.IsValid())
 	{
-		FSlateIcon Icon = SourceControlState->GetIcon();
-		
-		SetFromSlateIcon(Icon);
+		SetFromSlateIcon(SourceControlState->GetIcon());
+		SetToolTipText(SourceControlState->GetDisplayTooltip());
 	}
 	else
 	{
 		SetImage(nullptr);
+		SetToolTipText(TAttribute<FText>());
 		RemoveAllLayers();
 	}
 }
diff --git a/Engine/Source/Editor/SceneOutliner/Private/SSourceControlWidget.h b/Engine/Source/Editor/SceneOutliner/Private/SSourceControlWidget.h
--- a/Engine/Source/Editor/SceneOutliner/Private/SSourceControlWidget.h
+++ b/Engine/Source/Editor/SceneOutliner/Private/SSourceControlWidget.h
@@ -19,7 +19,7 @@ private:
 
 	virtual FReply OnMouseButtonDoubleClick(const FGeometry& InMyGeometry, const FPointerEvent& InMouseEvent) override;
 
-	void UpdateSourceControlStateIcon(FSourceControlStatePtr SourceControlState);
+	void UpdateSourceControlState(FSourceControlStatePtr SourceControlState);
 
 	/** The object that keeps the source control state for the outliner */
 	TSharedPtr<FSceneOutlinerTreeItemSCC> ItemSourceControl;
