SceneOutliner: Fix regression causing fstats for every item even when SCC column is not displayed

#rnx

diff --git a/Engine/Source/Editor/SceneOutliner/Private/SSceneOutliner.cpp b/Engine/Source/Editor/SceneOutliner/Private/SSceneOutliner.cpp
--- a/Engine/Source/Editor/SceneOutliner/Private/SSceneOutliner.cpp
+++ b/Engine/Source/Editor/SceneOutliner/Private/SSceneOutliner.cpp
@@ -1759,8 +1759,6 @@ bool SSceneOutliner::IsItemExpanded(const FSceneOutlinerTreeItemPtr& Item) const
 
 TSharedRef< ITableRow > SSceneOutliner::OnGenerateRowForOutlinerTree( FSceneOutlinerTreeItemPtr Item, const TSharedRef< STableViewBase >& OwnerTable )
 {
-	SourceControlHandler->AddItem(Item);
-
 	return SNew( SSceneOutlinerTreeRow, OutlinerTreeView.ToSharedRef(), SharedThis(this) ).Item( Item );
 }
 
diff --git a/Engine/Source/Editor/SceneOutliner/Private/SceneOutlinerSCCHandler.cpp b/Engine/Source/Editor/SceneOutliner/Private/SceneOutlinerSCCHandler.cpp
--- a/Engine/Source/Editor/SceneOutliner/Private/SceneOutlinerSCCHandler.cpp
+++ b/Engine/Source/Editor/SceneOutliner/Private/SceneOutlinerSCCHandler.cpp
@@ -10,16 +10,16 @@
 
 #define LOCTEXT_NAMESPACE "FSceneOutlinerSCCHandler"
 
-void FSceneOutlinerSCCHandler::AddItem(FSceneOutlinerTreeItemPtr Item)
-{
-	TSharedPtr<FSceneOutlinerTreeItemSCC> SourceControl(new FSceneOutlinerTreeItemSCC(Item));
-	ItemSourceControls.Add(Item, SourceControl);
-}
-
 TSharedPtr<FSceneOutlinerTreeItemSCC> FSceneOutlinerSCCHandler::GetItemSourceControl(const FSceneOutlinerTreeItemPtr& InItem)
 {
 	TSharedPtr<FSceneOutlinerTreeItemSCC>* Result = ItemSourceControls.Find(InItem);
-	return Result == nullptr ? nullptr : *Result;
+	if (!Result)
+	{
+		Result = &ItemSourceControls.Add(InItem, MakeShared<FSceneOutlinerTreeItemSCC>(InItem));
+	}
+	
+	check(Result);
+	return *Result;
 }
 
 bool FSceneOutlinerSCCHandler::AddSourceControlMenuOptions(UToolMenu* Menu, TArray<FSceneOutlinerTreeItemPtr> InSelectedItems)
diff --git a/Engine/Source/Editor/SceneOutliner/Public/SceneOutlinerSCCHandler.h b/Engine/Source/Editor/SceneOutliner/Public/SceneOutlinerSCCHandler.h
--- a/Engine/Source/Editor/SceneOutliner/Public/SceneOutlinerSCCHandler.h
+++ b/Engine/Source/Editor/SceneOutliner/Public/SceneOutlinerSCCHandler.h
@@ -11,8 +11,6 @@ class FSceneOutlinerSCCHandler : public TSharedFromThis<class FSceneOutlinerSCCH
 {
 public:
 
-	void AddItem(FSceneOutlinerTreeItemPtr Item);
-
 	bool AddSourceControlMenuOptions(UToolMenu* Menu, TArray<FSceneOutlinerTreeItemPtr> InSelectedItems);
 
 	TSharedPtr<FSceneOutlinerTreeItemSCC> GetItemSourceControl(const FSceneOutlinerTreeItemPtr& InItem);
