Outliner: Change unloaded, but selected items to use the correct Foreground color to improve readibility

diff --git a/Engine/Source/Editor/SceneOutliner/Private/ActorDescTreeItem.cpp b/Engine/Source/Editor/SceneOutliner/Private/ActorDescTreeItem.cpp
--- a/Engine/Source/Editor/SceneOutliner/Private/ActorDescTreeItem.cpp
+++ b/Engine/Source/Editor/SceneOutliner/Private/ActorDescTreeItem.cpp
@@ -203,8 +203,17 @@ private:
 			{
 				return BaseColor.GetValue();
 			}
+			
+			if(WeakSceneOutliner.IsValid())
+			{
+				// Use the normal foreground color for selected items to make them readable
+				if(WeakSceneOutliner.Pin()->GetTree().IsItemSelected(TreeItem))
+				{
+					return FSlateColor::UseSubduedForeground();
+				}
+			}
 		}
-
+		
 		return FSceneOutlinerCommonLabelData::DarkColor;
 	}
 
