Double-click in the world outliner on an unloaded actor will focus on actor bounds and show its bounds for a short period of time.

diff --git a/Engine/Source/Editor/SceneOutliner/Private/ActorBrowsingMode.cpp b/Engine/Source/Editor/SceneOutliner/Private/ActorBrowsingMode.cpp
--- a/Engine/Source/Editor/SceneOutliner/Private/ActorBrowsingMode.cpp
+++ b/Engine/Source/Editor/SceneOutliner/Private/ActorBrowsingMode.cpp
@@ -908,6 +908,10 @@ void FActorBrowsingMode::OnItemDoubleClick(FSceneOutlinerTreeItemPtr Item)
 			GEditor->MoveViewportCamerasToActor(*Actor, bActiveViewportOnly);
 		}
 	}
+	else if (const FActorDescTreeItem* ActorDescItem = Item->CastTo<FActorDescTreeItem>())
+	{
+		ActorDescItem->FocusActorBounds();
+	}
 }
 
 void FActorBrowsingMode::OnFilterTextCommited(FSceneOutlinerItemSelection& Selection, ETextCommit::Type CommitType)
diff --git a/Engine/Source/Editor/SceneOutliner/Private/ActorDescTreeItem.cpp b/Engine/Source/Editor/SceneOutliner/Private/ActorDescTreeItem.cpp
--- a/Engine/Source/Editor/SceneOutliner/Private/ActorDescTreeItem.cpp
+++ b/Engine/Source/Editor/SceneOutliner/Private/ActorDescTreeItem.cpp
@@ -272,8 +272,11 @@ void FActorDescTreeItem::FocusActorBounds() const
 	{
 		if (FLevelEditorViewportClient* LevelViewportClient = GCurrentLevelEditingViewportClient)
 		{
-			LevelViewportClient->FocusViewportOnBox(ActorDesc->GetBounds(), false);
-		}	
+			FBox Bounds = ActorDesc->GetBounds();
+			LevelViewportClient->FocusViewportOnBox(Bounds, false);
+			const FColor DarkGray(40, 40, 40);
+			DrawDebugBox(ActorDesc->GetContainer()->GetWorld(), Bounds.GetCenter(), Bounds.GetExtent(), DarkGray, false, 0.5f);
+		}
 	}
 }
 
