Outliner: Add a filter bar and the ability to save searches as filters

diff --git a/Engine/Source/Developer/ToolWidgets/Public/Filters/SBasicFilterBar.h b/Engine/Source/Developer/ToolWidgets/Public/Filters/SBasicFilterBar.h
--- a/Engine/Source/Developer/ToolWidgets/Public/Filters/SBasicFilterBar.h
+++ b/Engine/Source/Developer/ToolWidgets/Public/Filters/SBasicFilterBar.h
@@ -52,6 +52,14 @@ enum class EFilterBarLayout : uint8
 	Vertical
 };
 
+/** Describes how each individual filter pill looks like */
+UENUM()
+enum class EFilterPillStyle : uint8
+{
+	Default, // The default style you see in the Content Browser, pills showing the color and filter name
+	Basic    // Simple filter pills that only show the name like you see in the Outliner
+};
+
 /** Forward declaration for a helper struct used to activate external filters */
 template<typename FilterType>
 struct FFrontendFilterExternalActivationHelper;
@@ -84,6 +92,7 @@ public:
  	SLATE_BEGIN_ARGS( SBasicFilterBar<FilterType> )
  		: _FilterBarLayout(EFilterBarLayout::Horizontal)
 		, _CanChangeOrientation(false)
+		, _FilterPillStyle(EFilterPillStyle::Default)
 	{}
 
  		/** Delegate for when filters have changed */
@@ -112,6 +121,9 @@ public:
 	
 		/** If true, allow dynamically changing the orientation and saving in the config */
 		SLATE_ARGUMENT(bool, CanChangeOrientation)
+
+		/** Determines how each individual filter pill looks like */
+		SLATE_ARGUMENT(EFilterPillStyle, FilterPillStyle)
 	
  	SLATE_END_ARGS()
 
@@ -122,6 +134,7 @@ public:
  		CreateTextFilter = InArgs._CreateTextFilter;
  		FilterBarLayout = InArgs._FilterBarLayout;
  		bCanChangeOrientation = InArgs._CanChangeOrientation;
+ 		FilterPillStyle = InArgs._FilterPillStyle;
 
  		// We use a widgetswitcher to allow dynamically swapping between the layouts
  		ChildSlot
@@ -364,6 +377,9 @@ protected:
 			/** Invoked when a request to remove all filters originated from within this filter */
 			SLATE_EVENT( FOnRequestRemoveAllButThis, OnRequestRemoveAllButThis )
 
+			/** Determines how each individual filter pill looks like */
+			SLATE_ARGUMENT(EFilterPillStyle, FilterPillStyle)
+
 		SLATE_END_ARGS()
 
 		/** Constructs this widget with InArgs */
@@ -387,7 +403,7 @@ protected:
 				FilterToolTip = TAttribute<FText>::Create(TAttribute<FText>::FGetter::CreateSP(FrontendFilter.ToSharedRef(), &FFilterBase<FilterType>::GetToolTipText));
 			}
 
-			Construct_Internal();
+			Construct_Internal(InArgs._FilterPillStyle);
 		}
 
 		/** Sets whether or not this filter is applied to the combined filter */
@@ -446,42 +462,66 @@ protected:
 	protected:
 
 		/** Function that constructs the actual widget for subclasses to call */
-		void Construct_Internal()
+		void Construct_Internal(EFilterPillStyle InFilterPillStyle)
 		{
-			ChildSlot
-			[
-				SNew(SBorder)
-				 .Padding(1.0f)
-				.BorderImage(FAppStyle::Get().GetBrush("ContentBrowser.FilterBackground"))
-				 [
-					SAssignNew( ToggleButtonPtr, SFilterCheckBox )
-					.Style(FAppStyle::Get(), "ContentBrowser.FilterButton")
+			TSharedPtr<SWidget> ContentWidget;
+
+			switch(InFilterPillStyle)
+			{
+			case EFilterPillStyle::Basic:
+				ContentWidget = SAssignNew( ToggleButtonPtr, SFilterCheckBox )
+					.Style(FAppStyle::Get(), "DetailsView.SectionButton")
 					.ToolTipText(FilterToolTip)
-					.Padding(0.0f)
 					.IsChecked(this, &SFilter::IsChecked)
 					.OnCheckStateChanged(this, &SFilter::FilterToggled)
 					.CheckBoxContentUsesAutoWidth(false)
 					.OnGetMenuContent(this, &SFilter::GetRightClickMenuContent)
 					[
-						SNew(SHorizontalBox)
-						+SHorizontalBox::Slot()
-						.VAlign(VAlign_Center)
-						.AutoWidth()
-						[
-							SNew(SImage)
-							.Image(FAppStyle::Get().GetBrush("ContentBrowser.FilterImage"))
-							.ColorAndOpacity(this, &SFilter::GetFilterImageColorAndOpacity)
-						]
-						+SHorizontalBox::Slot()
-						.Padding(TAttribute<FMargin>(this, &SFilter::GetFilterNamePadding))
-						.VAlign(VAlign_Center)
+						SNew(STextBlock)
+						.Margin(FMargin(0.0f))
+						.TextStyle(FAppStyle::Get(), "SmallText")
+						.Text(this, &SFilter::GetFilterDisplayName)
+					];
+				break;
+			case EFilterPillStyle::Default:
+			default:
+				ContentWidget = SNew(SBorder)
+				 .Padding(1.0f)
+				 .BorderImage(FAppStyle::Get().GetBrush("ContentBrowser.FilterBackground"))
+				 [
+					 SNew(SHorizontalBox)
+					 +SHorizontalBox::Slot()
+					 .VAlign(VAlign_Center)
+					 .AutoWidth()
+					 [
+						 SNew(SImage)
+						 .Image(FAppStyle::Get().GetBrush("ContentBrowser.FilterImage"))
+						 .ColorAndOpacity(this, &SFilter::GetFilterImageColorAndOpacity)
+					 ]
+					 +SHorizontalBox::Slot()
+					 .Padding(TAttribute<FMargin>(this, &SFilter::GetFilterNamePadding))
+					 .VAlign(VAlign_Center)
+					 [
+						 SAssignNew( ToggleButtonPtr, SFilterCheckBox )
+						.Style(FAppStyle::Get(), "ContentBrowser.FilterButton")
+						.ToolTipText(FilterToolTip)
+						.IsChecked(this, &SFilter::IsChecked)
+						.OnCheckStateChanged(this, &SFilter::FilterToggled)
+						.CheckBoxContentUsesAutoWidth(false)
+						.OnGetMenuContent(this, &SFilter::GetRightClickMenuContent)
 						[
 							SNew(STextBlock)
 							.Text(this, &SFilter::GetFilterDisplayName)
 							.IsEnabled_Lambda([this] {return bEnabled;})
 						]
-					]
-				]
+					 ]
+				 ];
+					
+			}
+			
+			ChildSlot
+			[
+				ContentWidget.ToSharedRef()
 			];
 
 			ToggleButtonPtr->SetOnFilterCtrlClicked(FOnClicked::CreateSP(this, &SFilter::FilterCtrlClicked));
@@ -864,10 +904,23 @@ protected:
 	/** Add a widget to the current filter layout */
 	void AddWidgetToLayout(const TSharedRef<SWidget> WidgetToAdd)
 	{
+		FMargin HorizontalFilterPadding;
+		
+		switch(FilterPillStyle)
+		{
+		case EFilterPillStyle::Basic:
+			HorizontalFilterPadding = FMargin(2);
+			break;
+		case EFilterPillStyle::Default:
+		default:
+			HorizontalFilterPadding = FMargin(3);
+		}
+			
+		
 		if(FilterBarLayout == EFilterBarLayout::Horizontal)
 		{
 			HorizontalFilterBox->AddSlot()
-			.Padding(3, 3)
+			.Padding(HorizontalFilterPadding)
 			[
 				WidgetToAdd
 			];
@@ -921,6 +974,7 @@ protected:
 		TSharedRef<SFilter> NewFilter =
 			SNew(SFilter)
 			.FrontendFilter(Filter)
+			.FilterPillStyle(FilterPillStyle)
 			.OnFilterChanged( this, &SBasicFilterBar<FilterType>::FrontendFilterChanged, Filter )
 			.OnRequestRemove(this, &SBasicFilterBar<FilterType>::RemoveFilterAndUpdate)
 			.OnRequestEnableOnly(this, &SBasicFilterBar<FilterType>::EnableOnlyThisFilter)
@@ -1568,6 +1622,9 @@ protected:
 
 	/** Whether the orientation can be changed after initailization */
 	bool bCanChangeOrientation;
+
+	/** Determines how each individual pill looks like */
+	EFilterPillStyle FilterPillStyle;
 	
 	friend struct FFrontendFilterExternalActivationHelper<FilterType>;
 };
diff --git a/Engine/Source/Editor/ContentBrowser/Private/SContentBrowser.cpp b/Engine/Source/Editor/ContentBrowser/Private/SContentBrowser.cpp
--- a/Engine/Source/Editor/ContentBrowser/Private/SContentBrowser.cpp
+++ b/Engine/Source/Editor/ContentBrowser/Private/SContentBrowser.cpp
@@ -709,7 +709,7 @@ void SContentBrowser::OnFilterBarLayoutChanging(EFilterBarLayout NewLayout)
 			+ SHorizontalBox::Slot()
 			.AutoWidth()
 			.VAlign(VAlign_Top)
-			.Padding(6, 4, 0, 0)
+			.Padding(6, 4, 2, 0)
 			[
 				FilterComboButton.ToSharedRef()
 			]
@@ -813,6 +813,7 @@ void SContentBrowser::OnFilterBarLayoutChanging(EFilterBarLayout NewLayout)
 
 					+ SHorizontalBox::Slot()
 					.AutoWidth()
+					.Padding(0, 0, 2, 0)
 					[
 						FilterComboButton.ToSharedRef()
 					]
diff --git a/Engine/Source/Editor/EditorWidgets/Private/Filters/CustomClassFilterData.cpp b/Engine/Source/Editor/EditorWidgets/Private/Filters/CustomClassFilterData.cpp
--- a/Engine/Source/Editor/EditorWidgets/Private/Filters/CustomClassFilterData.cpp
+++ b/Engine/Source/Editor/EditorWidgets/Private/Filters/CustomClassFilterData.cpp
@@ -10,7 +10,7 @@
 
 void FCustomClassFilterData::AddCategory(TSharedPtr<FFilterCategory> InCategory)
 {
-	Categories.Add(InCategory);
+	Categories.AddUnique(InCategory);
 }
 
 TSharedPtr<IAssetTypeActions> FCustomClassFilterData::GetAssetTypeActions() const
diff --git a/Engine/Source/Editor/EditorWidgets/Public/Filters/SAssetFilterBar.h b/Engine/Source/Editor/EditorWidgets/Public/Filters/SAssetFilterBar.h
--- a/Engine/Source/Editor/EditorWidgets/Public/Filters/SAssetFilterBar.h
+++ b/Engine/Source/Editor/EditorWidgets/Public/Filters/SAssetFilterBar.h
@@ -110,6 +110,7 @@ public:
 		: _UseDefaultAssetFilters(true)
 		, _FilterBarLayout(EFilterBarLayout::Horizontal)
 		, _CanChangeOrientation(false)
+		, _FilterPillStyle(EFilterPillStyle::Default)
 		{
 		
 		}
@@ -148,6 +149,9 @@ public:
 		
 		/** If true, allow dynamically changing the orientation and saving in the config */
 		SLATE_ARGUMENT(bool, CanChangeOrientation)
+
+		/** Determines how each individual filter pill looks like */
+		SLATE_ARGUMENT(EFilterPillStyle, FilterPillStyle)
 	
 	SLATE_END_ARGS()
 
@@ -167,6 +171,7 @@ public:
 		Args._FilterSearchBox = InArgs._FilterSearchBox;
 		Args._FilterBarLayout = InArgs._FilterBarLayout;
 		Args._CanChangeOrientation = InArgs._CanChangeOrientation;
+		Args._FilterPillStyle = InArgs._FilterPillStyle;
 		
 		SBasicFilterBar<FilterType>::Construct(Args);
 
@@ -294,6 +299,9 @@ protected:
 				FilterBarConfig->TypeFilters.Add(FilterName, Filter->IsEnabled());
 			}
 		}
+
+		// Add back all the unknown filters which we encountered while loading, but didn't find in our instance
+		FilterBarConfig->TypeFilters.Append(UnknownTypeFilters);
 	}
 
 	/** Load all the custom class filters (created by the user) from the specified config */
@@ -349,6 +357,9 @@ protected:
 			}
 		}
 
+		// Save a copy of all the type filters, and then remove anything we find
+		UnknownTypeFilters = FilterBarConfig->TypeFilters;
+
 		// Load all the type filters
 		for(const TSharedRef<FCustomClassFilterData> &CustomClassFilter : this->CustomClassFilters)
 		{
@@ -357,10 +368,13 @@ protected:
 				const FString FilterName = CustomClassFilter->GetFilterName();
 
 				// Try to find this type filter in our list of saved filters
-				if ( const bool* bIsActive = FilterBarConfig->TypeFilters.Find(CustomClassFilter->GetFilterName()) )
+				if ( const bool* bIsActive = FilterBarConfig->TypeFilters.Find(FilterName) )
 				{
 					TSharedRef<SFilter> NewFilter = this->AddAssetFilterToBar(CustomClassFilter);
 					NewFilter->SetEnabled(*bIsActive, false);
+
+					// Remove it from the Unknown list, we have encountered this type filter in the current instance
+					UnknownTypeFilters.Remove(FilterName);
 				}
 			}
 		}
@@ -409,6 +423,9 @@ protected:
 			/** Invoked when a request to remove all filters originated from within this filter */
 			SLATE_EVENT( FOnRequestRemoveAllButThis, OnRequestRemoveAllButThis )
 
+			/** Determines how each individual filter pill looks like */
+			SLATE_ARGUMENT(EFilterPillStyle, FilterPillStyle)
+
 		SLATE_END_ARGS()
 
 		/** Constructs this widget with InArgs */
@@ -439,7 +456,7 @@ protected:
 				this->FilterToolTip = TAttribute<FText>::Create(TAttribute<FText>::FGetter::CreateSP(this->FrontendFilter.ToSharedRef(), &FFilterBase<FilterType>::GetToolTipText));
 			}
 
-			SFilter::Construct_Internal();
+			SFilter::Construct_Internal(InArgs._FilterPillStyle);
 		}
 
 	public:
@@ -674,6 +691,7 @@ protected:
 	{
 		TSharedRef<SAssetFilter> NewFilter =
 			SNew(SAssetFilter)
+			.FilterPillStyle(this->FilterPillStyle)
 			.CustomClassFilter(CustomClassFilter)
 			.OnFilterChanged(this->OnFilterChanged)
 			.OnRequestRemove(this, &SAssetFilterBar<FilterType>::RemoveFilterAndUpdate)
@@ -1161,6 +1179,10 @@ protected:
 	/** The AssetClassAction used to get the permission list from the Asset Tools Module */
 	EAssetClassAction AssetClassAction = EAssetClassAction::ViewAsset;
 
+	/** A map of any type filters we encounter while calling LoadSettings(), but don't have in this instance of the filter
+	 *  bar. These are stored so we can save these when calling SaveSettings() so that they are not lost */
+	TMap<FString, bool> UnknownTypeFilters; 
+
 };
 
 #undef LOCTEXT_NAMESPACE
\ No newline at end of file
diff --git a/Engine/Source/Editor/EditorWidgets/Public/Filters/SFilterBar.h b/Engine/Source/Editor/EditorWidgets/Public/Filters/SFilterBar.h
--- a/Engine/Source/Editor/EditorWidgets/Public/Filters/SFilterBar.h
+++ b/Engine/Source/Editor/EditorWidgets/Public/Filters/SFilterBar.h
@@ -49,6 +49,7 @@ public:
 		, _UseDefaultAssetFilters(true)
 	    , _FilterBarLayout(EFilterBarLayout::Horizontal)
 		, _CanChangeOrientation(false)
+		, _FilterPillStyle(EFilterPillStyle::Default)
 		{
 		
 		}
@@ -93,6 +94,9 @@ public:
 			
 		/** If true, allow dynamically changing the orientation and saving in the config */
 		SLATE_ARGUMENT(bool, CanChangeOrientation)
+		
+		/** Determines how each individual filter pill looks like */
+		SLATE_ARGUMENT(EFilterPillStyle, FilterPillStyle)
 	
 	SLATE_END_ARGS()
 
@@ -110,6 +114,7 @@ public:
 		Args._FilterBarIdentifier = InArgs._FilterBarIdentifier;
 		Args._FilterBarLayout = InArgs._FilterBarLayout;
         Args._CanChangeOrientation = InArgs._CanChangeOrientation;
+		Args._FilterPillStyle = InArgs._FilterPillStyle;
 		
 		SAssetFilterBar<FilterType>::Construct(Args);
 		
diff --git a/Engine/Source/Editor/SceneOutliner/Private/ActorBrowsingMode.cpp b/Engine/Source/Editor/SceneOutliner/Private/ActorBrowsingMode.cpp
--- a/Engine/Source/Editor/SceneOutliner/Private/ActorBrowsingMode.cpp
+++ b/Engine/Source/Editor/SceneOutliner/Private/ActorBrowsingMode.cpp
@@ -37,6 +37,10 @@
 #include "EditorLevelUtils.h"
 #include "EditorViewportCommands.h"
 #include "SceneOutlinerActorSCCColumn.h"
+#include "Filters/GenericFilter.h"
+#include "Engine/StaticMeshActor.h"
+#include "Filters/CustomClassFilterData.h"
+#include "IPlacementModeModule.h"
 
 DEFINE_LOG_CATEGORY_STATIC(LogActorBrowser, Log, All);
 
@@ -2106,4 +2110,105 @@ void FActorBrowsingMode::SaveConfig()
 	UActorBrowserConfig::Get()->SaveEditorConfig();
 }
 
+void FActorBrowsingMode::CreateOutlinerFilterBarFilters(FSceneOutlinerFilterBarOptions& OutFilterBarOptions)
+{
+	OutFilterBarOptions.bHasFilterBar = true;
+
+	// All level editor outliners share their custom text filters
+	OutFilterBarOptions.bUseSharedSettings = true;
+
+	IPlacementModeModule& PlacementModeModule = IPlacementModeModule::Get();
+
+	// Get all the Categories
+	TArray<FPlacementCategoryInfo> Categories;
+	PlacementModeModule.GetSortedCategories(Categories);
+
+	// Array of unique class types to Add as filters
+	TMap<FTopLevelAssetPath, TSharedRef<FCustomClassFilterData>> FilterClasses;
+
+	// Remove the Recently Placed (doesn't make sense) and All Classes (too much bloat) categories from the filters
+	Categories.RemoveAll([](const FPlacementCategoryInfo& Category)
+	{
+		return Category.UniqueHandle == FBuiltInPlacementCategories::RecentlyPlaced() || Category.UniqueHandle == FBuiltInPlacementCategories::AllClasses();
+	});
+
+	// We will keep a reference to the shapes category if we find it for later
+	TSharedPtr<FFilterCategory> ShapesCategory;
+	
+	auto FindOrAddClassFilter = [&FilterClasses](UClass* Class, TSharedPtr<FFilterCategory> FilterCategory)
+	{
+		// If the underlying class already exists, just add this category to it
+		if(TSharedRef<FCustomClassFilterData>* ExistingClassData = FilterClasses.Find(Class->GetClassPathName()))
+		{
+			(*ExistingClassData)->AddCategory(FilterCategory);
+		}
+		else
+		{
+			TSharedRef<FCustomClassFilterData> NewClassData = MakeShared<FCustomClassFilterData>(Class, FilterCategory, FLinearColor::White);
+
+			FilterClasses.Add(Class->GetClassPathName(), NewClassData);
+		}
+	};
+	
+	for (const FPlacementCategoryInfo& Category : Categories)
+	{
+		// Make an FFilterCategory using the current PlacementCategory
+		TSharedPtr<FFilterCategory> FilterCategory = MakeShared<FFilterCategory>(Category.DisplayName, FText::GetEmpty());
+
+		if(Category.UniqueHandle == FBuiltInPlacementCategories::Shapes())
+		{
+			ShapesCategory = FilterCategory;
+		}
+
+		// Get all the items belonging to the current category
+		TArray<TSharedPtr<FPlaceableItem>> Items;
+		PlacementModeModule.RegenerateItemsForCategory(Category.UniqueHandle);
+		PlacementModeModule.GetItemsForCategory(Category.UniqueHandle, Items);
+
+		for(TSharedPtr<FPlaceableItem>& Item : Items)
+		{
+			// Get the underlying class from the Actor belonging to this item
+			const bool bIsClass = Item->AssetData.GetClass() == UClass::StaticClass();
+			const bool bIsActor = bIsClass ? CastChecked<UClass>(Item->AssetData.GetAsset())->IsChildOf(AActor::StaticClass()) : false;
+			AActor* DefaultActor = nullptr;
+			if (Item->Factory != nullptr)
+			{
+				DefaultActor = Item->Factory->GetDefaultActor(Item->AssetData);
+			}
+			else if (bIsActor)
+			{
+				DefaultActor = CastChecked<AActor>(CastChecked<UClass>(Item->AssetData.GetAsset())->ClassDefaultObject);
+			
+			}
+			if(!DefaultActor)
+			{
+				continue;
+			}
+			UClass* Class = DefaultActor->GetClass();
+
+			FindOrAddClassFilter(Class, FilterCategory);
+		}
+	}
+
+	if(ShapesCategory)
+	{
+		// Special case: Add Static Mesh Actor to the Shapes category
+		// TODO: Add a better way to register exception and extra classes if we want to deviate from the place actors set later on
+		FindOrAddClassFilter(AStaticMeshActor::StaticClass(), ShapesCategory);
+	}
+
+	// Sort the filter menu by name
+	FilterClasses.ValueSort([](const TSharedRef<FCustomClassFilterData>& ClassA, const TSharedRef<FCustomClassFilterData>& ClassB)
+		{
+			return ClassA->GetName().CompareTo(ClassB->GetName()) < 0;
+		});
+
+	// Now add all the classes to the filter bar
+	for(const TPair<FTopLevelAssetPath, TSharedRef<FCustomClassFilterData>>& FilterClass : FilterClasses)
+	{
+		OutFilterBarOptions.CustomClassFilters.Add(FilterClass.Value);
+	}
+	
+}
+
 #undef LOCTEXT_NAMESPACE
diff --git a/Engine/Source/Editor/SceneOutliner/Private/SSceneOutliner.cpp b/Engine/Source/Editor/SceneOutliner/Private/SSceneOutliner.cpp
--- a/Engine/Source/Editor/SceneOutliner/Private/SSceneOutliner.cpp
+++ b/Engine/Source/Editor/SceneOutliner/Private/SSceneOutliner.cpp
@@ -38,7 +38,8 @@
 #include "EditorFolderUtils.h"
 #include "SceneOutlinerConfig.h"
 #include "Algo/ForEach.h"
-
+#include "SceneOutlinerFilterBar.h"
+#include "ActorTreeItem.h"
 #include "Framework/Notifications/NotificationManager.h"
 #include "Widgets/Notifications/SNotificationList.h"
 
@@ -100,11 +101,19 @@ void SSceneOutliner::Construct(const FArguments& InArgs, const FSceneOutlinerIni
 	// @todo outliner: Should probably save this in layout!
 	// @todo outliner: Should save spacing for list view in layout
 
-	//Setup the SearchBox filter
+	// Setup the SearchBox
 	{
-		auto Delegate = SceneOutliner::TreeItemTextFilter::FItemToStringArray::CreateSP( this, &SSceneOutliner::PopulateSearchStrings );
-		SearchBoxFilter = MakeShareable( new SceneOutliner::TreeItemTextFilter( Delegate ) );
+		SearchBoxFilter = CreateTextFilter();
+		
+		FilterTextBoxWidget = SNew(SFilterSearchBox)
+		.Visibility( InInitOptions.bShowSearchBox ? EVisibility::Visible : EVisibility::Collapsed )
+		.HintText( LOCTEXT( "FilterSearch", "Search..." ) )
+		.ToolTipText( LOCTEXT("FilterSearchHint", "Type here to search (pressing enter selects the results)") )
+		.OnTextChanged( this, &SSceneOutliner::OnFilterTextChanged )
+		.OnTextCommitted( this, &SSceneOutliner::OnFilterTextCommitted );
 	}
+	
+	CreateFilterBar(InInitOptions.FilterBarOptions);
 
 	FSceneOutlinerModule& SceneOutlinerModule = FModuleManager::LoadModuleChecked<FSceneOutlinerModule>("SceneOutliner");
 	SceneOutlinerModule.OnColumnPermissionListChanged().AddSP(this, &SSceneOutliner::OnColumnPermissionListChanged);
@@ -138,12 +147,7 @@ void SSceneOutliner::Construct(const FArguments& InArgs, const FSceneOutlinerIni
 	Toolbar->AddSlot()
 	.VAlign(VAlign_Center)
 	[
-		SAssignNew( FilterTextBoxWidget, SSearchBox )
-		.Visibility( InInitOptions.bShowSearchBox ? EVisibility::Visible : EVisibility::Collapsed )
-		.HintText( LOCTEXT( "FilterSearch", "Search..." ) )
-		.ToolTipText( LOCTEXT("FilterSearchHint", "Type here to search (pressing enter selects the results)") )
-		.OnTextChanged( this, &SSceneOutliner::OnFilterTextChanged )
-		.OnTextCommitted( this, &SSceneOutliner::OnFilterTextCommitted )
+		FilterTextBoxWidget.ToSharedRef()
 	];
 
 	if (Mode->CanCustomizeToolbar())
@@ -164,7 +168,7 @@ void SSceneOutliner::Construct(const FArguments& InArgs, const FSceneOutlinerIni
 				.OnClicked(this, &SSceneOutliner::OnCreateFolderClicked)
 				[
 					SNew(SImage)
-	                .ColorAndOpacity(FSlateColor::UseForeground())
+					.ColorAndOpacity(FSlateColor::UseForeground())
 					.Image(FAppStyle::Get().GetBrush("SceneOutliner.NewFolderIcon"))
 				]
 			];
@@ -185,7 +189,7 @@ void SSceneOutliner::Construct(const FArguments& InArgs, const FSceneOutlinerIni
 			.ButtonContent()
 			[
 				SNew(SImage)
-	            .ColorAndOpacity(FSlateColor::UseForeground())
+				.ColorAndOpacity(FSlateColor::UseForeground())
 				.Image( FAppStyle::Get().GetBrush("Icons.Settings") )
 			]
 		];
@@ -198,6 +202,26 @@ void SSceneOutliner::Construct(const FArguments& InArgs, const FSceneOutlinerIni
 		Toolbar
 	];
 
+	// Add the FilterBar and the Add Filter button if it exists
+	if(FilterBar)
+	{
+		// Add Filter Menu
+		Toolbar->InsertSlot(0)
+		.VAlign(VAlign_Center)
+		.Padding(0.0f, 0.0f, 2.0f, 0.0f)
+		.AutoWidth()
+		[
+			SSceneOutlinerFilterBar::MakeAddFilterButton(FilterBar.ToSharedRef())
+		];
+
+		VerticalBox->AddSlot()
+		.AutoHeight()
+		.Padding( 0.0f, 0.0f, 0.0f, 4.0f )
+		[
+			FilterBar.ToSharedRef()
+		];
+	}
+
 	VerticalBox->AddSlot()
 	.FillHeight(1.0)
 	[
@@ -766,7 +790,7 @@ void SSceneOutliner::OnChildRemovedFromParent(ISceneOutlinerTreeItem& Parent)
 void SSceneOutliner::OnItemMoved(const FSceneOutlinerTreeItemRef& ReferenceItem)
 {
 	// Just remove the item if it no longer matches the filters
-	if (!ReferenceItem->Flags.bIsFilteredOut && !SearchBoxFilter->PassesFilter(*ReferenceItem))
+	if (!ReferenceItem->Flags.bIsFilteredOut && !PassesAllFilters(ReferenceItem))
 	{
 		// This will potentially remove any non-matching, empty parents as well
 		RemoveItemFromTree(ReferenceItem);
@@ -872,7 +896,7 @@ bool SSceneOutliner::AddItemToTree(FSceneOutlinerTreeItemRef Item)
 	}
 
 	// Set the filtered out flag
-	Item->Flags.bIsFilteredOut = !SearchBoxFilter->PassesFilter(*Item);
+	Item->Flags.bIsFilteredOut = !PassesAllFilters(Item);
 
 	if (!Item->Flags.bIsFilteredOut)
 	{
@@ -979,6 +1003,14 @@ void SSceneOutliner::PopulateSearchStrings(const ISceneOutlinerTreeItem& Item, T
 	}
 }
 
+/** Creates a TextFilter for ISceneOutlinerTreeItem */
+TSharedPtr< SceneOutliner::TreeItemTextFilter > SSceneOutliner::CreateTextFilter() const
+{
+	auto Delegate = SceneOutliner::TreeItemTextFilter::FItemToStringArray::CreateSP( this, &SSceneOutliner::PopulateSearchStrings );
+
+	return MakeShareable( new SceneOutliner::TreeItemTextFilter( Delegate ) );
+}
+
 void SSceneOutliner::GetSelectedFolders(TArray<FFolderTreeItem*>& OutFolders) const
 {
 	return FSceneOutlinerItemSelection(*OutlinerTreeView).Get<FFolderTreeItem>(OutFolders);
@@ -2014,7 +2046,7 @@ void SSceneOutliner::OnFilterTextCommitted( const FText& InFilterText, ETextComm
 			// Gather all of the items that match the filter text
 			for (auto& Pair : TreeItemMap)
 			{
-				Pair.Value->Flags.bIsFilteredOut = !SearchBoxFilter->PassesFilter(*Pair.Value);
+				Pair.Value->Flags.bIsFilteredOut = !PassesAllFilters(Pair.Value);
 				if (!Pair.Value->Flags.bIsFilteredOut)
 				{
 					Selection.Add(Pair.Value);
@@ -2387,4 +2419,54 @@ void SSceneOutliner::AddSourceControlMenuOptions(UToolMenu* Menu)
 	SourceControlHandler->AddSourceControlMenuOptions(Menu, SelectedItems);
 }
 
+void SSceneOutliner::OnFilterBarFilterChanged()
+{
+	FilterCollection = FilterBar->GetAllActiveFilters();
+	FilterBar->SaveSettings();
+	
+	FullRefresh();
+}
+
+bool SSceneOutliner::ConvertItemToAssetData(SceneOutliner::FilterBarType InItem, FAssetData &OutAssetData) const
+{
+	// Type filtering only supported for Actors currently
+	if (const FActorTreeItem* ActorItem = InItem.CastTo<FActorTreeItem>())
+	{
+		AActor* Actor = ActorItem->Actor.Get();
+
+		OutAssetData = FAssetData(Actor);
+
+		return true;
+	}
+
+	return false;
+}
+
+void SSceneOutliner::CreateFilterBar(const FSceneOutlinerFilterBarOptions& FilterBarOptions)
+{
+	if(!FilterBarOptions.bHasFilterBar)
+	{
+		return;
+	}
+
+	FName FilterBarIdentifier = OutlinerIdentifier;
+
+	SAssignNew(FilterBar, SSceneOutlinerFilterBar)
+	.OnConvertItemToAssetData(this, &SSceneOutliner::ConvertItemToAssetData)
+	.OnFilterChanged(this, &SSceneOutliner::OnFilterBarFilterChanged)
+	.CustomClassFilters(FilterBarOptions.CustomClassFilters)
+	.CustomFilters(FilterBarOptions.CustomFilters)
+	.FilterSearchBox(FilterTextBoxWidget)
+	.FilterBarIdentifier(FilterBarIdentifier)
+	.UseSharedSettings(FilterBarOptions.bUseSharedSettings)
+	.CreateTextFilter(SSceneOutlinerFilterBar::FCreateTextFilter::CreateLambda([this]()
+	{
+		TSharedPtr< SceneOutliner::TreeItemTextFilter > Filter = this->CreateTextFilter();
+
+		return MakeShareable(new FCustomTextFilter<SceneOutliner::FilterBarType>(Filter));
+	}));
+
+	FilterBar->LoadSettings();
+}
+
 #undef LOCTEXT_NAMESPACE
\ No newline at end of file
diff --git a/Engine/Source/Editor/SceneOutliner/Private/SceneOutlinerFilterBar.cpp b/Engine/Source/Editor/SceneOutliner/Private/SceneOutlinerFilterBar.cpp
new file mode 100644
--- /dev/null
+++ b/Engine/Source/Editor/SceneOutliner/Private/SceneOutlinerFilterBar.cpp
@@ -0,0 +1,277 @@
+// Copyright Epic Games, Inc. All Rights Reserved.
+
+#include "SceneOutlinerFilterBar.h"
+
+const FName SSceneOutlinerFilterBar::SharedIdentifier("SceneOutlinerFilterBarSharedSettings");
+SSceneOutlinerFilterBar::FCustomTextFilterEvent SSceneOutlinerFilterBar::CustomTextFilterEvent;
+
+void SSceneOutlinerFilterBar::Construct( const FArguments& InArgs )
+{
+	bUseSharedSettings = InArgs._UseSharedSettings;
+
+	SFilterBar<SceneOutliner::FilterBarType>::FArguments Args;
+	Args._OnFilterChanged = InArgs._OnFilterChanged;
+	Args._CustomFilters = InArgs._CustomFilters;
+	Args._UseDefaultAssetFilters = false;
+	Args._CustomClassFilters = InArgs._CustomClassFilters;
+	Args._CreateTextFilter = InArgs._CreateTextFilter;
+	Args._FilterSearchBox = InArgs._FilterSearchBox;
+	Args._FilterBarIdentifier = InArgs._FilterBarIdentifier;
+	Args._OnConvertItemToAssetData = InArgs._OnConvertItemToAssetData;
+	Args._FilterPillStyle = EFilterPillStyle::Basic;
+
+	SFilterBar<SceneOutliner::FilterBarType>::Construct(Args);
+	
+	/* If we are using shared settings, add a default config for the shared settings in case it doesnt exist
+	 * This needs to go after SAssetFilterBar<FAssetFilterType>::Construct() to ensure UFilterBarConfig is valid
+	 */
+	if(bUseSharedSettings)
+	{
+		UFilterBarConfig::Get()->FilterBars.FindOrAdd(SharedIdentifier);
+
+		// Bind our delegate for when another SSceneOutlinerFilterBar creates a custom text filter, so we can sync our list
+		CustomTextFilterEvent.AddSP(this, &SSceneOutlinerFilterBar::OnExternalCustomTextFilterCreated);
+	}
+}
+
+void SSceneOutlinerFilterBar::SaveSettings()
+{
+	// If this instance doesn't want to use the shared settings, save the settings normally
+	if(!bUseSharedSettings)
+	{
+		SAssetFilterBar<SceneOutliner::FilterBarType>::SaveSettings();
+		return;
+	}
+
+	if(FilterBarIdentifier.IsNone())
+	{
+		UE_LOG(LogSlate, Error, TEXT("SSceneOutlinerFilterBar Requires that you specify a FilterBarIdentifier to save settings"));
+		return;
+	}
+
+	// Get the settings unique to this instance and the common settings
+	FFilterBarSettings* InstanceSettings = &UFilterBarConfig::Get()->FilterBars.FindOrAdd(FilterBarIdentifier);
+	FFilterBarSettings* SharedSettings = &UFilterBarConfig::Get()->FilterBars.FindOrAdd(SharedIdentifier);
+
+	// Empty both the configs, we are just going to re-save everything there
+	InstanceSettings->Empty();
+	SharedSettings->Empty();
+
+	// Save all the programatically added filters normally
+	SaveFilters(InstanceSettings);
+
+	/** For each custom text filter: Save the filterdata into the common settings, so that all instances that use it
+	 *	are synced.
+	 *	For each CHECKED custom text filter: Save just the filter name, and the checked and active state into the
+	 *	instance settings. Those are specific to this instance (i.e we don't want a filter to be active in all
+	 *	instances if activated in one)
+	 */
+	for (const TSharedRef<ICustomTextFilter<SceneOutliner::FilterBarType>>& CustomTextFilter : this->CustomTextFilters)
+	{
+		// Get the actual FFilterBase
+		TSharedRef<FFilterBase<SceneOutliner::FilterBarType>> CustomFilter = CustomTextFilter->GetFilter().ToSharedRef();
+			
+		// Is the filter "checked", i.e visible in the filter bar
+		bool bIsChecked = IsFrontendFilterInUse(CustomFilter);
+
+		// Is the filter "active", i.e visible and enabled in the filter bar
+		bool bIsActive = IsFilterActive(CustomFilter);
+
+		// Get the data associated with this filter
+		FCustomTextFilterData FilterData = CustomTextFilter->CreateCustomTextFilterData();
+
+		// Just save the filter data into the shared settings
+		FCustomTextFilterState SharedFilterState;
+		SharedFilterState.FilterData = FilterData;
+		SharedSettings->CustomTextFilters.Add(SharedFilterState);
+
+		if(bIsChecked)
+		{
+			// Create a duplicate filter data that just contains the filter label for this instance to know
+			FCustomTextFilterData InstanceFilterData;
+			InstanceFilterData.FilterLabel = FilterData.FilterLabel;
+			
+			// Just save the filter name and enabled/active state into the shared settings
+			FCustomTextFilterState InstanceFilterState;
+			InstanceFilterState.bIsChecked = bIsChecked;
+			InstanceFilterState.bIsActive = bIsActive;
+			InstanceFilterState.FilterData = InstanceFilterData;
+			
+			InstanceSettings->CustomTextFilters.Add(InstanceFilterState);
+		}
+	}
+
+	SaveConfig();
+}
+
+void SSceneOutlinerFilterBar::LoadSettings()
+{
+	// If this instance doesn't want to use the shared settings, load the settings normally
+	if(!bUseSharedSettings)
+	{
+		SAssetFilterBar<SceneOutliner::FilterBarType>::LoadSettings();
+		return;
+	}
+
+	if(FilterBarIdentifier.IsNone())
+	{
+		UE_LOG(LogSlate, Error, TEXT("SSceneOutlinerFilterBar Requires that you specify a FilterBarIdentifier to load settings"));
+		return;
+	}
+
+	// Get the settings unique to this instance and the common settings
+	const FFilterBarSettings* InstanceSettings = UFilterBarConfig::Get()->FilterBars.Find(FilterBarIdentifier);
+	const FFilterBarSettings* SharedSettings = UFilterBarConfig::Get()->FilterBars.Find(SharedIdentifier);
+
+	// Load the filters specified programatically normally
+	LoadFilters(InstanceSettings);
+
+	// Load the custom text filters from the shared settings
+	LoadCustomTextFilters(SharedSettings);
+	
+	// From the instance settings, get each checked filter and set the checked and active state
+	for(const FCustomTextFilterState& FilterState : InstanceSettings->CustomTextFilters)
+	{
+		RestoreCustomTextFilterState(FilterState);
+	}
+
+	this->OnFilterChanged.ExecuteIfBound();
+}
+
+void SSceneOutlinerFilterBar::OnCreateCustomTextFilter(const FCustomTextFilterData& InFilterData, bool bApplyFilter)
+{
+	SAssetFilterBar<SceneOutliner::FilterBarType>::OnCreateCustomTextFilter(InFilterData, bApplyFilter);
+
+	// If we are using shared settings (i.e sharing custom text filters) broadcast the event for all other instances to update
+	if(bUseSharedSettings)
+	{
+		// First save the shared settings for other instances to use
+		SaveSettings();
+		CustomTextFilterEvent.Broadcast(AsShared());
+	}
+}
+
+void SSceneOutlinerFilterBar::OnModifyCustomTextFilter(const FCustomTextFilterData& InFilterData, TSharedPtr<ICustomTextFilter<SceneOutliner::FilterBarType>> InFilter)
+{
+	SAssetFilterBar<SceneOutliner::FilterBarType>::OnModifyCustomTextFilter(InFilterData, InFilter);
+
+	// If we are using shared settings (i.e sharing custom text filters) broadcast the event for all other instances to update
+	if(bUseSharedSettings)
+	{
+		// First save the shared settings for other instances to use
+		SaveSettings();
+		CustomTextFilterEvent.Broadcast(AsShared());
+	}
+}
+
+void SSceneOutlinerFilterBar::OnDeleteCustomTextFilter(const TSharedPtr<ICustomTextFilter<SceneOutliner::FilterBarType>> InFilter)
+{
+	SAssetFilterBar<SceneOutliner::FilterBarType>::OnDeleteCustomTextFilter(InFilter);
+
+	// If we are using shared settings (i.e sharing custom text filters) broadcast the event for all other instances to update
+	if(bUseSharedSettings)
+	{
+		// First save the shared settings for other instances to use
+		SaveSettings();
+		CustomTextFilterEvent.Broadcast(AsShared());
+	}
+}
+
+void SSceneOutlinerFilterBar::LoadCustomTextFilters(const FFilterBarSettings* FilterBarConfig)
+{
+	CustomTextFilters.Empty();
+	
+	// Extract just the filter data from the common settings
+	for(const FCustomTextFilterState& FilterState : FilterBarConfig->CustomTextFilters)
+	{
+		// Create an ICustomTextFilter using the provided delegate
+		TSharedRef<ICustomTextFilter<SceneOutliner::FilterBarType>> NewTextFilter = this->CreateTextFilter.Execute().ToSharedRef();
+
+		// Get the actual FFilterBase
+		TSharedRef<FFilterBase<SceneOutliner::FilterBarType>> NewFilter = NewTextFilter->GetFilter().ToSharedRef();
+
+		// Set the internals of the custom text filter from what we have saved
+		NewTextFilter->SetFromCustomTextFilterData(FilterState.FilterData);
+
+		// Add this to our list of custom text filters
+		CustomTextFilters.Add(NewTextFilter);
+	}
+}
+
+
+void SSceneOutlinerFilterBar::RestoreCustomTextFilterState(const FCustomTextFilterState& InFilterState)
+{
+	// Find the filter associated with the current instance data from our list of custom text filters
+	TSharedRef< ICustomTextFilter<SceneOutliner::FilterBarType> >* Filter =
+		CustomTextFilters.FindByPredicate([&InFilterState](const TSharedRef< ICustomTextFilter<SceneOutliner::FilterBarType> >& Element)
+	{
+		return Element->CreateCustomTextFilterData().FilterLabel.EqualTo(InFilterState.FilterData.FilterLabel);
+	});
+
+	// InFilterState should always be a filter that already exists!
+	check(Filter);
+
+	// Get the actual FFilterBase
+	TSharedRef<FFilterBase<SceneOutliner::FilterBarType>> ActualFilter = Filter->Get().GetFilter().ToSharedRef();
+
+	// Add it to the filter bar, since if it exists in this list it is checked
+	TSharedRef<SFilter> AddedFilter = this->AddFilterToBar(ActualFilter);
+
+	// Set the filter as active if it was previously
+	AddedFilter->SetEnabled(InFilterState.bIsActive, false);
+	this->SetFrontendFilterActive(ActualFilter, InFilterState.bIsActive);
+}
+
+void SSceneOutlinerFilterBar::OnExternalCustomTextFilterCreated(TSharedPtr<SWidget> BroadcastingFilterBar)
+{
+	// Do nothing if we aren't using shared settings or if the event was broadcasted by this filter list
+	if(!bUseSharedSettings || BroadcastingFilterBar == AsShared())
+	{
+		return;
+	}
+
+	/* We are going to remove all our custom text filters and re-load them from the shared settings, since a different
+	 * instance modified them.
+	 */
+
+	// To preserve the state of any checked/active custom text filters
+	TArray<FCustomTextFilterState> CurrentCustomTextFilterStates;
+	
+	for (const TSharedRef<ICustomTextFilter<SceneOutliner::FilterBarType>>& CustomTextFilter : this->CustomTextFilters)
+	{
+		// Get the actual FFilterBase
+		TSharedRef<FFilterBase<SceneOutliner::FilterBarType>> CustomFilter = CustomTextFilter->GetFilter().ToSharedRef();
+			
+		// Is the filter "checked", i.e visible in the filter bar
+		bool bIsChecked = IsFrontendFilterInUse(CustomFilter);
+
+		// Is the filter "active", i.e visible and enabled in the filter bar
+		bool bIsActive = IsFilterActive(CustomFilter);
+
+		// Only save the state if the filter is checked so we can restore it
+		if(bIsChecked)
+		{
+			/* Remove the filter from the list (calling SBasicFilterBar::RemoveFilter because we get a compiler error
+			*  due to SAssetFilterBar overriding RemoveFilter that takes in an SFilter that hides the parent class function
+			*/
+			SBasicFilterBar<SceneOutliner::FilterBarType>::RemoveFilter(CustomFilter, false);
+			
+			FCustomTextFilterState FilterState;
+			FilterState.FilterData = CustomTextFilter->CreateCustomTextFilterData();
+			FilterState.bIsChecked = bIsChecked;
+			FilterState.bIsActive = bIsActive;
+			
+			CurrentCustomTextFilterStates.Add(FilterState);
+		}
+	}
+
+	// Get the shared settings and reload the filters
+	FFilterBarSettings* SharedSettings = &UFilterBarConfig::Get()->FilterBars.FindOrAdd(SharedIdentifier);
+	LoadCustomTextFilters(SharedSettings);
+
+	// Restore the state of any previously active ones
+	for(const FCustomTextFilterState& SavedFilterState : CurrentCustomTextFilterStates)
+	{
+		RestoreCustomTextFilterState(SavedFilterState);
+	}
+}
\ No newline at end of file
diff --git a/Engine/Source/Editor/SceneOutliner/Private/SceneOutlinerFilterBar.h b/Engine/Source/Editor/SceneOutliner/Private/SceneOutlinerFilterBar.h
new file mode 100644
--- /dev/null
+++ b/Engine/Source/Editor/SceneOutliner/Private/SceneOutlinerFilterBar.h
@@ -0,0 +1,99 @@
+// Copyright Epic Games, Inc. All Rights Reserved.
+
+#pragma once
+
+#include "Filters/SFilterBar.h"
+#include "SceneOutlinerPublicTypes.h"
+
+class SSceneOutlinerFilterBar : public SFilterBar< SceneOutliner::FilterBarType > 
+{
+public:
+	
+	/**
+	 * An event delegate that is executed when a custom text filter has been created/modified/deleted in any SSceneOutlinerFilterBar
+	 * that is using shared settings.
+	 */
+	DECLARE_EVENT_OneParam(SSceneOutlinerFilterBar, FCustomTextFilterEvent, TSharedPtr<SWidget> /* BroadcastingFilterBar */);
+	
+	using FOnFilterChanged = typename SBasicFilterBar<SceneOutliner::FilterBarType>::FOnFilterChanged;
+	using FConvertItemToAssetData = typename FAssetFilter<SceneOutliner::FilterBarType>::FConvertItemToAssetData;
+	using FCreateTextFilter = typename SBasicFilterBar<SceneOutliner::FilterBarType>::FCreateTextFilter;
+	
+	SLATE_BEGIN_ARGS( SSceneOutlinerFilterBar )
+		: _FilterBarIdentifier(NAME_None)
+		, _UseDefaultAssetFilters(true)
+		{
+		
+		}
+
+		/** Delegate for when filters have changed */
+		SLATE_EVENT( FOnFilterChanged, OnFilterChanged )
+
+		/** Specify this delegate to use asset comparison through IAssetRegistry, where you specify how to convert your Item into an FAssetData */
+		SLATE_EVENT( FConvertItemToAssetData, OnConvertItemToAssetData )
+
+		/** Initial List of Custom Filters that will be added to the AddFilter Menu */
+		SLATE_ARGUMENT( TArray<TSharedRef<FFilterBase<SceneOutliner::FilterBarType>>>, CustomFilters)
+
+		/** A unique identifier for this filter bar needed to enable saving settings in a config file */
+		SLATE_ARGUMENT(FName, FilterBarIdentifier)
+	
+		/** Initial List of Custom Class Filters that will be added to the AddFilter Menu */
+		SLATE_ARGUMENT( TArray<TSharedRef<FCustomClassFilterData>>, CustomClassFilters)
+	
+		/** Whether the filter bar should provide the default asset filters */
+		SLATE_ARGUMENT(bool, UseDefaultAssetFilters)
+
+		/** A delegate to create a TTextFilter for FilterType items. If provided, will allow creation of custom text filters
+		 *  from the filter dropdown menu.
+		 */
+		SLATE_ARGUMENT(FCreateTextFilter, CreateTextFilter)
+			
+		/** An SFilterSearchBox that can be attached to this filter bar. When provided along with a CreateTextFilter
+		 *  delegate, allows the user to save searches from the Search Box as text filters for the filter bar.
+		 *	NOTE: Will bind a delegate to SFilterSearchBox::OnClickedAddSearchHistoryButton
+		 */
+		SLATE_ARGUMENT(TSharedPtr<SFilterSearchBox>, FilterSearchBox)
+
+		/** If true, CustomTextFilters are saved/loaded from a config shared between multiple Filterbars
+		 *	Currently used to sync all level editor outliner custom text filters.
+		 */
+		SLATE_ARGUMENT(bool, UseSharedSettings)
+
+	SLATE_END_ARGS()
+
+	void Construct( const FArguments& InArgs );
+
+	virtual void SaveSettings() override;
+	virtual void LoadSettings() override;
+
+protected:
+
+	/** Handler for when a custom text filter is created */
+	virtual void OnCreateCustomTextFilter(const FCustomTextFilterData& InFilterData, bool bApplyFilter) override;
+	/** Handler for when a custom text filter is modified */
+	virtual void OnModifyCustomTextFilter(const FCustomTextFilterData& InFilterData, TSharedPtr<ICustomTextFilter<SceneOutliner::FilterBarType>> InFilter) override;
+	/** Handler for when a custom text filter is deleted */
+	virtual void OnDeleteCustomTextFilter(const TSharedPtr<ICustomTextFilter<SceneOutliner::FilterBarType>> InFilter) override;
+
+private:
+	/** Handler for when another Filter Bar using shared settings creates a custom text filter */
+	void OnExternalCustomTextFilterCreated(TSharedPtr<SWidget> BroadcastingFilterBar);
+	
+	/** Empty our list of custom text filters, and load from the given config */
+	void LoadCustomTextFilters(const FFilterBarSettings* FilterBarConfig);
+
+	/** Find the custom text filter corresponding to the specified state, and restore it's state to what is specified */
+	void RestoreCustomTextFilterState(const FCustomTextFilterState& InFilterState);
+	
+private:
+	/** The event that executes when a custom text filter is created/modified/deleted in any filter list that is using Shared Settings */
+	static FCustomTextFilterEvent CustomTextFilterEvent;
+
+	/** An identifier shared by all SSceneOutlinerFilterBars, used to save and load settings common to every instance */
+	static const FName SharedIdentifier;
+
+	/** Whether this Filter bar wants to load/save from the settings common to all instances */
+	bool bUseSharedSettings;
+	
+};
\ No newline at end of file
diff --git a/Engine/Source/Editor/SceneOutliner/Private/SceneOutlinerModule.cpp b/Engine/Source/Editor/SceneOutliner/Private/SceneOutlinerModule.cpp
--- a/Engine/Source/Editor/SceneOutliner/Private/SceneOutlinerModule.cpp
+++ b/Engine/Source/Editor/SceneOutliner/Private/SceneOutlinerModule.cpp
@@ -172,6 +172,10 @@ TSharedRef< ISceneOutliner > FSceneOutlinerModule::CreateActorBrowser(const FSce
 		}
 		
 	}
+
+	// Setup the FilterBar options for the Actor Browser
+	FActorBrowsingMode::CreateOutlinerFilterBarFilters(InitOptions.FilterBarOptions);
+	
 	return CreateSceneOutliner(InitOptions);
 }
 
diff --git a/Engine/Source/Editor/SceneOutliner/Public/ActorBrowsingMode.h b/Engine/Source/Editor/SceneOutliner/Public/ActorBrowsingMode.h
--- a/Engine/Source/Editor/SceneOutliner/Public/ActorBrowsingMode.h
+++ b/Engine/Source/Editor/SceneOutliner/Public/ActorBrowsingMode.h
@@ -98,6 +98,9 @@ public:
 
 	/** Called by engine when edit delete actors ends */
 	void OnDeleteActorsEnd();
+
+	/* Helper function to create the basic set of filters that the Actor Browser will have */
+	static void CreateOutlinerFilterBarFilters(FSceneOutlinerFilterBarOptions& OutFilterBarOptions);
 private:
 	/** Build and up the context menu */
 	TSharedPtr<SWidget> BuildContextMenu();
diff --git a/Engine/Source/Editor/SceneOutliner/Public/SSceneOutliner.h b/Engine/Source/Editor/SceneOutliner/Public/SSceneOutliner.h
--- a/Engine/Source/Editor/SceneOutliner/Public/SSceneOutliner.h
+++ b/Engine/Source/Editor/SceneOutliner/Public/SSceneOutliner.h
@@ -39,8 +39,9 @@ class ISceneOutlinerColumn;
 class SComboButton;
 class ULevel;
 struct FToolMenuSection;
-
+template<typename FilterType> class SFilterBar;
 template<typename ItemType> class STreeView;
+class SFilterSearchBox;
 
 /**
  * Scene Outliner definition
@@ -519,13 +520,19 @@ private:
 		{
 			const TreeItemType Temporary(Data);
 			bool bPassesFilters = Filters->PassesAllFilters(Temporary);
+
+			if(FilterCollection)
+			{
+				bPassesFilters &= FilterCollection->PassesAllFilters(Temporary);
+			}
+			
 			if (bPassesFilters)
 			{
 				OnItemPassesFilters(Temporary);
 			}
 
 		bPassesFilters &= SearchBoxFilter->PassesFilter(Temporary);
-
+			
 		if (bForce || bPassesFilters)
 		{
 			FSceneOutlinerTreeItemPtr Result = MakeShareable(new TreeItemType(Data));
@@ -558,7 +565,20 @@ private:
 
 	bool PassesTextFilter(const FSceneOutlinerTreeItemPtr& Item) const
 	{
-		return SearchBoxFilter->PassesFilter(*Item);
+		return PassesAllFilters(Item);
+	}
+
+	bool PassesAllFilters(const FSceneOutlinerTreeItemPtr& Item) const
+	{
+		bool bPassesFilters = SearchBoxFilter->PassesFilter(*Item);
+
+		if(FilterCollection)
+		{
+			bPassesFilters &= FilterCollection->PassesAllFilters(*Item);
+		}
+
+		return bPassesFilters;
+		
 	}
 
 	bool HasSelectorFocus(FSceneOutlinerTreeItemPtr Item) const
@@ -584,7 +604,22 @@ private:
 	/** Populates OutSearchStrings with the strings associated with TreeItem that should be used in searching */
 	void PopulateSearchStrings( const ISceneOutlinerTreeItem& TreeItem, OUT TArray< FString >& OutSearchStrings ) const;
 
+	/** Filter Bar related private functionality */
+	
+	/** Creates a TextFilter for ISceneOutlinerTreeItem used to save searches as Text Filters */
+	TSharedPtr< SceneOutliner::TreeItemTextFilter > CreateTextFilter() const;
+
+	/** Convert a SceneOutliner item into an FAssetData struct to compare against type filters *
+	 *  TODO: Move to the ISceneOutlinerMode maybe to allow each mode to specify how to convert items?
+	 */
+	bool ConvertItemToAssetData(SceneOutliner::FilterBarType InItem, FAssetData &OutAssetData) const;
+
+	/** Delegate for when a filter in the filter bar is changed */
+	void OnFilterBarFilterChanged();
 
+	/** Create the filter bar for this outliner using the specified init options */
+	void CreateFilterBar(const FSceneOutlinerFilterBarOptions& FilterBarOptions);
+	
 public:
 	/** Miscellaneous helper functions */
 
@@ -800,6 +835,7 @@ private:
 
 	/** Updates the expansion state of parent items after a repopulate, according to the previous state */
 	void SetParentsExpansionState(const FParentsExpansionState& ExpansionStateInfo) const;
+
 private:
 
 	/** True if the outliner needs to be repopulated at the next appropriate opportunity, usually because our
@@ -824,7 +860,7 @@ private:
 	bool bIsReentrant;
 
 	/* Widget containing the filtering text box */
-	TSharedPtr< SSearchBox > FilterTextBoxWidget;
+	TSharedPtr< SFilterSearchBox > FilterTextBoxWidget;
 
 	/** The header row of the scene outliner */
 	TSharedPtr< SHeaderRow > HeaderRowWidget;
@@ -838,6 +874,12 @@ private:
 	/** The TextFilter attached to the SearchBox widget of the Scene Outliner */
 	TSharedPtr< SceneOutliner::TreeItemTextFilter > SearchBoxFilter;
 
+	/** The FilterBar attached to this outliner to filter down assets further */
+	TSharedPtr< SFilterBar< SceneOutliner::FilterBarType > > FilterBar;
+
+	/** The FilterCollection belonging to the Filter Bar */
+	TSharedPtr< TFilterCollection< SceneOutliner::FilterBarType > > FilterCollection;
+
 	/** True if the search box will take keyboard focus next frame */
 	bool bPendingFocusNextFrame;
 
diff --git a/Engine/Source/Editor/SceneOutliner/Public/SceneOutlinerPublicTypes.h b/Engine/Source/Editor/SceneOutliner/Public/SceneOutlinerPublicTypes.h
--- a/Engine/Source/Editor/SceneOutliner/Public/SceneOutlinerPublicTypes.h
+++ b/Engine/Source/Editor/SceneOutliner/Public/SceneOutlinerPublicTypes.h
@@ -13,12 +13,20 @@
 
 class FExtender;
 struct FToolMenuContext;
+class FCustomClassFilterData;
+template<typename FilterType> class FFilterBase;
 
 DECLARE_DELEGATE_TwoParams(FSceneOutlinerModifyContextMenu, FName& /* MenuName */, FToolMenuContext& /* MenuContext */);
 
 /** A delegate used as a factory to defer mode creation in the outliner */
 DECLARE_DELEGATE_RetVal_OneParam(ISceneOutlinerMode*, FCreateSceneOutlinerMode, SSceneOutliner*);
 
+namespace SceneOutliner
+{
+	/** The type of item that the Outliner's Filter Bar operates on */
+	typedef const ISceneOutlinerTreeItem& FilterBarType;
+}
+
 /** Container for built in column types. Function-static so they are available without linking */
 struct FSceneOutlinerBuiltInColumnTypes
 {
@@ -117,6 +125,29 @@ public:
 	void UseDefaultColumns();
 };
 
+
+/* Settings for the Filter Bar attached to the Scene Outliner. Can be specified through FSceneOutlinerInitializationOptions */
+struct FSceneOutlinerFilterBarOptions
+{
+	/** If true, the Scene Outliner has a filter bar attached to it */
+	bool bHasFilterBar = false;
+
+	/** These are the custom filters that the Scene Outliner will have. All active filters will be AND'd together to test
+	 *  against.
+	 *  @see FGenericFilter on how to create generic filters
+	 */
+	TArray<TSharedRef<FFilterBase<SceneOutliner::FilterBarType>>> CustomFilters;
+
+	/** These are the asset type filters that the Scene Outliner will have. All active filters will be OR'd together to
+	 *  test against.
+	 *  Can be created using IAssetTypeActions or UClass (@see constructor)
+	 */
+	TArray<TSharedRef<FCustomClassFilterData>> CustomClassFilters;
+
+	/** If true, share this Outliner filter bar's custom text filters with the level editor outliners */
+	bool bUseSharedSettings = false;
+};
+
 /**
 	* Settings for the Scene Outliner set by the programmer before spawning an instance of the widget.  This
 	* is used to modify the outliner's behavior in various ways, such as filtering in or out specific classes
@@ -144,6 +175,9 @@ struct FSceneOutlinerInitializationOptions : FSharedSceneOutlinerData
 	/** Identifier for this outliner; NAME_None if this view is anonymous (Needs to be specified to save visibility of columns in EditorConfig)*/
 	FName OutlinerIdentifier;
 
+	/** Init options related to the filter bar */
+	FSceneOutlinerFilterBarOptions FilterBarOptions;
+
 public:
 
 	/** Constructor */
diff --git a/Engine/Source/Editor/SceneOutliner/SceneOutliner.Build.cs b/Engine/Source/Editor/SceneOutliner/SceneOutliner.Build.cs
--- a/Engine/Source/Editor/SceneOutliner/SceneOutliner.Build.cs
+++ b/Engine/Source/Editor/SceneOutliner/SceneOutliner.Build.cs
@@ -8,17 +8,19 @@ public class SceneOutliner : ModuleRules
 	{
 		PrivateIncludePathModuleNames.AddRange(
 			new string[] {
+				"PlacementMode"
 			}
 		);
-
+		
 		PrivateDependencyModuleNames.AddRange(
-			new string[] {
-				"Core", 
+			new string[]
+			{
+				"Core",
 				"CoreUObject",
 				"Engine",
-                "ApplicationCore",
-                "InputCore",
-				"Slate", 
+				"ApplicationCore",
+				"InputCore",
+				"Slate",
 				"SlateCore",
 				"EditorFramework",
 				"UnrealEd",
@@ -26,7 +28,10 @@ public class SceneOutliner : ModuleRules
 				"SourceControl",
 				"EditorConfig",
 				"SourceControlWindows",
+				"EditorWidgets",
+				"ToolWidgets",
 			}
 		);
+		
 	}
 }
