Protected double-click behavior in the outliner against deleted actors.

diff --git a/Engine/Source/Editor/SceneOutliner/Private/ActorBrowsingMode.cpp b/Engine/Source/Editor/SceneOutliner/Private/ActorBrowsingMode.cpp
--- a/Engine/Source/Editor/SceneOutliner/Private/ActorBrowsingMode.cpp
+++ b/Engine/Source/Editor/SceneOutliner/Private/ActorBrowsingMode.cpp
@@ -882,35 +882,35 @@ void FActorBrowsingMode::OnItemDoubleClick(FSceneOutlinerTreeItemPtr Item)
 {
 	if (const FActorTreeItem* ActorItem = Item->CastTo<FActorTreeItem>())
 	{
-		AActor* Actor = ActorItem->Actor.Get();
-		check(Actor);
-
-		ILevelInstanceInterface* LevelInstance = Cast<ILevelInstanceInterface>(Actor);
-		if (LevelInstance && FSlateApplication::Get().GetModifierKeys().IsAltDown())
+		if (AActor* Actor = ActorItem->Actor.Get())
 		{
-			if (LevelInstance->CanEnterEdit())
+			ILevelInstanceInterface* LevelInstance = Cast<ILevelInstanceInterface>(Actor);
+			if (LevelInstance && FSlateApplication::Get().GetModifierKeys().IsAltDown())
 			{
-				LevelInstance->EnterEdit();
+				if (LevelInstance->CanEnterEdit())
+				{
+					LevelInstance->EnterEdit();
+				}
+				else if (LevelInstance->CanExitEdit())
+				{
+					LevelInstance->ExitEdit();
+				}
 			}
-			else if (LevelInstance->CanExitEdit())
+			else if (Item->CanInteract())
 			{
-				LevelInstance->ExitEdit();
+				FSceneOutlinerItemSelection Selection(SceneOutliner->GetSelection());
+				if (Selection.Has<FActorTreeItem>())
+				{
+					const bool bActiveViewportOnly = false;
+					GEditor->MoveViewportCamerasToActor(Selection.GetData<AActor*>(SceneOutliner::FActorSelector()), bActiveViewportOnly);
+				}
 			}
-		}
-		else if (Item->CanInteract())
-		{
-			FSceneOutlinerItemSelection Selection(SceneOutliner->GetSelection());
-			if (Selection.Has<FActorTreeItem>())
+			else
 			{
 				const bool bActiveViewportOnly = false;
-				GEditor->MoveViewportCamerasToActor(Selection.GetData<AActor*>(SceneOutliner::FActorSelector()), bActiveViewportOnly);
+				GEditor->MoveViewportCamerasToActor(*Actor, bActiveViewportOnly);
 			}
 		}
-		else
-		{
-			const bool bActiveViewportOnly = false;
-			GEditor->MoveViewportCamerasToActor(*Actor, bActiveViewportOnly);
-		}
 	}
 	else if (const FActorDescTreeItem* ActorDescItem = Item->CastTo<FActorDescTreeItem>())
 	{
