Outliner: Add the ability to filter unloaded actors, and move some actor specific code to ActorBrowsingMode

diff --git a/Engine/Source/Editor/SceneOutliner/Private/ActorBrowsingMode.cpp b/Engine/Source/Editor/SceneOutliner/Private/ActorBrowsingMode.cpp
--- a/Engine/Source/Editor/SceneOutliner/Private/ActorBrowsingMode.cpp
+++ b/Engine/Source/Editor/SceneOutliner/Private/ActorBrowsingMode.cpp
@@ -2211,4 +2211,36 @@ void FActorBrowsingMode::CreateOutlinerFilterBarFilters(FSceneOutlinerFilterBarO
 	
 }
 
+/** Function called by the Outliner Filter Bar to compare an item with Asset Type Filters*/
+bool FActorBrowsingMode::ConvertItemToAssetData(SceneOutliner::FilterBarType InItem, FAssetData &OutAssetData) const
+{
+	// Type filtering only supported for Actors (and unloaded actors) currently
+	if (const FActorTreeItem* ActorItem = InItem.CastTo<FActorTreeItem>())
+	{
+		AActor* Actor = ActorItem->Actor.Get();
+
+		OutAssetData = FAssetData(Actor);
+
+		return true;
+	}
+	// TODO: Maybe can be improved to not use TryConvertShortClassNameToPathName to optimize?
+	else if (const FActorDescTreeItem* ActorDescItem = InItem.CastTo<FActorDescTreeItem>())
+	{
+		if (const FWorldPartitionActorDesc* ActorDesc = ActorDescItem->ActorDescHandle.Get())
+		{
+			// For Unloaded Actors, grab the native class and try to convert it to a path name
+			OutAssetData = FAssetData();
+			FTopLevelAssetPath ClassPath = FAssetData::TryConvertShortClassNameToPathName(ActorDesc->GetNativeClass());
+			if(ClassPath.IsValid())
+			{
+				OutAssetData.AssetClassPath = ClassPath;
+			}
+		}
+
+		return true;
+	}
+
+	return FActorModeInteractive::ConvertItemToAssetData(InItem, OutAssetData);
+}
+
 #undef LOCTEXT_NAMESPACE
diff --git a/Engine/Source/Editor/SceneOutliner/Private/SSceneOutliner.cpp b/Engine/Source/Editor/SceneOutliner/Private/SSceneOutliner.cpp
--- a/Engine/Source/Editor/SceneOutliner/Private/SSceneOutliner.cpp
+++ b/Engine/Source/Editor/SceneOutliner/Private/SSceneOutliner.cpp
@@ -39,7 +39,6 @@
 #include "SceneOutlinerConfig.h"
 #include "Algo/ForEach.h"
 #include "SceneOutlinerFilterBar.h"
-#include "ActorTreeItem.h"
 #include "Framework/Notifications/NotificationManager.h"
 #include "Widgets/Notifications/SNotificationList.h"
 
@@ -2429,17 +2428,7 @@ void SSceneOutliner::OnFilterBarFilterChanged()
 
 bool SSceneOutliner::ConvertItemToAssetData(SceneOutliner::FilterBarType InItem, FAssetData &OutAssetData) const
 {
-	// Type filtering only supported for Actors currently
-	if (const FActorTreeItem* ActorItem = InItem.CastTo<FActorTreeItem>())
-	{
-		AActor* Actor = ActorItem->Actor.Get();
-
-		OutAssetData = FAssetData(Actor);
-
-		return true;
-	}
-
-	return false;
+	return Mode->ConvertItemToAssetData(InItem, OutAssetData);
 }
 
 void SSceneOutliner::CreateFilterBar(const FSceneOutlinerFilterBarOptions& FilterBarOptions)
diff --git a/Engine/Source/Editor/SceneOutliner/Public/ActorBrowsingMode.h b/Engine/Source/Editor/SceneOutliner/Public/ActorBrowsingMode.h
--- a/Engine/Source/Editor/SceneOutliner/Public/ActorBrowsingMode.h
+++ b/Engine/Source/Editor/SceneOutliner/Public/ActorBrowsingMode.h
@@ -101,6 +101,9 @@ public:
 
 	/* Helper function to create the basic set of filters that the Actor Browser will have */
 	static void CreateOutlinerFilterBarFilters(FSceneOutlinerFilterBarOptions& OutFilterBarOptions);
+
+	/** Function called by the Outliner Filter Bar to compare an item with Asset Type Filters*/
+	virtual bool ConvertItemToAssetData(SceneOutliner::FilterBarType InItem, FAssetData &OutAssetData) const override;
 private:
 	/** Build and up the context menu */
 	TSharedPtr<SWidget> BuildContextMenu();
diff --git a/Engine/Source/Editor/SceneOutliner/Public/ISceneOutlinerMode.h b/Engine/Source/Editor/SceneOutliner/Public/ISceneOutlinerMode.h
--- a/Engine/Source/Editor/SceneOutliner/Public/ISceneOutlinerMode.h
+++ b/Engine/Source/Editor/SceneOutliner/Public/ISceneOutlinerMode.h
@@ -153,6 +153,9 @@ public:
 	virtual void PinSelectedItems() {}
 	/** Unpins all selected items */
 	virtual void UnpinSelectedItems() {}
+
+	/** Function called by the Outliner Filter Bar to compare an item with Asset Type Filters*/
+	virtual bool ConvertItemToAssetData(SceneOutliner::FilterBarType InItem, FAssetData &OutAssetData) const { return false; };
 public:
 	/* Getters */
 
