SceneOutliner: Fix SCC Context menu regression

diff --git a/Engine/Source/Editor/SceneOutliner/Private/SceneOutlinerSCCHandler.cpp b/Engine/Source/Editor/SceneOutliner/Private/SceneOutlinerSCCHandler.cpp
--- a/Engine/Source/Editor/SceneOutliner/Private/SceneOutlinerSCCHandler.cpp
+++ b/Engine/Source/Editor/SceneOutliner/Private/SceneOutlinerSCCHandler.cpp
@@ -10,7 +10,7 @@
 
 #define LOCTEXT_NAMESPACE "FSceneOutlinerSCCHandler"
 
-TSharedPtr<FSceneOutlinerTreeItemSCC> FSceneOutlinerSCCHandler::GetItemSourceControl(const FSceneOutlinerTreeItemPtr& InItem)
+TSharedPtr<FSceneOutlinerTreeItemSCC> FSceneOutlinerSCCHandler::GetItemSourceControl(const FSceneOutlinerTreeItemPtr& InItem) const
 {
 	TSharedPtr<FSceneOutlinerTreeItemSCC>* Result = ItemSourceControls.Find(InItem);
 	if (!Result)
@@ -88,14 +88,12 @@ void FSceneOutlinerSCCHandler::CacheCanExecuteVars()
 	{
 		for (FSceneOutlinerTreeItemPtr SelectedItem : SelectedItems)
 		{
-			TSharedPtr<FSceneOutlinerTreeItemSCC>* SourceControlLookup = ItemSourceControls.Find(SelectedItem);
-			if (SourceControlLookup == nullptr)
+			TSharedPtr<FSceneOutlinerTreeItemSCC> SourceControl = GetItemSourceControl(SelectedItem);
+			if (SourceControl == nullptr)
 			{
 				continue;
 			}
 
-			TSharedPtr<FSceneOutlinerTreeItemSCC> SourceControl = *SourceControlLookup;
-
 			if (!SourceControl->IsExternalPackage())
 			{
 				// this isn't an external package so we can't do anything with source control
@@ -218,13 +216,12 @@ void FSceneOutlinerSCCHandler::GetSelectedPackageNames(TArray<FString>& OutPacka
 {
 	for (FSceneOutlinerTreeItemPtr SelectedItem : SelectedItems)
 	{
-		const TSharedPtr<FSceneOutlinerTreeItemSCC>* SourceControlLookup = ItemSourceControls.Find(SelectedItem);
-		if (SourceControlLookup == nullptr)
+		const TSharedPtr<FSceneOutlinerTreeItemSCC> SourceControl = GetItemSourceControl(SelectedItem);
+		if (SourceControl == nullptr)
 		{
 			continue;
 		}
 
-		TSharedPtr<FSceneOutlinerTreeItemSCC> SourceControl = *SourceControlLookup;
 		FString PackageName = SourceControl->GetPackageName();
 		if (!PackageName.IsEmpty()) {
 			OutPackageNames.Add(PackageName);
@@ -236,13 +233,12 @@ void FSceneOutlinerSCCHandler::GetSelectedPackages(TArray<UPackage*>& OutPackage
 {
 	for (FSceneOutlinerTreeItemPtr SelectedItem : SelectedItems)
 	{
-		const TSharedPtr<FSceneOutlinerTreeItemSCC>* SourceControlLookup = ItemSourceControls.Find(SelectedItem);
-		if (SourceControlLookup == nullptr)
+		const TSharedPtr<FSceneOutlinerTreeItemSCC> SourceControl = GetItemSourceControl(SelectedItem);
+		if (SourceControl == nullptr)
 		{
 			continue;
 		}
 
-		TSharedPtr<FSceneOutlinerTreeItemSCC> SourceControl = *SourceControlLookup;
 		UPackage* Package = SourceControl->GetPackage();
 		if (Package != nullptr) {
 			OutPackages.Add(Package);
diff --git a/Engine/Source/Editor/SceneOutliner/Public/SceneOutlinerSCCHandler.h b/Engine/Source/Editor/SceneOutliner/Public/SceneOutlinerSCCHandler.h
--- a/Engine/Source/Editor/SceneOutliner/Public/SceneOutlinerSCCHandler.h
+++ b/Engine/Source/Editor/SceneOutliner/Public/SceneOutlinerSCCHandler.h
@@ -13,7 +13,7 @@ public:
 
 	bool AddSourceControlMenuOptions(UToolMenu* Menu, TArray<FSceneOutlinerTreeItemPtr> InSelectedItems);
 
-	TSharedPtr<FSceneOutlinerTreeItemSCC> GetItemSourceControl(const FSceneOutlinerTreeItemPtr& InItem);
+	TSharedPtr<FSceneOutlinerTreeItemSCC> GetItemSourceControl(const FSceneOutlinerTreeItemPtr& InItem) const;
 
 private:
 
@@ -43,5 +43,5 @@ private:
 	bool bCanExecuteSCCCheckIn = false;
 	bool bCanExecuteSCCHistory = false;
 
-	TMap<FSceneOutlinerTreeItemPtr, TSharedPtr<FSceneOutlinerTreeItemSCC>> ItemSourceControls;
+	mutable TMap<FSceneOutlinerTreeItemPtr, TSharedPtr<FSceneOutlinerTreeItemSCC>> ItemSourceControls;
 };
