ActorBrowsingMode: Fix an issue causing the "Hide Empty Folders" filter to not be loaded from the config correctly

diff --git a/Engine/Source/Editor/SceneOutliner/Private/ActorBrowsingMode.cpp b/Engine/Source/Editor/SceneOutliner/Private/ActorBrowsingMode.cpp
--- a/Engine/Source/Editor/SceneOutliner/Private/ActorBrowsingMode.cpp
+++ b/Engine/Source/Editor/SceneOutliner/Private/ActorBrowsingMode.cpp
@@ -177,6 +177,7 @@ FActorBrowsingMode::FActorBrowsingMode(SSceneOutliner* InSceneOutliner, TWeakObj
 		});
 	FilterInfoMap.Add(TEXT("HideUnloadedActorsFilter"), HideUnloadedActorsInfo);
 
+	bHideEmptyFolders = LocalSettings.bHideEmptyFolders;
 	FSceneOutlinerFilterInfo HideEmptyFoldersInfo(LOCTEXT("ToggleHideEmptyFolders", "Hide Empty Folders"), LOCTEXT("ToggleHideEmptyFoldersToolTip", "When enabled, hides all empty folders."), LocalSettings.bHideEmptyFolders, FCreateSceneOutlinerFilter::CreateStatic(&FActorBrowsingMode::CreateHideEmptyFoldersFilter));
 	HideEmptyFoldersInfo.OnToggle().AddLambda([this](bool bIsActive)
 		{
