ActorBrowsingMode: Fix an issue causing some Outliner filters to not be read from the config correctly

diff --git a/Engine/Source/Editor/SceneOutliner/Private/ActorBrowsingMode.cpp b/Engine/Source/Editor/SceneOutliner/Private/ActorBrowsingMode.cpp
--- a/Engine/Source/Editor/SceneOutliner/Private/ActorBrowsingMode.cpp
+++ b/Engine/Source/Editor/SceneOutliner/Private/ActorBrowsingMode.cpp
@@ -143,6 +143,7 @@ FActorBrowsingMode::FActorBrowsingMode(SSceneOutliner* InSceneOutliner, TWeakObj
 
 	FilterInfoMap.Add(TEXT("HideComponentsFilter"), HideComponentsInfo);
 
+	bHideLevelInstanceHierarchy = LocalSettings.bHideLevelInstanceHierarchy;
 	FSceneOutlinerFilterInfo HideLevelInstancesInfo(LOCTEXT("ToggleHideLevelInstanceContent", "Hide Level Instance Content"), LOCTEXT("ToggleHideLevelInstancesToolTip", "When enabled, hides all level instance content."), LocalSettings.bHideLevelInstanceHierarchy, FCreateSceneOutlinerFilter::CreateStatic(&FActorBrowsingMode::CreateHideLevelInstancesFilter));
 	HideLevelInstancesInfo.OnToggle().AddLambda([this](bool bIsActive)
 		{
@@ -160,6 +161,7 @@ FActorBrowsingMode::FActorBrowsingMode(SSceneOutliner* InSceneOutliner, TWeakObj
 		});
 	FilterInfoMap.Add(TEXT("HideLevelInstancesFilter"), HideLevelInstancesInfo);
 
+	bHideUnloadedActors = LocalSettings.bHideUnloadedActors;
 	FSceneOutlinerFilterInfo HideUnloadedActorsInfo(LOCTEXT("ToggleHideUnloadedActors", "Hide Unloaded Actors"), LOCTEXT("ToggleHideUnloadedActorsToolTip", "When enabled, hides all unloaded world partition actors."), LocalSettings.bHideUnloadedActors, FCreateSceneOutlinerFilter::CreateStatic(&FActorBrowsingMode::CreateHideUnloadedActorsFilter));
 	HideUnloadedActorsInfo.OnToggle().AddLambda([this] (bool bIsActive)
 		{
