Add a cvar to override the preferred representing world chosen by the scene outliner when the  'Auto' option is selected based off of the world's net mode.

diff --git a/Engine/Source/Editor/SceneOutliner/Private/ActorMode.cpp b/Engine/Source/Editor/SceneOutliner/Private/ActorMode.cpp
--- a/Engine/Source/Editor/SceneOutliner/Private/ActorMode.cpp
+++ b/Engine/Source/Editor/SceneOutliner/Private/ActorMode.cpp
@@ -22,6 +22,12 @@
 #include "LevelEditor.h"
 #include "Modules/ModuleManager.h"
 
+static int32 GSceneOutlinerAutoRepresentingWorldNetMode = NM_Client;
+static FAutoConsoleVariableRef CVarAutoRepresentingWorldNetMode(
+	TEXT("SceneOutliner.AutoRepresentingWorldNetMode"),
+	GSceneOutlinerAutoRepresentingWorldNetMode,
+	TEXT("The preferred NetMode of the world shown in the scene outliner when the 'Auto' option is chosen: 0=Standalone, 1=DedicatedServer, 2=ListenServer, 3=Client"));
+
 #define LOCTEXT_NAMESPACE "SceneOutliner_ActorMode"
 
 using FActorFilter = TSceneOutlinerPredicateFilter<FActorTreeItem>;
@@ -177,8 +183,8 @@ void FActorMode::ChooseRepresentingWorld()
 	// If the user did not manually select a world, try to pick the most suitable world context
 	if (!RepresentingWorld.IsValid())
 	{
-		// ideally we want a PIE world that is standalone or the first client
-		int32 LowestClientInstanceSeen = MAX_int32;
+		// Ideally we want a PIE world that is standalone or the first client, unless the preferred NetMode is overridden by CVar
+		int32 LowestPIEInstanceSeen = MAX_int32;
 		for (const FWorldContext& Context : GEngine->GetWorldContexts())
 		{
 			UWorld* World = Context.World();
@@ -189,10 +195,10 @@ void FActorMode::ChooseRepresentingWorld()
 					RepresentingWorld = World;
 					break;
 				}
-				else if ((World->GetNetMode() == NM_Client) && (Context.PIEInstance < LowestClientInstanceSeen))
+				else if ((World->GetNetMode() == ENetMode(GSceneOutlinerAutoRepresentingWorldNetMode)) && (Context.PIEInstance < LowestPIEInstanceSeen))
 				{
 					RepresentingWorld = World;
-					LowestClientInstanceSeen = Context.PIEInstance;
+					LowestPIEInstanceSeen = Context.PIEInstance;
 				}
 			}
 		}
