Fixed merge problem

diff --git a/Engine/Source/Editor/SceneOutliner/Private/ActorHierarchy.cpp b/Engine/Source/Editor/SceneOutliner/Private/ActorHierarchy.cpp
--- a/Engine/Source/Editor/SceneOutliner/Private/ActorHierarchy.cpp
+++ b/Engine/Source/Editor/SceneOutliner/Private/ActorHierarchy.cpp
@@ -29,6 +29,9 @@
 #include "ActorFolder.h"
 #include "ActorMode.h"
 
+// [jonathan.bard] Disable optimization, remove : 
+PRAGMA_DISABLE_OPTIMIZATION
+
 TUniquePtr<FActorHierarchy> FActorHierarchy::Create(ISceneOutlinerMode* Mode, const TWeakObjectPtr<UWorld>& World)
 {
 	FActorHierarchy* Hierarchy = new FActorHierarchy(Mode, World);
@@ -832,4 +835,7 @@ void FActorHierarchy::OnLevelActorFolderChanged(const AActor* InActor, FName Old
 		EventData.ItemIDs.Add(FSceneOutlinerTreeItemID(InActor));
 		HierarchyChangedEvent.Broadcast(EventData);
 	}
-}
\ No newline at end of file
+}
+
+// [jonathan.bard] Disable optimization, remove : 
+PRAGMA_ENABLE_OPTIMIZATION
\ No newline at end of file
