[Backout] - CL22072456
Original CL Desc
-----------------------------------------------------------------
Fixed merge problem

diff --git a/Engine/Source/Editor/SceneOutliner/Private/ActorHierarchy.cpp b/Engine/Source/Editor/SceneOutliner/Private/ActorHierarchy.cpp
--- a/Engine/Source/Editor/SceneOutliner/Private/ActorHierarchy.cpp
+++ b/Engine/Source/Editor/SceneOutliner/Private/ActorHierarchy.cpp
@@ -29,9 +29,6 @@
 #include "ActorFolder.h"
 #include "ActorMode.h"
 
-// [jonathan.bard] Disable optimization, remove : 
-PRAGMA_DISABLE_OPTIMIZATION
-
 TUniquePtr<FActorHierarchy> FActorHierarchy::Create(ISceneOutlinerMode* Mode, const TWeakObjectPtr<UWorld>& World)
 {
 	FActorHierarchy* Hierarchy = new FActorHierarchy(Mode, World);
@@ -835,7 +832,4 @@ void FActorHierarchy::OnLevelActorFolderChanged(const AActor* InActor, FName Old
 		EventData.ItemIDs.Add(FSceneOutlinerTreeItemID(InActor));
 		HierarchyChangedEvent.Broadcast(EventData);
 	}
-}
-
-// [jonathan.bard] Disable optimization, remove : 
-PRAGMA_ENABLE_OPTIMIZATION
\ No newline at end of file
+}
\ No newline at end of file
