World Partition: detect invalid actor files and show a warning on top of the scene outliner with a button to repair.

#rnx

diff --git a/Engine/Source/Editor/SceneOutliner/Private/ActorBrowsingMode.cpp b/Engine/Source/Editor/SceneOutliner/Private/ActorBrowsingMode.cpp
--- a/Engine/Source/Editor/SceneOutliner/Private/ActorBrowsingMode.cpp
+++ b/Engine/Source/Editor/SceneOutliner/Private/ActorBrowsingMode.cpp
@@ -13,6 +13,7 @@
 #include "Engine/Selection.h"
 #include "Editor/UnrealEdEngine.h"
 #include "HAL/PlatformApplicationMisc.h"
+#include "HAL/FileManager.h"
 #include "SceneOutlinerDragDrop.h"
 #include "EditorActorFolders.h"
 #include "EditorFolderUtils.h"
@@ -35,6 +36,9 @@
 #include "LevelInstance/LevelInstanceSubsystem.h"
 #include "LevelInstance/LevelInstanceEditorInstanceActor.h"
 #include "WorldPartition/WorldPartition.h"
+#include "ISourceControlModule.h"
+#include "ISourceControlProvider.h"
+#include "SourceControlOperations.h"
 #include "EditorLevelUtils.h"
 #include "EditorViewportCommands.h"
 #include "SceneOutlinerActorSCCColumn.h"
@@ -42,6 +46,7 @@
 #include "Engine/StaticMeshActor.h"
 #include "Filters/CustomClassFilterData.h"
 #include "IPlacementModeModule.h"
+#include "Misc/ScopedSlowTask.h"
 
 DEFINE_LOG_CATEGORY_STATIC(LogActorBrowser, Log, All);
 
@@ -1079,6 +1084,111 @@ bool FActorBrowsingMode::CanPaste() const
 	return CanPasteFoldersOnlyFromClipboard();
 }
 
+bool FActorBrowsingMode::HasErrors() const
+{
+	if (!bRepresentingWorldGameWorld && bRepresentingWorldPartitionedWorld)
+	{
+		if (UWorldPartition* const WorldPartition = RepresentingWorld->GetWorldPartition())
+		{
+			bool bHasErrors = 0;
+
+			WorldPartition->ForEachActorDescContainer([&bHasErrors](UActorDescContainer* ActorDescContainer)
+			{
+				if (ActorDescContainer->HasInvalidActors())
+				{
+					bHasErrors = true;
+				}
+			});
+
+			return bHasErrors;
+		}
+	}
+
+	return false;
+}
+
+FText FActorBrowsingMode::GetErrorsText() const
+{
+	return LOCTEXT("WorldHasInvalidActorFiles", "The world contains invalid actor files. Click the Repair button to repair them.");
+}
+
+void FActorBrowsingMode::RepairErrors() const
+{
+	if (!bRepresentingWorldGameWorld && bRepresentingWorldPartitionedWorld)
+	{
+		if (UWorldPartition* const WorldPartition = RepresentingWorld->GetWorldPartition())
+		{
+			ISourceControlModule& SourceControlModule = ISourceControlModule::Get();
+			ISourceControlProvider& SourceControlProvider = SourceControlModule.GetProvider();
+
+			TArray<FAssetData> InvalidActors;
+			WorldPartition->ForEachActorDescContainer([&InvalidActors ](UActorDescContainer* ActorDescContainer)
+			{
+				InvalidActors += ActorDescContainer->GetInvalidActors();
+				ActorDescContainer->ClearInvalidActors();
+			});
+
+			TArray<FString> ActorFilesToDelete;
+			TArray<FString> ActorFilesToRevert;
+			{
+				FScopedSlowTask SlowTask(InvalidActors.Num(), LOCTEXT("UpdatingSourceControlStatus", "Updating source control status..."));
+				SlowTask.MakeDialogDelayed(1.0f);
+
+				for (const FAssetData& InvalidActor : InvalidActors)
+				{
+					FPackagePath PackagePath;
+					if (FPackagePath::TryFromPackageName(InvalidActor.PackageName, PackagePath))
+					{
+						const FString ActorFile = PackagePath.GetLocalFullPath();
+						FSourceControlStatePtr SCState = SourceControlProvider.GetState(ActorFile, EStateCacheUsage::ForceUpdate);
+
+						if (SCState.IsValid() && SCState->IsSourceControlled())
+						{
+							if (SCState->IsAdded())
+							{
+								ActorFilesToRevert.Add(ActorFile);
+							}
+							else
+							{
+								if (SCState->IsCheckedOut())
+								{
+									ActorFilesToRevert.Add(ActorFile);
+								}
+
+								ActorFilesToDelete.Add(ActorFile);
+							}
+						}
+						else
+						{
+							IFileManager::Get().Delete(*ActorFile, false, true);
+						}
+					}
+
+					SlowTask.EnterProgressFrame(1);
+				}
+			}
+
+			if (ActorFilesToRevert.Num() || ActorFilesToDelete.Num())
+			{
+				FScopedSlowTask SlowTask(ActorFilesToRevert.Num() + ActorFilesToDelete.Num(), LOCTEXT("DeletingFiles", "Deleting files..."));
+				SlowTask.MakeDialogDelayed(1.0f);
+
+				if (ActorFilesToRevert.Num())
+				{
+					SourceControlProvider.Execute(ISourceControlOperation::Create<FRevert>(), ActorFilesToRevert);
+					SlowTask.EnterProgressFrame(ActorFilesToRevert.Num());
+				}
+	
+				if (ActorFilesToDelete.Num())
+				{
+					SourceControlProvider.Execute(ISourceControlOperation::Create<FDelete>(), ActorFilesToDelete);
+					SlowTask.EnterProgressFrame(ActorFilesToDelete.Num());
+				}
+			}
+		}
+	}
+}
+
 bool FActorBrowsingMode::CanPasteFoldersOnlyFromClipboard() const
 {
 	// Intentionally not checking if the level is locked/hidden here, as it's better feedback for the user if they attempt to paste
diff --git a/Engine/Source/Editor/SceneOutliner/Private/SSceneOutliner.cpp b/Engine/Source/Editor/SceneOutliner/Private/SSceneOutliner.cpp
--- a/Engine/Source/Editor/SceneOutliner/Private/SSceneOutliner.cpp
+++ b/Engine/Source/Editor/SceneOutliner/Private/SSceneOutliner.cpp
@@ -41,6 +41,8 @@
 #include "SceneOutlinerFilterBar.h"
 #include "Framework/Notifications/NotificationManager.h"
 #include "Widgets/Notifications/SNotificationList.h"
+#include "Widgets/Input/SMultiLineEditableTextBox.h"
+#include "DetailLayoutBuilder.h"
 
 DEFINE_LOG_CATEGORY_STATIC(LogSceneOutliner, Log, All);
 
@@ -141,6 +143,32 @@ void SSceneOutliner::Construct(const FArguments& InArgs, const FSceneOutlinerIni
 		VerticalBox
 	];
 
+	VerticalBox->AddSlot()
+	.AutoHeight()
+	[
+		SNew(SVerticalBox)
+		+ SVerticalBox::Slot()
+		.AutoHeight()
+		[
+			SNew(SMultiLineEditableTextBox)
+			.IsReadOnly(true)
+			.Visibility_Lambda([this]() { return Mode->HasErrors() ? EVisibility::Visible : EVisibility::Collapsed; })
+			.Font(IDetailLayoutBuilder::GetDetailFontBold())
+			.BackgroundColor(FAppStyle::GetColor("ErrorReporting.WarningBackgroundColor"))
+			.Text(Mode->GetErrorsText())
+			.AutoWrapText(true)			
+		]
+		+ SVerticalBox::Slot()
+		.AutoHeight()
+		[
+			SNew(SButton)
+			.OnClicked_Lambda([this] { Mode->RepairErrors(); return FReply::Handled(); })
+			.Visibility_Lambda([this]() { return Mode->HasErrors() ? EVisibility::Visible : EVisibility::Collapsed; })
+			.HAlign(HAlign_Center)
+			.Text(LOCTEXT("SceneOutlinerRepairErrors", "Repair Errors"))
+		]
+	];
+
 	auto Toolbar = SNew(SHorizontalBox);
 
 	Toolbar->AddSlot()
diff --git a/Engine/Source/Editor/SceneOutliner/Public/ActorBrowsingMode.h b/Engine/Source/Editor/SceneOutliner/Public/ActorBrowsingMode.h
--- a/Engine/Source/Editor/SceneOutliner/Public/ActorBrowsingMode.h
+++ b/Engine/Source/Editor/SceneOutliner/Public/ActorBrowsingMode.h
@@ -45,6 +45,9 @@ public:
 	virtual bool CanCopy() const override;
 	virtual bool CanPaste() const override;
 	virtual bool CanSupportDragAndDrop() const { return true; }
+	virtual bool HasErrors() const;
+	virtual FText GetErrorsText() const;
+	virtual void RepairErrors() const;	
 	virtual FFolder CreateNewFolder() override;
 	virtual FFolder GetFolder(const FFolder& ParentPath, const FName& LeafName) override;
 	virtual bool CreateFolder(const FFolder& NewFolder) override;
diff --git a/Engine/Source/Editor/SceneOutliner/Public/ISceneOutlinerMode.h b/Engine/Source/Editor/SceneOutliner/Public/ISceneOutlinerMode.h
--- a/Engine/Source/Editor/SceneOutliner/Public/ISceneOutlinerMode.h
+++ b/Engine/Source/Editor/SceneOutliner/Public/ISceneOutlinerMode.h
@@ -131,6 +131,16 @@ public:
 
 	/** Does the mode support drag and drop */
 	virtual bool CanSupportDragAndDrop() const { return false; }
+
+	/** Does the mode reports errors */
+	virtual bool HasErrors() const { return false; }
+
+	/** Return the errors text */
+	virtual FText GetErrorsText() const { return FText(); }
+
+	/** Repair errors*/
+	virtual void RepairErrors() const {}
+
 public:
 	/* Folder management */
 		
diff --git a/Engine/Source/Runtime/Engine/Private/WorldPartition/ActorDescContainer.cpp b/Engine/Source/Runtime/Engine/Private/WorldPartition/ActorDescContainer.cpp
--- a/Engine/Source/Runtime/Engine/Private/WorldPartition/ActorDescContainer.cpp
+++ b/Engine/Source/Runtime/Engine/Private/WorldPartition/ActorDescContainer.cpp
@@ -32,7 +32,7 @@ void UActorDescContainer::Initialize(UWorld* InWorld, FName InPackageName)
 	check(!bContainerInitialized);
 	ContainerPackageName = InPackageName;
 	TArray<FAssetData> Assets;
-	
+
 	if (!ContainerPackageName.IsNone())
 	{
 		const FString ContainerExternalActorsPath = GetExternalActorPath();
@@ -57,6 +57,10 @@ void UActorDescContainer::Initialize(UWorld* InWorld, FName InPackageName)
 		{
 			AddActorDescriptor(ActorDesc.Release());
 		}
+		else
+		{
+			InvalidActors.Add(Asset);
+		}
 	}
 
 	OnActorDescContainerInitialized.Broadcast(this);
diff --git a/Engine/Source/Runtime/Engine/Private/WorldPartition/WorldPartitionActorDescUtils.cpp b/Engine/Source/Runtime/Engine/Private/WorldPartition/WorldPartitionActorDescUtils.cpp
--- a/Engine/Source/Runtime/Engine/Private/WorldPartition/WorldPartitionActorDescUtils.cpp
+++ b/Engine/Source/Runtime/Engine/Private/WorldPartition/WorldPartitionActorDescUtils.cpp
@@ -44,34 +44,42 @@ UClass* FWorldPartitionActorDescUtils::GetActorNativeClassFromAssetData(const FA
 
 TUniquePtr<FWorldPartitionActorDesc> FWorldPartitionActorDescUtils::GetActorDescriptorFromAssetData(const FAssetData& InAssetData)
 {
-	FString ActorMetaDataClass;
-	FString ActorMetaDataStr;
-	if (InAssetData.GetTagValue(NAME_ActorMetaDataClass, ActorMetaDataClass) && InAssetData.GetTagValue(NAME_ActorMetaData, ActorMetaDataStr))
+	if (IsValidActorDescriptorFromAssetData(InAssetData))
 	{
-		bool bIsValidClass = true;
-		UClass* ActorClass = GetActorNativeClassFromAssetData(InAssetData);
-
-		if (!ActorClass)
-		{
-			ActorClass = AActor::StaticClass();
-			bIsValidClass = false;
-		}
-
 		FWorldPartitionActorDescInitData ActorDescInitData;
-		ActorDescInitData.NativeClass = ActorClass;
+		ActorDescInitData.NativeClass = GetActorNativeClassFromAssetData(InAssetData);
 		ActorDescInitData.PackageName = InAssetData.PackageName;
 		ActorDescInitData.ActorPath = InAssetData.GetSoftObjectPath();
-		FBase64::Decode(ActorMetaDataStr, ActorDescInitData.SerializedData);
 
-		TUniquePtr<FWorldPartitionActorDesc> NewActorDesc(AActor::StaticCreateClassActorDesc(ActorDescInitData.NativeClass));
+		FString ActorMetaDataStr;
+		verify(InAssetData.GetTagValue(NAME_ActorMetaData, ActorMetaDataStr));
+		verify(FBase64::Decode(ActorMetaDataStr, ActorDescInitData.SerializedData));
+
+		TUniquePtr<FWorldPartitionActorDesc> NewActorDesc(AActor::StaticCreateClassActorDesc(ActorDescInitData.NativeClass ? ActorDescInitData.NativeClass : AActor::StaticClass()));
 
 		NewActorDesc->Init(ActorDescInitData);
 			
-		if (!bIsValidClass)
+		if (!ActorDescInitData.NativeClass)
 		{
 			UE_LOG(LogWorldPartition, Warning, TEXT("Invalid class for actor guid `%s` ('%s') from package '%s'"), *NewActorDesc->GetGuid().ToString(), *NewActorDesc->GetActorName().ToString(), *NewActorDesc->GetActorPackage().ToString());
 			return nullptr;
 		}
+		else if (UClass* Class = FindObject<UClass>(InAssetData.AssetClassPath); !Class)
+		{
+			// We can't detect mising BP classes for inactive plugins, etc.
+			if (InAssetData.AssetClassPath.GetPackageName().ToString().StartsWith(TEXT("/Game/")))
+			{
+				TArray<FAssetData> BlueprintClass;
+				IAssetRegistry& AssetRegistry = FModuleManager::LoadModuleChecked<FAssetRegistryModule>(TEXT("AssetRegistry")).Get();
+				AssetRegistry.GetAssetsByPackageName(*InAssetData.AssetClassPath.GetPackageName().ToString(), BlueprintClass, true);
+
+				if (!BlueprintClass.Num())
+				{
+					UE_LOG(LogWorldPartition, Warning, TEXT("Failed to find class '%s' for actor '%s"), *InAssetData.AssetClassPath.ToString(), *NewActorDesc->GetActorPath().ToString());
+					return nullptr;
+				}
+			}
+		}
 
 		return NewActorDesc;
 	}
diff --git a/Engine/Source/Runtime/Engine/Public/WorldPartition/ActorDescContainer.h b/Engine/Source/Runtime/Engine/Public/WorldPartition/ActorDescContainer.h
--- a/Engine/Source/Runtime/Engine/Public/WorldPartition/ActorDescContainer.h
+++ b/Engine/Source/Runtime/Engine/Public/WorldPartition/ActorDescContainer.h
@@ -62,6 +62,10 @@ public:
 
 	const FLinkerInstancingContext* GetInstancingContext() const;
 	const FTransform& GetInstanceTransform() const;
+
+	bool HasInvalidActors() const { return InvalidActors.Num() > 0; }
+	const TArray<FAssetData>& GetInvalidActors() const { return InvalidActors; }
+	void ClearInvalidActors() { InvalidActors.Empty(); }
 #endif
 
 	UPROPERTY(Transient)
@@ -89,6 +93,8 @@ protected:
 
 	FName ContainerPackageName;
 	FGuid ContentBundleGuid;
+
+	TArray<FAssetData> InvalidActors;
 #endif
 
 private:
