Fix potrential crash when changing worlds.

#rnx

diff --git a/Engine/Source/Editor/SceneOutliner/Private/ActorBrowsingMode.cpp b/Engine/Source/Editor/SceneOutliner/Private/ActorBrowsingMode.cpp
--- a/Engine/Source/Editor/SceneOutliner/Private/ActorBrowsingMode.cpp
+++ b/Engine/Source/Editor/SceneOutliner/Private/ActorBrowsingMode.cpp
@@ -1086,7 +1086,7 @@ bool FActorBrowsingMode::CanPaste() const
 
 bool FActorBrowsingMode::HasErrors() const
 {
-	if (!bRepresentingWorldGameWorld && bRepresentingWorldPartitionedWorld)
+	if (RepresentingWorld && !bRepresentingWorldGameWorld && bRepresentingWorldPartitionedWorld)
 	{
 		if (UWorldPartition* const WorldPartition = RepresentingWorld->GetWorldPartition())
 		{
