WorldOutliner: Fix some Pin/Visibility states on Level Instances

#rnx

diff --git a/Engine/Source/Editor/SceneOutliner/Private/ActorDescTreeItem.cpp b/Engine/Source/Editor/SceneOutliner/Private/ActorDescTreeItem.cpp
--- a/Engine/Source/Editor/SceneOutliner/Private/ActorDescTreeItem.cpp
+++ b/Engine/Source/Editor/SceneOutliner/Private/ActorDescTreeItem.cpp
@@ -317,12 +317,29 @@ bool FActorDescTreeItem::GetVisibility() const
 	return true;
 }
 
+bool FActorDescTreeItem::ShouldShowPinnedState() const
+{
+	if (ActorDescHandle.IsValid())
+	{
+		// Pinning of ActorDescs is only supported on the main world partition
+		if (UWorldPartition* WorldPartition = Cast<UWorldPartition>(ActorDescHandle->GetContainer()))
+		{
+			return WorldPartition->IsMainWorldPartition();
+		}
+	}
+
+	return false;
+}
+
 bool FActorDescTreeItem::GetPinnedState() const
 {
 	if (ActorDescHandle.IsValid())
 	{
-		UWorldPartition* WorldPartition = ActorDescHandle->GetContainer()->GetWorld()->GetWorldPartition();
-		return WorldPartition ? WorldPartition->IsActorPinned(GetGuid()) : false;
+		// Pinning of ActorDescs is only supported on the main world partition
+		if (UWorldPartition* WorldPartition = Cast<UWorldPartition>(ActorDescHandle->GetContainer()))
+		{
+			return WorldPartition->IsMainWorldPartition() && WorldPartition->IsActorPinned(GetGuid());
+		}
 	}
 	return false;
 }
diff --git a/Engine/Source/Editor/SceneOutliner/Private/ActorFolderTreeItem.cpp b/Engine/Source/Editor/SceneOutliner/Private/ActorFolderTreeItem.cpp
--- a/Engine/Source/Editor/SceneOutliner/Private/ActorFolderTreeItem.cpp
+++ b/Engine/Source/Editor/SceneOutliner/Private/ActorFolderTreeItem.cpp
@@ -357,6 +357,14 @@ bool FActorFolderTreeItem::ShouldShowPinnedState() const
 	return false;
 }
 
+bool FActorFolderTreeItem::ShouldShowVisibilityState() const
+{
+	// Visibility state can only be shown when folder is part of the persistent level (PIE or Editor) or
+	// if it is an actual sub-level that is not instanced (sublevels prior to level instances or editing level instance)
+	ULevel* Level = FFolder::GetRootObjectAssociatedLevel(GetRootObject());
+	return World.IsValid() && Level && (Level->IsPersistentLevel() || !Level->IsInstancedLevel());
+}
+
 FFolder FActorFolderTreeItem::GetFolder() const
 {
 	// Use Folder resolved by ActorFolder when valid (a null owning world would will fail resolving, see UActorFolder::GetFolder())
diff --git a/Engine/Source/Editor/SceneOutliner/Private/ActorTreeItem.cpp b/Engine/Source/Editor/SceneOutliner/Private/ActorTreeItem.cpp
--- a/Engine/Source/Editor/SceneOutliner/Private/ActorTreeItem.cpp
+++ b/Engine/Source/Editor/SceneOutliner/Private/ActorTreeItem.cpp
@@ -390,9 +390,10 @@ bool FActorTreeItem::ShouldShowPinnedState() const
 {
 	if (const AActor* ActorPtr = Actor.Get())
 	{
+		// Pinning of Actors is only supported on the main world partition
 		const ULevel* Level = ActorPtr->GetLevel();
 		const UWorld* World = Level->GetWorld();
-		return World && !World->IsGameWorld() && !!Level->GetWorldPartition();
+		return World && !World->IsGameWorld() && !!World->GetWorldPartition() && Level->IsPersistentLevel();
 	}
 
 	return false;
diff --git a/Engine/Source/Editor/SceneOutliner/Public/ActorDescTreeItem.h b/Engine/Source/Editor/SceneOutliner/Public/ActorDescTreeItem.h
--- a/Engine/Source/Editor/SceneOutliner/Public/ActorDescTreeItem.h
+++ b/Engine/Source/Editor/SceneOutliner/Public/ActorDescTreeItem.h
@@ -53,7 +53,7 @@ public:
 	virtual void GenerateContextMenu(UToolMenu* Menu, SSceneOutliner& Outliner) override;
 	virtual bool HasVisibilityInfo() const override { return true; }
 	virtual bool GetVisibility() const override;
-	virtual bool ShouldShowPinnedState() const override { return true; }
+	virtual bool ShouldShowPinnedState() const override;
 	virtual bool ShouldShowVisibilityState() const override { return false; }
 	virtual bool HasPinnedStateInfo() const override { return true; }
 	virtual bool GetPinnedState() const override;
diff --git a/Engine/Source/Editor/SceneOutliner/Public/ActorFolderTreeItem.h b/Engine/Source/Editor/SceneOutliner/Public/ActorFolderTreeItem.h
--- a/Engine/Source/Editor/SceneOutliner/Public/ActorFolderTreeItem.h
+++ b/Engine/Source/Editor/SceneOutliner/Public/ActorFolderTreeItem.h
@@ -23,6 +23,7 @@ public:
 	virtual bool CanInteract() const override;
 	virtual TSharedRef<SWidget> GenerateLabelWidget(ISceneOutliner& Outliner, const STableRow<FSceneOutlinerTreeItemPtr>& InRow) override;
 	virtual bool ShouldShowPinnedState() const override;
+	virtual bool ShouldShowVisibilityState() const override;
 	virtual FFolder GetFolder() const override;
 	/* End FFolderTreeItem Implementation */
 		
