[NO NON-PII COMMIT MESSAGE]

diff --git a/Engine/Source/Editor/AnimationBlueprintEditor/AnimationBlueprintEditor.Build.cs b/Engine/Source/Editor/AnimationBlueprintEditor/AnimationBlueprintEditor.Build.cs
--- a/Engine/Source/Editor/AnimationBlueprintEditor/AnimationBlueprintEditor.Build.cs
+++ b/Engine/Source/Editor/AnimationBlueprintEditor/AnimationBlueprintEditor.Build.cs
@@ -61,6 +61,7 @@ public class AnimationBlueprintEditor : ModuleRules
             new string[] {
                 "Kismet",
                 "Persona",
+				"EditorFramework"
             }
         );
     }
diff --git a/Engine/Source/Editor/SceneOutliner/Private/ActorBrowsingMode.cpp b/Engine/Source/Editor/SceneOutliner/Private/ActorBrowsingMode.cpp
--- a/Engine/Source/Editor/SceneOutliner/Private/ActorBrowsingMode.cpp
+++ b/Engine/Source/Editor/SceneOutliner/Private/ActorBrowsingMode.cpp
@@ -2267,12 +2267,11 @@ bool FActorBrowsingMode::CompareItemWithClassName(SceneOutliner::FilterBarType I
 		
 		return AssetClassPaths.Contains(AssetClassPath);
 	}
-	// TODO: Maybe can be improved to not use TryConvertShortClassNameToPathName to optimize?
 	else if (const FActorDescTreeItem* ActorDescItem = InItem.CastTo<FActorDescTreeItem>())
 	{
 		if (const FWorldPartitionActorDesc* ActorDesc = ActorDescItem->ActorDescHandle.Get())
 		{
-			// For Unloaded Actors, grab the native class and try to convert it to a path name
+			// For Unloaded Actors, grab the native class 
 			FTopLevelAssetPath ClassPath = ActorDesc->GetNativeClass();
 			return AssetClassPaths.Contains(ClassPath);
 		}
