Remove unnecessary line break in Outliner folder tooltips.

diff --git a/Engine/Source/Editor/SceneOutliner/Private/ActorFolderTreeItem.cpp b/Engine/Source/Editor/SceneOutliner/Private/ActorFolderTreeItem.cpp
--- a/Engine/Source/Editor/SceneOutliner/Private/ActorFolderTreeItem.cpp
+++ b/Engine/Source/Editor/SceneOutliner/Private/ActorFolderTreeItem.cpp
@@ -109,8 +109,8 @@ private:
 	{
 		if (const FSceneOutlinerTreeItemPtr TreeItem = TreeItemPtr.Pin())
 		{
-			FText Description = IsInActorEditorContext() ? LOCTEXT("ActorFolderIsCurrentDescription", "This Folder is set as Current Folder. New actors will be added to this Folder.") : FText::GetEmpty();
-			return FText::Format(LOCTEXT("DataLayerTooltipText", "{0}\n{1}"), FText::FromString(TreeItem->GetDisplayString()), Description);
+			FText Description = IsInActorEditorContext() ? LOCTEXT("ActorFolderIsCurrentDescription", "\nThis is your current folder. New actors you create will appear here.") : FText::GetEmpty();
+			return FText::Format(LOCTEXT("DataLayerTooltipText", "{0}{1}"), FText::FromString(TreeItem->GetDisplayString()), Description);
 		}
 
 		return FText();
