Fix regression: double-click on the level name in Outliner should open World Settings.

diff --git a/Engine/Source/Editor/SceneOutliner/Private/WorldTreeItem.cpp b/Engine/Source/Editor/SceneOutliner/Private/WorldTreeItem.cpp
--- a/Engine/Source/Editor/SceneOutliner/Private/WorldTreeItem.cpp
+++ b/Engine/Source/Editor/SceneOutliner/Private/WorldTreeItem.cpp
@@ -114,6 +114,7 @@ struct SWorldTreeLabel : FSceneOutlinerCommonLabelData, public SCompoundWidget
 				.Text(this, &SWorldTreeLabel::GetDisplayText)
 				.HighlightText(SceneOutliner.GetFilterHighlightText())
 				.ColorAndOpacity(this, &SWorldTreeLabel::GetForegroundColor)
+				.OnDoubleClicked(this, &SWorldTreeLabel::HandleDoubleClick)
 				.ToolTip(IDocumentation::Get()->CreateToolTip(
 					TAttribute<FText>(this, &SWorldTreeLabel::GetTooltipText),
 					nullptr,
@@ -139,11 +140,11 @@ private:
 		FText PersistentLevelDisplayName = Item.IsValid() ? FText::FromString(Item->GetWorldName()) : FText();
 		if (Item->CanInteract())
 		{
-			return FText::Format(LOCTEXT("WorldLabel_Tooltip", "The world settings for {0}, double-click to edit"), PersistentLevelDisplayName);
+			return FText::Format(LOCTEXT("WorldLabel_Tooltip", "The {0} level. Double-click to edit World Settings."), PersistentLevelDisplayName);
 		}
 		else
 		{
-			return FText::Format(LOCTEXT("WorldLabel_TooltipNonInteractive", "The world {0}"), PersistentLevelDisplayName);
+			return FText::Format(LOCTEXT("WorldLabel_TooltipNonInteractive", "The {0} level."), PersistentLevelDisplayName);
 		}
 	}
 
@@ -156,6 +157,16 @@ private:
 
 		return FSlateColor::UseForeground();
 	}
+
+	FReply HandleDoubleClick(const FGeometry& MyGeometry, const FPointerEvent& PointerEvent) const
+	{
+		auto Item = TreeItemPtr.Pin();
+		if (Item->CanInteract())
+		{
+			TreeItemPtr.Pin()->OpenWorldSettings();
+		}
+		return FReply::Handled();
+	}
 };
 
 FWorldTreeItem::FWorldTreeItem(UWorld* InWorld)
