Hide 'Show in Changelist' item from the context menu show in the outliner for Skein

diff --git a/Engine/Source/Editor/SceneOutliner/Private/SceneOutlinerSCCHandler.cpp b/Engine/Source/Editor/SceneOutliner/Private/SceneOutlinerSCCHandler.cpp
--- a/Engine/Source/Editor/SceneOutliner/Private/SceneOutlinerSCCHandler.cpp
+++ b/Engine/Source/Editor/SceneOutliner/Private/SceneOutlinerSCCHandler.cpp
@@ -84,6 +84,8 @@ void FSceneOutlinerSCCHandler::CacheCanExecuteVars()
 	bCanExecuteSCCCheckOut = false;
 	bCanExecuteSCCCheckIn = false;
 	bCanExecuteSCCHistory = false;
+	bUsesFileRevisions = false;
+	bUsesChangelists = false;
 
 	if ( ISourceControlModule::Get().IsEnabled() || FUncontrolledChangelistsModule::Get().IsEnabled())
 	{
@@ -127,6 +129,13 @@ void FSceneOutlinerSCCHandler::CacheCanExecuteVars()
 				}
 			}
 		}
+
+		if (ISourceControlModule::Get().GetProvider().IsAvailable())
+		{
+			ISourceControlProvider& Provider = ISourceControlModule::Get().GetProvider();
+			bUsesFileRevisions = Provider.UsesFileRevisions();
+			bUsesChangelists = Provider.UsesChangelists();
+		}
 	}
 }
 
@@ -142,14 +151,6 @@ bool FSceneOutlinerSCCHandler::CanExecuteSCCCheckOut() const
 
 bool FSceneOutlinerSCCHandler::CanExecuteSCCCheckIn() const
 {
-	bool bUsesFileRevisions = true;
-
-	ISourceControlProvider& SourceControlProvider = ISourceControlModule::Get().GetProvider();
-	if (ISourceControlModule::Get().IsEnabled())
-	{
-		bUsesFileRevisions = SourceControlProvider.UsesFileRevisions();
-	}
-
 	return bCanExecuteSCCCheckIn && bUsesFileRevisions;
 }
 
@@ -163,6 +164,11 @@ bool FSceneOutlinerSCCHandler::CanExecuteSCCRefresh() const
 	return ISourceControlModule::Get().IsEnabled();
 }
 
+bool FSceneOutlinerSCCHandler::CanExecuteSCCShowInChangelist() const
+{
+	return bUsesChangelists;
+}
+
 void FSceneOutlinerSCCHandler::FillSourceControlSubMenu(UToolMenu* Menu)
 {
 	FToolMenuSection& Section = Menu->AddSection("AssetSourceControlActions", LOCTEXT("AssetSourceControlActionsMenuHeading", "Source Control"));
@@ -220,16 +226,20 @@ void FSceneOutlinerSCCHandler::FillSourceControlSubMenu(UToolMenu* Menu)
 		);
 	}
 
-	Section.AddSeparator(NAME_None);
-	Section.AddMenuEntry(
-		"SCCFindInChangelist",
-		LOCTEXT("SCCShowInChangelist", "Show in Changelist"),
-		LOCTEXT("SCCShowInChangelistTooltip", "Show the selected assets in the Changelist window."),
-		FSlateIcon(FAppStyle::GetAppStyleSetName(), "SourceControl.ChangelistsTab"),
-		FUIAction(
-			FExecuteAction::CreateSP(this, &FSceneOutlinerSCCHandler::ExecuteSCCShowInChangelist)
-		)
-	);
+	if (CanExecuteSCCShowInChangelist())
+	{
+		Section.AddSeparator(NAME_None);
+		Section.AddMenuEntry(
+			"SCCFindInChangelist",
+			LOCTEXT("SCCShowInChangelist", "Show in Changelist"),
+			LOCTEXT("SCCShowInChangelistTooltip", "Show the selected assets in the Changelist window."),
+			FSlateIcon(FAppStyle::GetAppStyleSetName(), "SourceControl.ChangelistsTab"),
+			FUIAction(
+				FExecuteAction::CreateSP(this, &FSceneOutlinerSCCHandler::ExecuteSCCShowInChangelist),
+				FCanExecuteAction::CreateSP(this, &FSceneOutlinerSCCHandler::CanExecuteSCCShowInChangelist)
+			)
+		);
+	}
 }
 
 void FSceneOutlinerSCCHandler::GetSelectedPackageNames(TArray<FString>& OutPackageNames) const
diff --git a/Engine/Source/Editor/SceneOutliner/Public/SceneOutlinerSCCHandler.h b/Engine/Source/Editor/SceneOutliner/Public/SceneOutlinerSCCHandler.h
--- a/Engine/Source/Editor/SceneOutliner/Public/SceneOutlinerSCCHandler.h
+++ b/Engine/Source/Editor/SceneOutliner/Public/SceneOutlinerSCCHandler.h
@@ -24,6 +24,7 @@ private:
 	bool CanExecuteSCCCheckIn() const;
 	bool CanExecuteSCCHistory() const;
 	bool CanExecuteSCCRefresh() const;
+	bool CanExecuteSCCShowInChangelist() const;
 	void FillSourceControlSubMenu(UToolMenu* Menu);
 	void GetSelectedPackageNames(TArray<FString>& OutPackageNames) const;
 	void GetSelectedPackages(TArray<UPackage*>& OutPackages) const;
@@ -43,6 +44,8 @@ private:
 	bool bCanExecuteSCCCheckOut = false;
 	bool bCanExecuteSCCCheckIn = false;
 	bool bCanExecuteSCCHistory = false;
+	bool bUsesFileRevisions = false;
+	bool bUsesChangelists = false;
 
 	mutable TMap<FSceneOutlinerTreeItemPtr, TSharedPtr<FSceneOutlinerTreeItemSCC>> ItemSourceControls;
 };
