Add ContentBundle Column to actors in scene outliner

diff --git a/Engine/Source/Editor/SceneOutliner/Private/SceneOutlinerModule.cpp b/Engine/Source/Editor/SceneOutliner/Private/SceneOutlinerModule.cpp
--- a/Engine/Source/Editor/SceneOutliner/Private/SceneOutlinerModule.cpp
+++ b/Engine/Source/Editor/SceneOutliner/Private/SceneOutlinerModule.cpp
@@ -25,6 +25,8 @@
 #include "WorldPartition/ActorDescContainer.h"
 #include "WorldPartition/DataLayer/DataLayerInstance.h"
 #include "WorldPartition/DataLayer/DataLayerSubsystem.h"
+#include "WorldPartition/ContentBundle/ContentBundleEngineSubsystem.h"
+#include "WorldPartition/ContentBundle/ContentBundleDescriptor.h"
 #include "LevelInstance/LevelInstanceSubsystem.h"
 #include "LevelInstance/LevelInstanceInterface.h"
 #include "Editor.h"
@@ -309,6 +311,33 @@ void FSceneOutlinerModule::CreateActorInfoColumns(FSceneOutlinerInitializationOp
 		return Builder.ToString();
 	});
 
+	FGetTextForItem ContentBundleInfoText = FGetTextForItem::CreateLambda([](const ISceneOutlinerTreeItem& Item)->FString
+	{	
+		UContentBundleEngineSubsystem* ContentBundleEngineSubsystem = UContentBundleEngineSubsystem::Get();
+		if (const FActorTreeItem* ActorItem = Item.CastTo<FActorTreeItem>())
+		{
+			if (AActor* Actor = ActorItem->Actor.Get())
+			{
+				if (const UContentBundleDescriptor * Descriptor = ContentBundleEngineSubsystem->GetContentBundleDescriptor(Actor->GetContentBundleGuid()))
+				{
+					return Descriptor->GetDisplayName();
+				}
+			}
+		}
+		else if (const FActorDescTreeItem* ActorDescItem = Item.CastTo<FActorDescTreeItem>())
+		{
+			if (const FWorldPartitionActorDesc* ActorDesc = ActorDescItem->ActorDescHandle.Get())
+			{
+				if (const UContentBundleDescriptor* Descriptor = ContentBundleEngineSubsystem->GetContentBundleDescriptor(ActorDesc->GetContentBundleGuid()))
+				{
+					return Descriptor->GetDisplayName();
+				}
+			}
+		}
+
+		return TEXT("");
+	});
+
 	FGetTextForItem SubPackageInfoText = FGetTextForItem::CreateLambda([](const ISceneOutlinerTreeItem& Item) -> FString
 	{
 		if (const FActorTreeItem* ActorItem = Item.CastTo<FActorTreeItem>())
@@ -456,6 +485,7 @@ void FSceneOutlinerModule::CreateActorInfoColumns(FSceneOutlinerInitializationOp
 	AddTextInfoColumn(FSceneOutlinerBuiltInColumnTypes::Level(), LevelColumnName, LevelInfoText);
 	AddTextInfoColumn(FSceneOutlinerBuiltInColumnTypes::Layer(), FSceneOutlinerBuiltInColumnTypes::Layer_Localized(), LayerInfoText);
 	AddTextInfoColumn(FSceneOutlinerBuiltInColumnTypes::DataLayer(), FSceneOutlinerBuiltInColumnTypes::DataLayer_Localized(), DataLayerInfoText);
+	AddTextInfoColumn(FSceneOutlinerBuiltInColumnTypes::ContentBundle(), FSceneOutlinerBuiltInColumnTypes::ContentBundle_Localized(), ContentBundleInfoText);
 	AddTextInfoColumn(FSceneOutlinerBuiltInColumnTypes::SubPackage(), FSceneOutlinerBuiltInColumnTypes::SubPackage_Localized(), SubPackageInfoText);
 	AddTextInfoColumn(FSceneOutlinerBuiltInColumnTypes::Socket(), FSceneOutlinerBuiltInColumnTypes::Socket_Localized(), SocketInfoText);
 	AddTextInfoColumn(FSceneOutlinerBuiltInColumnTypes::IDName(), FSceneOutlinerBuiltInColumnTypes::IDName_Localized(), InternalNameInfoText);
diff --git a/Engine/Source/Editor/SceneOutliner/Public/SceneOutlinerPublicTypes.h b/Engine/Source/Editor/SceneOutliner/Public/SceneOutlinerPublicTypes.h
--- a/Engine/Source/Editor/SceneOutliner/Public/SceneOutlinerPublicTypes.h
+++ b/Engine/Source/Editor/SceneOutliner/Public/SceneOutlinerPublicTypes.h
@@ -47,10 +47,11 @@ struct FSceneOutlinerBuiltInColumnTypes
 	DEFINE_SCENEOUTLINER_BUILTIN_COLUMN_TYPE(Gutter, "Visibility", "VisibilityColumnName", "Visibility");
 	DEFINE_SCENEOUTLINER_BUILTIN_COLUMN_TYPE(ActorInfo, "Type", "TypeColumnName", "Type");
 	DEFINE_SCENEOUTLINER_BUILTIN_COLUMN_TYPE(SourceControl, "Source Control", "SourceControlColumnName", "Source Control");
-	DEFINE_SCENEOUTLINER_BUILTIN_COLUMN_TYPE(Mobility, "Mobility", "SceneOutlinerMobilityColumn", "Mobility");
+	DEFINE_SCENEOUTLINER_BUILTIN_COLUMN_TYPE(Mobility, "Mobility", "SceneOutlinerMobilityColumn", "Mobility"); 
 	DEFINE_SCENEOUTLINER_BUILTIN_COLUMN_TYPE(Level, "Level", "SceneOutlinerLevelColumn", "Level");
 	DEFINE_SCENEOUTLINER_BUILTIN_COLUMN_TYPE(Layer, "Layer", "SceneOutlinerLayerColumn", "Layer");
 	DEFINE_SCENEOUTLINER_BUILTIN_COLUMN_TYPE(DataLayer, "Data Layer", "SceneOutlinerDataLayerColumn", "Data Layer");
+	DEFINE_SCENEOUTLINER_BUILTIN_COLUMN_TYPE(ContentBundle, "Content Bundle", "SceneOutlinerContentBundleColumn", "Content Bundle");
 	DEFINE_SCENEOUTLINER_BUILTIN_COLUMN_TYPE(SubPackage, "Sub Package", "SceneOutlinerSubPackageColumn", "Sub Package");
 	DEFINE_SCENEOUTLINER_BUILTIN_COLUMN_TYPE(Pinned, "Pinned", "SceneOutlinerPinnedColumn", "Pinned");
 	DEFINE_SCENEOUTLINER_BUILTIN_COLUMN_TYPE(IDName, "ID Name", "SceneOutlinerIDNameColumn", "ID Name");
diff --git a/Engine/Source/Runtime/Engine/Private/WorldPartition/ContentBundle/ContentBundleEngineSubsystem.cpp b/Engine/Source/Runtime/Engine/Private/WorldPartition/ContentBundle/ContentBundleEngineSubsystem.cpp
--- a/Engine/Source/Runtime/Engine/Private/WorldPartition/ContentBundle/ContentBundleEngineSubsystem.cpp
+++ b/Engine/Source/Runtime/Engine/Private/WorldPartition/ContentBundle/ContentBundleEngineSubsystem.cpp
@@ -105,6 +105,20 @@ void UContentBundleEngineSubsystem::RequestContentRemoval(FContentBundleClient&
 	}
 }
 
+const UContentBundleDescriptor* UContentBundleEngineSubsystem::GetContentBundleDescriptor(const FGuid& ContentBundleGuid) const
+{
+	for (const TSharedPtr<FContentBundleClient>& Client : ContentBundleClients)
+	{
+		const UContentBundleDescriptor* Descriptor = Client->GetDescriptor();
+		if (Descriptor->GetGuid() == ContentBundleGuid)
+		{
+			return Descriptor;
+		}
+	}
+
+	return nullptr;
+}
+
 UContentBundleEngineSubsystem* UContentBundleEngineSubsystem::Get()
 {
 	if (GEngine)
diff --git a/Engine/Source/Runtime/Engine/Public/WorldPartition/ContentBundle/ContentBundleEngineSubsystem.h b/Engine/Source/Runtime/Engine/Public/WorldPartition/ContentBundle/ContentBundleEngineSubsystem.h
--- a/Engine/Source/Runtime/Engine/Public/WorldPartition/ContentBundle/ContentBundleEngineSubsystem.h
+++ b/Engine/Source/Runtime/Engine/Public/WorldPartition/ContentBundle/ContentBundleEngineSubsystem.h
@@ -28,6 +28,8 @@ public:
 	void RequestContentInjection(FContentBundleClient& Client);
 	void RequestContentRemoval(FContentBundleClient& Client);
 
+	const UContentBundleDescriptor* GetContentBundleDescriptor(const FGuid& ContentBundleGuid) const;
+
 	DECLARE_MULTICAST_DELEGATE_OneParam(FContentBundleRegisterDelegate, TSharedPtr<FContentBundleClient>&);
 	FContentBundleRegisterDelegate OnContentBundleClientRegistered;
 
