World Outliner: fix pin state not sowing anymore on unloaded actors since 22525535.

#rnx

diff --git a/Engine/Source/Editor/SceneOutliner/Private/ActorDescTreeItem.cpp b/Engine/Source/Editor/SceneOutliner/Private/ActorDescTreeItem.cpp
--- a/Engine/Source/Editor/SceneOutliner/Private/ActorDescTreeItem.cpp
+++ b/Engine/Source/Editor/SceneOutliner/Private/ActorDescTreeItem.cpp
@@ -322,9 +322,9 @@ bool FActorDescTreeItem::ShouldShowPinnedState() const
 	if (ActorDescHandle.IsValid())
 	{
 		// Pinning of ActorDescs is only supported on the main world partition
-		if (UWorldPartition* WorldPartition = Cast<UWorldPartition>(ActorDescHandle->GetContainer()))
+		if (UActorDescContainer* Container = ActorDescHandle->GetContainer())
 		{
-			return WorldPartition->IsMainWorldPartition();
+			return Container->IsMainPartitionContainer();
 		}
 	}
 
@@ -335,11 +335,8 @@ bool FActorDescTreeItem::GetPinnedState() const
 {
 	if (ActorDescHandle.IsValid())
 	{
-		// Pinning of ActorDescs is only supported on the main world partition
-		if (UWorldPartition* WorldPartition = Cast<UWorldPartition>(ActorDescHandle->GetContainer()))
-		{
-			return WorldPartition->IsMainWorldPartition() && WorldPartition->IsActorPinned(GetGuid());
-		}
+		UWorldPartition* WorldPartition = ActorDescHandle->GetContainer()->GetWorld()->GetWorldPartition();
+		return WorldPartition ? WorldPartition->IsActorPinned(GetGuid()) : false;
 	}
 	return false;
 }
