Outliner: Fix an issue causing SaveConfig() to be called for each column visible during load

diff --git a/Engine/Source/Editor/SceneOutliner/Private/SSceneOutliner.cpp b/Engine/Source/Editor/SceneOutliner/Private/SSceneOutliner.cpp
--- a/Engine/Source/Editor/SceneOutliner/Private/SSceneOutliner.cpp
+++ b/Engine/Source/Editor/SceneOutliner/Private/SSceneOutliner.cpp
@@ -362,14 +362,34 @@ void SSceneOutliner::HandleHiddenColumnsChanged()
 	{
 		TMap<FName, bool> ColumnVisibilities = OutlinerConfig->ColumnVisibilities;
 
+		bool bAnyColumnVisibilityChanged = false;
+		
 		for (const TPair<FName, TSharedPtr<ISceneOutlinerColumn>>& Pair : Columns)
 		{
-			ColumnVisibilities.Add(Pair.Key, HiddenColumns.Find(Pair.Key) == INDEX_NONE);
-		}
+			const bool bIsColumnVisible = HiddenColumns.Find(Pair.Key) == INDEX_NONE;
 
-		OutlinerConfig->ColumnVisibilities = ColumnVisibilities;
+			// If this column already has a visibility saved in the config, only update if it changed
+			if (bool* ExistingColumnVisibility = ColumnVisibilities.Find(Pair.Key))
+			{
+				if(*ExistingColumnVisibility != bIsColumnVisible)
+				{
+					*ExistingColumnVisibility = bIsColumnVisible;
+					bAnyColumnVisibilityChanged = true;
+				}
+			}
+			else
+			{
+				ColumnVisibilities.Add(Pair.Key, bIsColumnVisible);
+				bAnyColumnVisibilityChanged = true;
+			}
+		}
 
-		SaveConfig();
+		// Only call SaveConfig if something actually changed
+		if(bAnyColumnVisibilityChanged)
+		{
+			OutlinerConfig->ColumnVisibilities = ColumnVisibilities;
+			SaveConfig();
+		}
 	}
 }
 
