Outliner: Fix a regression causing items to not update their state when the packaging mode was change from internal -> external

diff --git a/Engine/Source/Editor/SceneOutliner/Private/SceneOutlinerTreeItemSCC.cpp b/Engine/Source/Editor/SceneOutliner/Private/SceneOutlinerTreeItemSCC.cpp
--- a/Engine/Source/Editor/SceneOutliner/Private/SceneOutlinerTreeItemSCC.cpp
+++ b/Engine/Source/Editor/SceneOutliner/Private/SceneOutlinerTreeItemSCC.cpp
@@ -25,24 +25,21 @@ FSceneOutlinerTreeItemSCC::FSceneOutlinerTreeItemSCC(FSceneOutlinerTreeItemPtr I
 		{
 			if (AActor* Actor = ActorItem->Actor.Get())
 			{
-				if (Actor->IsPackageExternal())
+				ActorPackingModeChangedDelegateHandle = Actor->OnPackagingModeChanged.AddLambda([this](AActor* InActor, bool bExternal)
 				{
-					ActorPackingModeChangedDelegateHandle = Actor->OnPackagingModeChanged.AddLambda([this](AActor* InActor, bool bExternal)
+					if (bExternal)
 					{
-						if (bExternal)
-						{
-							ExternalPackageName = USourceControlHelpers::PackageFilename(InActor->GetExternalPackage());
-							ExternalPackage = InActor->GetExternalPackage();
-							ConnectSourceControl();
-						}
-						else
-						{
-							ExternalPackageName = FString();
-							ExternalPackage = nullptr;
-							DisconnectSourceControl();
-						}
-					});
-				}
+						ExternalPackageName = USourceControlHelpers::PackageFilename(InActor->GetExternalPackage());
+						ExternalPackage = InActor->GetExternalPackage();
+						ConnectSourceControl();
+					}
+					else
+					{
+						ExternalPackageName = FString();
+						ExternalPackage = nullptr;
+						DisconnectSourceControl();
+					}
+				});
 			}
 		}
 		
