Fix crash in scene outliner when deleting more than one folder at a time.
Parent folder enumerated deleted child folders.

diff --git a/Engine/Source/Editor/SceneOutliner/Private/ActorFolderTreeItem.cpp b/Engine/Source/Editor/SceneOutliner/Private/ActorFolderTreeItem.cpp
--- a/Engine/Source/Editor/SceneOutliner/Private/ActorFolderTreeItem.cpp
+++ b/Engine/Source/Editor/SceneOutliner/Private/ActorFolderTreeItem.cpp
@@ -306,6 +306,12 @@ void FActorFolderTreeItem::Delete(const FFolder& InNewParentFolder)
 	}
 
 	FActorFolders::Get().DeleteFolder(*World, GetFolder());
+
+	// Remove this folder from parent otherwise when parent folder and child folder deleted together, parent will call MoveTo on deleted child folder leading to crash
+	if (FSceneOutlinerTreeItemPtr ParentPtr = GetParent())
+	{
+		ParentPtr->RemoveChild(SharedThis(this));
+	}
 }
 
 void FActorFolderTreeItem::MoveTo(const FFolder& InNewParentFolder)
