Outliner: Remove the type text that shows in the Label column during search

diff --git a/Engine/Source/Editor/SceneOutliner/Private/ActorTreeItem.cpp b/Engine/Source/Editor/SceneOutliner/Private/ActorTreeItem.cpp
--- a/Engine/Source/Editor/SceneOutliner/Private/ActorTreeItem.cpp
+++ b/Engine/Source/Editor/SceneOutliner/Private/ActorTreeItem.cpp
@@ -62,17 +62,6 @@ struct SActorTreeLabel : FSceneOutlinerCommonLabelData, public SCompoundWidget
 				{
 					return !CanExecuteRenameRequest(Item.Get());
 				})
-			]
-
-			+ SHorizontalBox::Slot()
-			.VAlign(VAlign_Center)
-			.AutoWidth()
-			.Padding(0.0f, 0.f, 3.0f, 0.0f)
-			[
-				SNew(STextBlock)
-				.Text(this, &SActorTreeLabel::GetTypeText)
-				.Visibility(this, &SActorTreeLabel::GetTypeTextVisibility)
-				.HighlightText(HighlightText)
 			];
 
 		if (WeakSceneOutliner.Pin()->GetMode()->IsInteractive())
@@ -145,21 +134,6 @@ private:
 		return FText();
 	}
 
-	FText GetTypeText() const
-	{
-		if (const AActor* Actor = ActorPtr.Get())
-		{
-			return FText::FromName(Actor->GetClass()->GetFName());
-		}
-
-		return FText();
-	}
-
-	EVisibility GetTypeTextVisibility() const
-	{
-		return HighlightText.Get().IsEmpty() ? EVisibility::Collapsed : EVisibility::Visible;
-	}
-
 	const FSlateBrush* GetIcon() const
 	{
 		if (const AActor* Actor = ActorPtr.Get())
