Added world outliner filter "Only in current Data Layers" similar to "Only in Current Level".

diff --git a/Engine/Source/Editor/SceneOutliner/Private/ActorBrowsingMode.cpp b/Engine/Source/Editor/SceneOutliner/Private/ActorBrowsingMode.cpp
--- a/Engine/Source/Editor/SceneOutliner/Private/ActorBrowsingMode.cpp
+++ b/Engine/Source/Editor/SceneOutliner/Private/ActorBrowsingMode.cpp
@@ -39,6 +39,8 @@
 #include "LevelInstance/LevelInstanceEditorInstanceActor.h"
 #include "WorldPartition/WorldPartition.h"
 #include "WorldPartition/WorldPartitionSubsystem.h"
+#include "WorldPartition/DataLayer/DataLayerSubsystem.h"
+#include "Subsystems/ActorEditorContextSubsystem.h"
 #include "ISourceControlModule.h"
 #include "ISourceControlProvider.h"
 #include "SourceControlOperations.h"
@@ -66,6 +68,8 @@ FActorBrowsingMode::FActorBrowsingMode(SSceneOutliner* InSceneOutliner, TWeakObj
 
 	GEditor->OnSelectUnloadedActorsEvent().AddRaw(this, &FActorBrowsingMode::OnSelectUnloadedActors);
 
+	UActorEditorContextSubsystem::Get()->OnActorEditorContextSubsystemChanged().AddRaw(this, &FActorBrowsingMode::OnActorEditorContextSubsystemChanged);
+
 	FEditorDelegates::OnEditCutActorsBegin.AddRaw(this, &FActorBrowsingMode::OnEditCutActorsBegin);
 	FEditorDelegates::OnEditCutActorsEnd.AddRaw(this, &FActorBrowsingMode::OnEditCutActorsEnd);
 	FEditorDelegates::OnEditCopyActorsBegin.AddRaw(this, &FActorBrowsingMode::OnEditCopyActorsBegin);
@@ -129,6 +133,42 @@ FActorBrowsingMode::FActorBrowsingMode(SSceneOutliner* InSceneOutliner, TWeakObj
 		});
 	FilterInfoMap.Add(TEXT("ShowOnlyCurrentLevel"), OnlyCurrentLevelInfo);
 
+	FSceneOutlinerFilterInfo OnlyCurrentDataLayersInfo(LOCTEXT("ToggleShowOnlyCurrentDataLayers", "Only in Current Data Layers"), LOCTEXT("ToggleShowOnlyCurrentDataLayersToolTip", "When enabled, only shows Actors that are in the Current Data Layers."), LocalSettings.bShowOnlyActorsInCurrentDataLayers, FCreateSceneOutlinerFilter::CreateStatic(&FActorBrowsingMode::CreateIsInCurrentDataLayersFilter));
+	OnlyCurrentDataLayersInfo.OnToggle().AddLambda([this](bool bIsActive)
+		{
+			FActorBrowsingModeConfig* Settings = GetMutableConfig();
+			if (Settings)
+			{
+				Settings->bShowOnlyActorsInCurrentDataLayers = bIsActive;
+				SaveConfig();
+			}
+		});
+	FilterInfoMap.Add(TEXT("ShowOnlyCurrentDataLayers"), OnlyCurrentDataLayersInfo);
+
+	// Add a filter for unloaded actors to properly reflect the bShowOnlyActorsInCurrentDataLayers flag.
+	SceneOutliner->AddFilter(MakeShared<FActorDescFilter>(FActorDescTreeItem::FFilterPredicate::CreateLambda([this](const FWorldPartitionActorDesc* ActorDesc)
+		{
+			FActorBrowsingModeConfig* Settings = GetMutableConfig();
+			if (Settings && Settings->bShowOnlyActorsInCurrentDataLayers)
+			{
+				if (const UDataLayerSubsystem* DataLayerSubsystem = RepresentingWorld->GetSubsystem<UDataLayerSubsystem>())
+				{
+					for (const FName& DataLayerInstanceName : ActorDesc->GetDataLayerInstanceNames())
+					{
+						if (const UDataLayerInstance* const DataLayerInstance = DataLayerSubsystem->GetDataLayerInstance(DataLayerInstanceName))
+						{
+							if (DataLayerInstance->IsInActorEditorContext())
+							{
+								return true;
+							}
+						}
+					}
+				}
+				return false;
+			}
+			return true;
+		}), FSceneOutlinerFilter::EDefaultBehaviour::Pass));
+
 	bHideComponents = LocalSettings.bHideActorComponents;
 	FSceneOutlinerFilterInfo HideComponentsInfo(LOCTEXT("ToggleHideActorComponents", "Hide Actor Components"), LOCTEXT("ToggleHideActorComponentsToolTip", "When enabled, hides components belonging to actors."), LocalSettings.bHideActorComponents, FCreateSceneOutlinerFilter::CreateStatic(&FActorBrowsingMode::CreateHideComponentsFilter));
 	HideComponentsInfo.OnToggle().AddLambda([this](bool bIsActive)
@@ -242,6 +282,8 @@ FActorBrowsingMode::~FActorBrowsingMode()
 
 	GEditor->OnSelectUnloadedActorsEvent().RemoveAll(this);
 
+	UActorEditorContextSubsystem::Get()->OnActorEditorContextSubsystemChanged().RemoveAll(this);
+
 	FEditorDelegates::OnEditCutActorsBegin.RemoveAll(this);
 	FEditorDelegates::OnEditCutActorsEnd.RemoveAll(this);
 	FEditorDelegates::OnEditCopyActorsBegin.RemoveAll(this);
@@ -341,6 +383,11 @@ FSlateColor FActorBrowsingMode::GetStatusTextColor() const
 	}
 }
 
+void FActorBrowsingMode::OnActorEditorContextSubsystemChanged()
+{
+	SceneOutliner->FullRefresh();
+}
+
 void FActorBrowsingMode::OnToggleAlwaysFrameSelection()
 {
 	bAlwaysFrameSelection = !bAlwaysFrameSelection;
@@ -419,6 +466,22 @@ TSharedRef<FSceneOutlinerFilter> FActorBrowsingMode::CreateIsInCurrentLevelFilte
 		}), FSceneOutlinerFilter::EDefaultBehaviour::Pass));
 }
 
+TSharedRef<FSceneOutlinerFilter> FActorBrowsingMode::CreateIsInCurrentDataLayersFilter()
+{
+	return MakeShareable(new FActorFilter(FActorTreeItem::FFilterPredicate::CreateStatic([](const AActor* InActor)
+		{
+			for (const UDataLayerInstance* DataLayerInstance : InActor->GetDataLayerInstances())
+			{
+				if (DataLayerInstance->IsInActorEditorContext())
+				{
+					return true;
+				}
+			}
+
+			return false;
+		}), FSceneOutlinerFilter::EDefaultBehaviour::Pass));
+}
+
 TSharedRef<FSceneOutlinerFilter> FActorBrowsingMode::CreateHideComponentsFilter()
 {
 	return MakeShared<TSceneOutlinerPredicateFilter<FComponentTreeItem>>(TSceneOutlinerPredicateFilter<FComponentTreeItem>(
diff --git a/Engine/Source/Editor/SceneOutliner/Private/ActorBrowsingModeSettings.h b/Engine/Source/Editor/SceneOutliner/Private/ActorBrowsingModeSettings.h
--- a/Engine/Source/Editor/SceneOutliner/Private/ActorBrowsingModeSettings.h
+++ b/Engine/Source/Editor/SceneOutliner/Private/ActorBrowsingModeSettings.h
@@ -27,6 +27,10 @@ public:
 	UPROPERTY()
 	bool bShowOnlyActorsInCurrentLevel = false;
 
+	/** True when the Scene Outliner is showing only Actors that exist in the current data layers */
+	UPROPERTY()
+	bool bShowOnlyActorsInCurrentDataLayers = false;
+
 	/** True when the Scene Outliner is only displaying selected Actors */
 	UPROPERTY()
 	bool bShowOnlySelectedActors = false;
diff --git a/Engine/Source/Editor/SceneOutliner/Public/ActorBrowsingMode.h b/Engine/Source/Editor/SceneOutliner/Public/ActorBrowsingMode.h
--- a/Engine/Source/Editor/SceneOutliner/Public/ActorBrowsingMode.h
+++ b/Engine/Source/Editor/SceneOutliner/Public/ActorBrowsingMode.h
@@ -118,10 +118,13 @@ private:
 
 	void SynchronizeSelectedActorDescs();
 
+	void OnActorEditorContextSubsystemChanged();
+
 	/** Filter factories */
 	static TSharedRef<FSceneOutlinerFilter> CreateShowOnlySelectedActorsFilter();
 	static TSharedRef<FSceneOutlinerFilter> CreateHideTemporaryActorsFilter();
 	static TSharedRef<FSceneOutlinerFilter> CreateIsInCurrentLevelFilter();
+	static TSharedRef<FSceneOutlinerFilter> CreateIsInCurrentDataLayersFilter();
 	static TSharedRef<FSceneOutlinerFilter> CreateHideComponentsFilter();
 	static TSharedRef<FSceneOutlinerFilter> CreateHideLevelInstancesFilter();
 	static TSharedRef<FSceneOutlinerFilter> CreateHideUnloadedActorsFilter();
