Added world outliner filter "Only in current Content Bundle" similar to "Only in Current Level".

diff --git a/Engine/Source/Editor/SceneOutliner/Private/ActorBrowsingMode.cpp b/Engine/Source/Editor/SceneOutliner/Private/ActorBrowsingMode.cpp
--- a/Engine/Source/Editor/SceneOutliner/Private/ActorBrowsingMode.cpp
+++ b/Engine/Source/Editor/SceneOutliner/Private/ActorBrowsingMode.cpp
@@ -40,6 +40,8 @@
 #include "WorldPartition/WorldPartition.h"
 #include "WorldPartition/WorldPartitionSubsystem.h"
 #include "WorldPartition/DataLayer/DataLayerSubsystem.h"
+#include "WorldPartition/ContentBundle/ContentBundleEditorSubsystem.h"
+#include "WorldPartition/ContentBundle/ContentBundleEditor.h"
 #include "Subsystems/ActorEditorContextSubsystem.h"
 #include "ISourceControlModule.h"
 #include "ISourceControlProvider.h"
@@ -169,6 +171,41 @@ FActorBrowsingMode::FActorBrowsingMode(SSceneOutliner* InSceneOutliner, TWeakObj
 			return true;
 		}), FSceneOutlinerFilter::EDefaultBehaviour::Pass));
 
+	FSceneOutlinerFilterInfo OnlyCurrentContentBundleInfo(LOCTEXT("ToggleShowOnlyCurrentContentBundle", "Only in Current Content Bundle"), LOCTEXT("ToggleShowOnlyCurrentContentBundleToolTip", "When enabled, only shows Actors that are in the Current Content Bundle."), LocalSettings.bShowOnlyActorsInCurrentContentBundle, FCreateSceneOutlinerFilter::CreateStatic(&FActorBrowsingMode::CreateIsInCurrentContentBundleFilter));
+	OnlyCurrentContentBundleInfo.OnToggle().AddLambda([this](bool bIsActive)
+		{
+			FActorBrowsingModeConfig* Settings = GetMutableConfig();
+			if (Settings)
+			{
+				Settings->bShowOnlyActorsInCurrentContentBundle = bIsActive;
+				SaveConfig();
+			}
+		});
+	FilterInfoMap.Add(TEXT("ShowOnlyCurrentContentBundle"), OnlyCurrentContentBundleInfo);
+
+	// Add a filter for unloaded actors to properly reflect the bShowOnlyActorsInCurrentContentBundle flag.
+	SceneOutliner->AddFilter(MakeShared<FActorDescFilter>(FActorDescTreeItem::FFilterPredicate::CreateLambda([this](const FWorldPartitionActorDesc* ActorDesc)
+		{
+			FActorBrowsingModeConfig* Settings = GetMutableConfig();
+			if (Settings && Settings->bShowOnlyActorsInCurrentContentBundle)
+			{
+				UContentBundleEditorSubsystem* ContentBundleEditorSubsystem = UContentBundleEditorSubsystem::Get();
+				if (!ContentBundleEditorSubsystem || !ContentBundleEditorSubsystem->IsEditingContentBundle())
+				{
+					return true;
+				}
+
+				if (ActorDesc->GetContentBundleGuid().IsValid())
+				{
+					TSharedPtr<FContentBundleEditor> ContentBundleEditor = ContentBundleEditorSubsystem->GetEditorContentBundle(ActorDesc->GetContentBundleGuid());
+					return ContentBundleEditor && ContentBundleEditorSubsystem->IsContentBundleEditingActivated(ContentBundleEditor);
+				}
+
+				return false;
+			}
+			return true;
+		}), FSceneOutlinerFilter::EDefaultBehaviour::Pass));
+
 	bHideComponents = LocalSettings.bHideActorComponents;
 	FSceneOutlinerFilterInfo HideComponentsInfo(LOCTEXT("ToggleHideActorComponents", "Hide Actor Components"), LOCTEXT("ToggleHideActorComponentsToolTip", "When enabled, hides components belonging to actors."), LocalSettings.bHideActorComponents, FCreateSceneOutlinerFilter::CreateStatic(&FActorBrowsingMode::CreateHideComponentsFilter));
 	HideComponentsInfo.OnToggle().AddLambda([this](bool bIsActive)
@@ -482,6 +519,25 @@ TSharedRef<FSceneOutlinerFilter> FActorBrowsingMode::CreateIsInCurrentDataLayers
 		}), FSceneOutlinerFilter::EDefaultBehaviour::Pass));
 }
 
+TSharedRef<FSceneOutlinerFilter> FActorBrowsingMode::CreateIsInCurrentContentBundleFilter()
+{
+	return MakeShareable(new FActorFilter(FActorTreeItem::FFilterPredicate::CreateStatic([](const AActor* InActor)
+		{
+			UContentBundleEditorSubsystem* ContentBundleEditorSubsystem = UContentBundleEditorSubsystem::Get();
+			if (!ContentBundleEditorSubsystem || !ContentBundleEditorSubsystem->IsEditingContentBundle())
+			{
+				return true;
+			}
+
+			if (InActor->GetContentBundleGuid().IsValid())
+			{
+				TSharedPtr<FContentBundleEditor> ContentBundleEditor = ContentBundleEditorSubsystem->GetEditorContentBundle(InActor->GetContentBundleGuid());
+				return ContentBundleEditor && ContentBundleEditorSubsystem->IsContentBundleEditingActivated(ContentBundleEditor);
+			}
+			return false;
+		}), FSceneOutlinerFilter::EDefaultBehaviour::Pass));
+}
+
 TSharedRef<FSceneOutlinerFilter> FActorBrowsingMode::CreateHideComponentsFilter()
 {
 	return MakeShared<TSceneOutlinerPredicateFilter<FComponentTreeItem>>(TSceneOutlinerPredicateFilter<FComponentTreeItem>(
diff --git a/Engine/Source/Editor/SceneOutliner/Private/ActorBrowsingModeSettings.h b/Engine/Source/Editor/SceneOutliner/Private/ActorBrowsingModeSettings.h
--- a/Engine/Source/Editor/SceneOutliner/Private/ActorBrowsingModeSettings.h
+++ b/Engine/Source/Editor/SceneOutliner/Private/ActorBrowsingModeSettings.h
@@ -31,6 +31,10 @@ public:
 	UPROPERTY()
 	bool bShowOnlyActorsInCurrentDataLayers = false;
 
+	/** True when the Scene Outliner is showing only Actors that exist in the current content bundle */
+	UPROPERTY()
+	bool bShowOnlyActorsInCurrentContentBundle = false;
+
 	/** True when the Scene Outliner is only displaying selected Actors */
 	UPROPERTY()
 	bool bShowOnlySelectedActors = false;
diff --git a/Engine/Source/Editor/SceneOutliner/Public/ActorBrowsingMode.h b/Engine/Source/Editor/SceneOutliner/Public/ActorBrowsingMode.h
--- a/Engine/Source/Editor/SceneOutliner/Public/ActorBrowsingMode.h
+++ b/Engine/Source/Editor/SceneOutliner/Public/ActorBrowsingMode.h
@@ -125,6 +125,7 @@ private:
 	static TSharedRef<FSceneOutlinerFilter> CreateHideTemporaryActorsFilter();
 	static TSharedRef<FSceneOutlinerFilter> CreateIsInCurrentLevelFilter();
 	static TSharedRef<FSceneOutlinerFilter> CreateIsInCurrentDataLayersFilter();
+	static TSharedRef<FSceneOutlinerFilter> CreateIsInCurrentContentBundleFilter();
 	static TSharedRef<FSceneOutlinerFilter> CreateHideComponentsFilter();
 	static TSharedRef<FSceneOutlinerFilter> CreateHideLevelInstancesFilter();
 	static TSharedRef<FSceneOutlinerFilter> CreateHideUnloadedActorsFilter();
diff --git a/Engine/Source/Editor/SceneOutliner/SceneOutliner.Build.cs b/Engine/Source/Editor/SceneOutliner/SceneOutliner.Build.cs
--- a/Engine/Source/Editor/SceneOutliner/SceneOutliner.Build.cs
+++ b/Engine/Source/Editor/SceneOutliner/SceneOutliner.Build.cs
@@ -25,6 +25,7 @@ public class SceneOutliner : ModuleRules
 				"UncontrolledChangelists",
 				"EditorWidgets",
 				"ToolWidgets",
+				"WorldPartitionEditor",
 				"UnsavedAssetsTracker",
 				"UncontrolledChangelists"
 			}
diff --git a/Engine/Source/Editor/WorldPartitionEditor/Private/WorldPartition/ContentBundle/ContentBundleEditorSubsystem.cpp b/Engine/Source/Editor/WorldPartitionEditor/Private/WorldPartition/ContentBundle/ContentBundleEditorSubsystem.cpp
--- a/Engine/Source/Editor/WorldPartitionEditor/Private/WorldPartition/ContentBundle/ContentBundleEditorSubsystem.cpp
+++ b/Engine/Source/Editor/WorldPartitionEditor/Private/WorldPartition/ContentBundle/ContentBundleEditorSubsystem.cpp
@@ -264,6 +264,16 @@ TSharedPtr<FContentBundleEditor> UContentBundleEditorSubsystem::GetEditorContent
 	return nullptr;
 }
 
+TSharedPtr<FContentBundleEditor> UContentBundleEditorSubsystem::GetEditorContentBundle(const FGuid& ContentBundleGuid) const
+{
+	if (UContentBundleManager* ContentBundleManager = GetWorld()->ContentBundleManager)
+	{
+		return ContentBundleManager->GetEditorContentBundle(ContentBundleGuid);
+	}
+
+	return nullptr;
+}
+
 bool UContentBundleEditorSubsystem::HasContentBundle(const UContentBundleDescriptor* ContentBundleDescriptor) const
 {
 	return GetEditorContentBundle(ContentBundleDescriptor) != nullptr;
@@ -289,6 +299,11 @@ void UContentBundleEditorSubsystem::UnreferenceAllActors(FContentBundleEditor& E
 	EditorContentBundle.UnreferenceAllActors();
 }
 
+bool UContentBundleEditorSubsystem::IsEditingContentBundle() const
+{
+	return ContentBundleEditionSubModule->IsEditingContentBundle();
+}
+
 bool UContentBundleEditorSubsystem::ActivateContentBundleEditing(TSharedPtr<FContentBundleEditor>& ContentBundleEditor)  const
 {
 	return ContentBundleEditionSubModule->ActivateContentBundleEditing(ContentBundleEditor);
diff --git a/Engine/Source/Editor/WorldPartitionEditor/Public/WorldPartition/ContentBundle/ContentBundleEditorSubsystem.h b/Engine/Source/Editor/WorldPartitionEditor/Public/WorldPartition/ContentBundle/ContentBundleEditorSubsystem.h
--- a/Engine/Source/Editor/WorldPartitionEditor/Public/WorldPartition/ContentBundle/ContentBundleEditorSubsystem.h
+++ b/Engine/Source/Editor/WorldPartitionEditor/Public/WorldPartition/ContentBundle/ContentBundleEditorSubsystem.h
@@ -107,6 +107,7 @@ public:
 
 	TArray<TSharedPtr<FContentBundleEditor>> GetEditorContentBundles();
 	TSharedPtr<FContentBundleEditor> GetEditorContentBundle(const UContentBundleDescriptor* ContentBundleDescriptor) const;
+	TSharedPtr<FContentBundleEditor> GetEditorContentBundle(const FGuid& ContentBundleGuid) const;
 
 	bool HasContentBundle(const UContentBundleDescriptor* ContentBundleDescriptor) const;
 
@@ -116,6 +117,7 @@ public:
 	void ReferenceAllActors(FContentBundleEditor& EditorContentBundle);
 	void UnreferenceAllActors(FContentBundleEditor& EditorContentBundle);
 
+	bool IsEditingContentBundle() const;
 	bool ActivateContentBundleEditing(TSharedPtr<FContentBundleEditor>& ContentBundleEditor) const;
 	bool DeactivateContentBundleEditing(TSharedPtr<FContentBundleEditor>& ContentBundleEditor) const;
 	bool IsContentBundleEditingActivated(TSharedPtr<FContentBundleEditor>& ContentBundleEditor) const;
diff --git a/Engine/Source/Runtime/Engine/Private/WorldPartition/ContentBundle/ContentBundleContainer.cpp b/Engine/Source/Runtime/Engine/Private/WorldPartition/ContentBundle/ContentBundleContainer.cpp
--- a/Engine/Source/Runtime/Engine/Private/WorldPartition/ContentBundle/ContentBundleContainer.cpp
+++ b/Engine/Source/Runtime/Engine/Private/WorldPartition/ContentBundle/ContentBundleContainer.cpp
@@ -124,6 +124,19 @@ TSharedPtr<FContentBundleEditor> FContentBundleContainer::GetEditorContentBundle
 	return nullptr;
 }
 
+TSharedPtr<FContentBundleEditor> FContentBundleContainer::GetEditorContentBundle(const FGuid& ContentBundleGuid) const
+{
+	for (const TSharedPtr<FContentBundleEditor>& ContentBundle : GetEditorContentBundles())
+	{
+		if (ContentBundle->GetDescriptor()->GetGuid() == ContentBundleGuid)
+		{
+			return ContentBundle;
+		}
+	}
+
+	return nullptr;
+}
+
 const TArray<TSharedPtr<FContentBundleEditor>>& FContentBundleContainer::GetEditorContentBundles() const
 {
 	check(UseEditorContentBundle());
diff --git a/Engine/Source/Runtime/Engine/Private/WorldPartition/ContentBundle/ContentBundleWorldSubsystem.cpp b/Engine/Source/Runtime/Engine/Private/WorldPartition/ContentBundle/ContentBundleWorldSubsystem.cpp
--- a/Engine/Source/Runtime/Engine/Private/WorldPartition/ContentBundle/ContentBundleWorldSubsystem.cpp
+++ b/Engine/Source/Runtime/Engine/Private/WorldPartition/ContentBundle/ContentBundleWorldSubsystem.cpp
@@ -146,6 +146,20 @@ TSharedPtr<FContentBundleEditor> UContentBundleManager::GetEditorContentBundle(c
 	return nullptr;
 }
 
+TSharedPtr<FContentBundleEditor> UContentBundleManager::GetEditorContentBundle(const FGuid& ContentBundleGuid) const
+{
+	for (const TUniquePtr<FContentBundleContainer>& ContentBundleContainer : ContentBundleContainers)
+	{
+		TSharedPtr<FContentBundleEditor> ContentBundleEditor = ContentBundleContainer->GetEditorContentBundle(ContentBundleGuid);
+		if (ContentBundleEditor)
+		{
+			return ContentBundleEditor;
+		}
+	}
+
+	return nullptr;
+}
+
 #endif
 
 uint32 UContentBundleManager::GetContentBundleContainerIndex(const UWorld* InjectedWorld) const
diff --git a/Engine/Source/Runtime/Engine/Public/WorldPartition/ContentBundle/ContentBundleContainer.h b/Engine/Source/Runtime/Engine/Public/WorldPartition/ContentBundle/ContentBundleContainer.h
--- a/Engine/Source/Runtime/Engine/Public/WorldPartition/ContentBundle/ContentBundleContainer.h
+++ b/Engine/Source/Runtime/Engine/Public/WorldPartition/ContentBundle/ContentBundleContainer.h
@@ -50,6 +50,7 @@ public:
 
 #if WITH_EDITOR
 	TSharedPtr<FContentBundleEditor> GetEditorContentBundle(const UContentBundleDescriptor* Descriptor) const;
+	TSharedPtr<FContentBundleEditor> GetEditorContentBundle(const FGuid& ContentBundleGuid) const;
 
 	const TArray<TSharedPtr<FContentBundleEditor>>& GetEditorContentBundles() const;
 	TArray<TSharedPtr<FContentBundleEditor>>& GetEditorContentBundles();
diff --git a/Engine/Source/Runtime/Engine/Public/WorldPartition/ContentBundle/ContentBundleWorldSubsystem.h b/Engine/Source/Runtime/Engine/Public/WorldPartition/ContentBundle/ContentBundleWorldSubsystem.h
--- a/Engine/Source/Runtime/Engine/Public/WorldPartition/ContentBundle/ContentBundleWorldSubsystem.h
+++ b/Engine/Source/Runtime/Engine/Public/WorldPartition/ContentBundle/ContentBundleWorldSubsystem.h
@@ -37,6 +37,7 @@ public:
 #if WITH_EDITOR
 	bool GetEditorContentBundle(TArray<TSharedPtr<FContentBundleEditor>>& OutContentBundles);
 	TSharedPtr<FContentBundleEditor> GetEditorContentBundle(const UContentBundleDescriptor* Descriptor, const UWorld* ContentBundleWorld) const;
+	TSharedPtr<FContentBundleEditor> GetEditorContentBundle(const FGuid& ContentBundleGuid) const;
 
 	UContentBundleDuplicateForPIEHelper* GetPIEDuplicateHelper() const { return PIEDuplicateHelper; }
 #endif
