Fix Boot Test

#rnx

diff --git a/Engine/Source/Editor/SceneOutliner/Private/ActorBrowsingMode.cpp b/Engine/Source/Editor/SceneOutliner/Private/ActorBrowsingMode.cpp
--- a/Engine/Source/Editor/SceneOutliner/Private/ActorBrowsingMode.cpp
+++ b/Engine/Source/Editor/SceneOutliner/Private/ActorBrowsingMode.cpp
@@ -282,7 +282,10 @@ FActorBrowsingMode::~FActorBrowsingMode()
 
 	GEditor->OnSelectUnloadedActorsEvent().RemoveAll(this);
 
-	UActorEditorContextSubsystem::Get()->OnActorEditorContextSubsystemChanged().RemoveAll(this);
+	if (UActorEditorContextSubsystem* ActorEditorContextSubsystem = UActorEditorContextSubsystem::Get())
+	{
+		ActorEditorContextSubsystem->OnActorEditorContextSubsystemChanged().RemoveAll(this);
+	}
 
 	FEditorDelegates::OnEditCutActorsBegin.RemoveAll(this);
 	FEditorDelegates::OnEditCutActorsEnd.RemoveAll(this);
