World outliner "Only in current data layers" filter will not filter if there are no current data layers.

diff --git a/Engine/Source/Editor/SceneOutliner/Private/ActorBrowsingMode.cpp b/Engine/Source/Editor/SceneOutliner/Private/ActorBrowsingMode.cpp
--- a/Engine/Source/Editor/SceneOutliner/Private/ActorBrowsingMode.cpp
+++ b/Engine/Source/Editor/SceneOutliner/Private/ActorBrowsingMode.cpp
@@ -133,7 +133,7 @@ FActorBrowsingMode::FActorBrowsingMode(SSceneOutliner* InSceneOutliner, TWeakObj
 		});
 	FilterInfoMap.Add(TEXT("ShowOnlyCurrentLevel"), OnlyCurrentLevelInfo);
 
-	FSceneOutlinerFilterInfo OnlyCurrentDataLayersInfo(LOCTEXT("ToggleShowOnlyCurrentDataLayers", "Only in Current Data Layers"), LOCTEXT("ToggleShowOnlyCurrentDataLayersToolTip", "When enabled, only shows Actors that are in the Current Data Layers."), LocalSettings.bShowOnlyActorsInCurrentDataLayers, FCreateSceneOutlinerFilter::CreateStatic(&FActorBrowsingMode::CreateIsInCurrentDataLayersFilter));
+	FSceneOutlinerFilterInfo OnlyCurrentDataLayersInfo(LOCTEXT("ToggleShowOnlyCurrentDataLayers", "Only in any Current Data Layers"), LOCTEXT("ToggleShowOnlyCurrentDataLayersToolTip", "When enabled, only shows Actors that are in any Current Data Layers."), LocalSettings.bShowOnlyActorsInCurrentDataLayers, FCreateSceneOutlinerFilter::CreateStatic(&FActorBrowsingMode::CreateIsInCurrentDataLayersFilter));
 	OnlyCurrentDataLayersInfo.OnToggle().AddLambda([this](bool bIsActive)
 		{
 			FActorBrowsingModeConfig* Settings = GetMutableConfig();
@@ -151,16 +151,19 @@ FActorBrowsingMode::FActorBrowsingMode(SSceneOutliner* InSceneOutliner, TWeakObj
 			FActorBrowsingModeConfig* Settings = GetMutableConfig();
 			if (Settings && Settings->bShowOnlyActorsInCurrentDataLayers)
 			{
-				if (const UDataLayerSubsystem* DataLayerSubsystem = RepresentingWorld->GetSubsystem<UDataLayerSubsystem>())
+				const UDataLayerSubsystem* DataLayerSubsystem = RepresentingWorld.IsValid() ? RepresentingWorld->GetSubsystem<UDataLayerSubsystem>() : nullptr;
+				if (!DataLayerSubsystem || DataLayerSubsystem->GetActorEditorContextDataLayers().IsEmpty())
 				{
-					for (const FName& DataLayerInstanceName : ActorDesc->GetDataLayerInstanceNames())
+					return true;
+				}
+				
+				for (const FName& DataLayerInstanceName : ActorDesc->GetDataLayerInstanceNames())
+				{
+					if (const UDataLayerInstance* const DataLayerInstance = DataLayerSubsystem->GetDataLayerInstance(DataLayerInstanceName))
 					{
-						if (const UDataLayerInstance* const DataLayerInstance = DataLayerSubsystem->GetDataLayerInstance(DataLayerInstanceName))
+						if (DataLayerInstance->IsInActorEditorContext())
 						{
-							if (DataLayerInstance->IsInActorEditorContext())
-							{
-								return true;
-							}
+							return true;
 						}
 					}
 				}
@@ -473,6 +476,12 @@ TSharedRef<FSceneOutlinerFilter> FActorBrowsingMode::CreateIsInCurrentDataLayers
 {
 	return MakeShareable(new FActorFilter(FActorTreeItem::FFilterPredicate::CreateStatic([](const AActor* InActor)
 		{
+			const UDataLayerSubsystem* DataLayerSubsystem = InActor->GetWorld() ? InActor->GetWorld()->GetSubsystem<UDataLayerSubsystem>() : nullptr;
+			if (!DataLayerSubsystem || DataLayerSubsystem->GetActorEditorContextDataLayers().IsEmpty())
+			{
+				return true;
+			}
+
 			for (const UDataLayerInstance* DataLayerInstance : InActor->GetDataLayerInstances())
 			{
 				if (DataLayerInstance->IsInActorEditorContext())
