Added world outliner filter "Only in current Content Bundle" similar to "Only in Current Level".

diff --git a/Engine/Source/Editor/SceneOutliner/Private/ActorBrowsingMode.cpp b/Engine/Source/Editor/SceneOutliner/Private/ActorBrowsingMode.cpp
--- a/Engine/Source/Editor/SceneOutliner/Private/ActorBrowsingMode.cpp
+++ b/Engine/Source/Editor/SceneOutliner/Private/ActorBrowsingMode.cpp
@@ -40,6 +40,7 @@
 #include "WorldPartition/WorldPartition.h"
 #include "WorldPartition/WorldPartitionSubsystem.h"
 #include "WorldPartition/DataLayer/DataLayerSubsystem.h"
+#include "WorldPartition/IWorldPartitionEditorModule.h"
 #include "Subsystems/ActorEditorContextSubsystem.h"
 #include "ISourceControlModule.h"
 #include "ISourceControlProvider.h"
@@ -61,6 +62,8 @@ FActorBrowsingMode::FActorBrowsingMode(SSceneOutliner* InSceneOutliner, TWeakObj
 	: FActorModeInteractive(FActorModeParams(InSceneOutliner, InSpecifiedWorldToDisplay,  /* bHideComponents */ true, /* bHideLevelInstanceHierarchy */ false, /* bHideUnloadedActors */ false, /* bHideEmptyFolders */ false))
 	, FilteredActorCount(0)
 {
+	WorldPartitionEditorModule = FModuleManager::GetModulePtr<IWorldPartitionEditorModule>("WorldPartitionEditor");
+
 	// Capture selection changes of bones from mesh selection in fracture tools
 	FSceneOutlinerDelegates::Get().OnComponentsUpdated.AddRaw(this, &FActorBrowsingMode::OnComponentsUpdated);
 
@@ -172,6 +175,34 @@ FActorBrowsingMode::FActorBrowsingMode(SSceneOutliner* InSceneOutliner, TWeakObj
 			return true;
 		}), FSceneOutlinerFilter::EDefaultBehaviour::Pass));
 
+	FSceneOutlinerFilterInfo OnlyCurrentContentBundleInfo(LOCTEXT("ToggleShowOnlyCurrentContentBundle", "Only in Current Content Bundle"), LOCTEXT("ToggleShowOnlyCurrentContentBundleToolTip", "When enabled, only shows Actors that are in the Current Content Bundle."), LocalSettings.bShowOnlyActorsInCurrentContentBundle, FCreateSceneOutlinerFilter::CreateRaw(this, &FActorBrowsingMode::CreateIsInCurrentContentBundleFilter));
+	OnlyCurrentContentBundleInfo.OnToggle().AddLambda([this](bool bIsActive)
+		{
+			FActorBrowsingModeConfig* Settings = GetMutableConfig();
+			if (Settings)
+			{
+				Settings->bShowOnlyActorsInCurrentContentBundle = bIsActive;
+				SaveConfig();
+			}
+		});
+	FilterInfoMap.Add(TEXT("ShowOnlyCurrentContentBundle"), OnlyCurrentContentBundleInfo);
+
+	// Add a filter for unloaded actors to properly reflect the bShowOnlyActorsInCurrentContentBundle flag.
+	SceneOutliner->AddFilter(MakeShared<FActorDescFilter>(FActorDescTreeItem::FFilterPredicate::CreateLambda([this](const FWorldPartitionActorDesc* ActorDesc)
+		{
+			FActorBrowsingModeConfig* Settings = GetMutableConfig();
+			if (Settings && Settings->bShowOnlyActorsInCurrentContentBundle)
+			{
+				if (!WorldPartitionEditorModule || !WorldPartitionEditorModule->IsEditingContentBundle())
+				{
+					return true;
+				}
+
+				return ActorDesc->GetContentBundleGuid().IsValid() && WorldPartitionEditorModule->IsEditingContentBundle(ActorDesc->GetContentBundleGuid());
+			}
+			return true;
+		}), FSceneOutlinerFilter::EDefaultBehaviour::Pass));
+
 	bHideComponents = LocalSettings.bHideActorComponents;
 	FSceneOutlinerFilterInfo HideComponentsInfo(LOCTEXT("ToggleHideActorComponents", "Hide Actor Components"), LOCTEXT("ToggleHideActorComponentsToolTip", "When enabled, hides components belonging to actors."), LocalSettings.bHideActorComponents, FCreateSceneOutlinerFilter::CreateStatic(&FActorBrowsingMode::CreateHideComponentsFilter));
 	HideComponentsInfo.OnToggle().AddLambda([this](bool bIsActive)
@@ -494,6 +525,18 @@ TSharedRef<FSceneOutlinerFilter> FActorBrowsingMode::CreateIsInCurrentDataLayers
 		}), FSceneOutlinerFilter::EDefaultBehaviour::Pass));
 }
 
+TSharedRef<FSceneOutlinerFilter> FActorBrowsingMode::CreateIsInCurrentContentBundleFilter()
+{
+	return MakeShareable(new FActorFilter(FActorTreeItem::FFilterPredicate::CreateLambda([this](const AActor* InActor)
+		{
+			if (!WorldPartitionEditorModule || !WorldPartitionEditorModule->IsEditingContentBundle())
+			{
+				return true;
+			}
+			return InActor->GetContentBundleGuid().IsValid() && WorldPartitionEditorModule->IsEditingContentBundle(InActor->GetContentBundleGuid());
+		}), FSceneOutlinerFilter::EDefaultBehaviour::Pass));
+}
+
 TSharedRef<FSceneOutlinerFilter> FActorBrowsingMode::CreateHideComponentsFilter()
 {
 	return MakeShared<TSceneOutlinerPredicateFilter<FComponentTreeItem>>(TSceneOutlinerPredicateFilter<FComponentTreeItem>(
diff --git a/Engine/Source/Editor/SceneOutliner/Private/ActorBrowsingModeSettings.h b/Engine/Source/Editor/SceneOutliner/Private/ActorBrowsingModeSettings.h
--- a/Engine/Source/Editor/SceneOutliner/Private/ActorBrowsingModeSettings.h
+++ b/Engine/Source/Editor/SceneOutliner/Private/ActorBrowsingModeSettings.h
@@ -31,6 +31,10 @@ public:
 	UPROPERTY()
 	bool bShowOnlyActorsInCurrentDataLayers = false;
 
+	/** True when the Scene Outliner is showing only Actors that exist in the current content bundle */
+	UPROPERTY()
+	bool bShowOnlyActorsInCurrentContentBundle = false;
+
 	/** True when the Scene Outliner is only displaying selected Actors */
 	UPROPERTY()
 	bool bShowOnlySelectedActors = false;
diff --git a/Engine/Source/Editor/SceneOutliner/Public/ActorBrowsingMode.h b/Engine/Source/Editor/SceneOutliner/Public/ActorBrowsingMode.h
--- a/Engine/Source/Editor/SceneOutliner/Public/ActorBrowsingMode.h
+++ b/Engine/Source/Editor/SceneOutliner/Public/ActorBrowsingMode.h
@@ -5,6 +5,7 @@
 #include "ActorModeInteractive.h"
 
 class UActorBrowsingModeSettings;
+class IWorldPartitionEditorModule;
 
 class FActorBrowsingMode : public FActorModeInteractive
 {
@@ -129,6 +130,7 @@ private:
 	static TSharedRef<FSceneOutlinerFilter> CreateHideLevelInstancesFilter();
 	static TSharedRef<FSceneOutlinerFilter> CreateHideUnloadedActorsFilter();
 	static TSharedRef<FSceneOutlinerFilter> CreateHideEmptyFoldersFilter();
+	TSharedRef<FSceneOutlinerFilter> CreateIsInCurrentContentBundleFilter();
 
 	/** Functions to expose selection framing to the UI */
 	void OnToggleAlwaysFrameSelection();
@@ -152,6 +154,8 @@ private:
 	void SaveConfig();
 	
 private:
+
+	IWorldPartitionEditorModule* WorldPartitionEditorModule;
 	/** Number of actors (including unloaded) which have passed through the filters */
 	uint32 FilteredActorCount = 0;
 	/** Number of unloaded actors which have passed through all the filters */
diff --git a/Engine/Source/Editor/WorldPartitionEditor/Private/WorldPartition/ContentBundle/ContentBundleEditorSubsystem.cpp b/Engine/Source/Editor/WorldPartitionEditor/Private/WorldPartition/ContentBundle/ContentBundleEditorSubsystem.cpp
--- a/Engine/Source/Editor/WorldPartitionEditor/Private/WorldPartition/ContentBundle/ContentBundleEditorSubsystem.cpp
+++ b/Engine/Source/Editor/WorldPartitionEditor/Private/WorldPartition/ContentBundle/ContentBundleEditorSubsystem.cpp
@@ -187,6 +187,12 @@ void UContentBundleEditionSubmodule::PopContentBundleEditing()
 	}
 }
 
+bool UContentBundleEditionSubmodule::IsEditingContentBundle(const FGuid& ContentBundleGuid) const
+{
+	TSharedPtr<FContentBundleEditor> ContentBundleEditorPin = EditingContentBundle.Pin();
+	return ContentBundleEditorPin && (ContentBundleEditorPin->GetDescriptor()->GetGuid() == ContentBundleGuid);
+}
+
 UContentBundleEditorSubsystem::UContentBundleEditorSubsystem()
 {
 	
@@ -264,6 +270,16 @@ TSharedPtr<FContentBundleEditor> UContentBundleEditorSubsystem::GetEditorContent
 	return nullptr;
 }
 
+bool UContentBundleEditorSubsystem::IsEditingContentBundle() const
+{
+	return ContentBundleEditionSubModule->IsEditingContentBundle();
+}
+
+bool UContentBundleEditorSubsystem::IsEditingContentBundle(const FGuid& ContentBundleGuid) const
+{
+	return ContentBundleEditionSubModule->IsEditingContentBundle(ContentBundleGuid);
+}
+
 bool UContentBundleEditorSubsystem::HasContentBundle(const UContentBundleDescriptor* ContentBundleDescriptor) const
 {
 	return GetEditorContentBundle(ContentBundleDescriptor) != nullptr;
diff --git a/Engine/Source/Editor/WorldPartitionEditor/Private/WorldPartitionEditorModule.cpp b/Engine/Source/Editor/WorldPartitionEditor/Private/WorldPartitionEditorModule.cpp
--- a/Engine/Source/Editor/WorldPartitionEditor/Private/WorldPartitionEditorModule.cpp
+++ b/Engine/Source/Editor/WorldPartitionEditor/Private/WorldPartitionEditorModule.cpp
@@ -223,6 +223,18 @@ TSharedRef<SWidget> FWorldPartitionEditorModule::CreateContentBundleBrowser()
 	return NewDataLayerBrowser;
 }
 
+bool FWorldPartitionEditorModule::IsEditingContentBundle() const
+{
+	UContentBundleEditorSubsystem* ContentBundleEditorSubsystem = UContentBundleEditorSubsystem::Get();
+	return ContentBundleEditorSubsystem && ContentBundleEditorSubsystem->IsEditingContentBundle();
+}
+
+bool FWorldPartitionEditorModule::IsEditingContentBundle(const FGuid& ContentBundleGuid) const
+{
+	UContentBundleEditorSubsystem* ContentBundleEditorSubsystem = UContentBundleEditorSubsystem::Get();
+	return ContentBundleEditorSubsystem && ContentBundleEditorSubsystem->IsEditingContentBundle(ContentBundleGuid);
+}
+
 int32 FWorldPartitionEditorModule::GetPlacementGridSize() const
 {
 	// Currently shares setting with Foliage. Can be changed when exposed.
diff --git a/Engine/Source/Editor/WorldPartitionEditor/Public/WorldPartition/ContentBundle/ContentBundleEditorSubsystem.h b/Engine/Source/Editor/WorldPartitionEditor/Public/WorldPartition/ContentBundle/ContentBundleEditorSubsystem.h
--- a/Engine/Source/Editor/WorldPartitionEditor/Public/WorldPartition/ContentBundle/ContentBundleEditorSubsystem.h
+++ b/Engine/Source/Editor/WorldPartitionEditor/Public/WorldPartition/ContentBundle/ContentBundleEditorSubsystem.h
@@ -58,6 +58,7 @@ public:
 	bool DeactivateCurrentContentBundleEditing();
 	bool IsEditingContentBundle() const { return EditingContentBundle.IsValid(); }
 	bool IsEditingContentBundle(const FContentBundleEditor* ContentBundleEditor) const { return EditingContentBundle.Pin().Get() == ContentBundleEditor; }
+	bool IsEditingContentBundle(const FGuid& ContentBundleGuid) const;
 
 protected:
 	//~ Begin UContentBundleEditorSubsystemModule interface
@@ -116,6 +117,8 @@ public:
 	void ReferenceAllActors(FContentBundleEditor& EditorContentBundle);
 	void UnreferenceAllActors(FContentBundleEditor& EditorContentBundle);
 
+	bool IsEditingContentBundle() const;
+	bool IsEditingContentBundle(const FGuid& ContentBundleGuid) const;
 	bool ActivateContentBundleEditing(TSharedPtr<FContentBundleEditor>& ContentBundleEditor) const;
 	bool DeactivateContentBundleEditing(TSharedPtr<FContentBundleEditor>& ContentBundleEditor) const;
 	bool IsContentBundleEditingActivated(TSharedPtr<FContentBundleEditor>& ContentBundleEditor) const;
diff --git a/Engine/Source/Editor/WorldPartitionEditor/Public/WorldPartitionEditorModule.h b/Engine/Source/Editor/WorldPartitionEditor/Public/WorldPartitionEditorModule.h
--- a/Engine/Source/Editor/WorldPartitionEditor/Public/WorldPartitionEditorModule.h
+++ b/Engine/Source/Editor/WorldPartitionEditor/Public/WorldPartitionEditorModule.h
@@ -86,6 +86,16 @@ public:
 	 */
 	TSharedRef<class SWidget> CreateContentBundleBrowser();
 
+	/**
+	 * Returns if there's a content bundle in editing mode.
+	 */
+	bool IsEditingContentBundle() const;
+
+	/**
+	 * Returns if the content bundle is in editing mode.
+	 */
+	bool IsEditingContentBundle(const FGuid& ContentBundleGuid) const;
+
 private:
 	/** Register menus */
 	void RegisterMenus();
diff --git a/Engine/Source/Runtime/Engine/Private/WorldPartition/ContentBundle/ContentBundleContainer.cpp b/Engine/Source/Runtime/Engine/Private/WorldPartition/ContentBundle/ContentBundleContainer.cpp
--- a/Engine/Source/Runtime/Engine/Private/WorldPartition/ContentBundle/ContentBundleContainer.cpp
+++ b/Engine/Source/Runtime/Engine/Private/WorldPartition/ContentBundle/ContentBundleContainer.cpp
@@ -124,6 +124,19 @@ TSharedPtr<FContentBundleEditor> FContentBundleContainer::GetEditorContentBundle
 	return nullptr;
 }
 
+TSharedPtr<FContentBundleEditor> FContentBundleContainer::GetEditorContentBundle(const FGuid& ContentBundleGuid) const
+{
+	for (const TSharedPtr<FContentBundleEditor>& ContentBundle : GetEditorContentBundles())
+	{
+		if (ContentBundle->GetDescriptor()->GetGuid() == ContentBundleGuid)
+		{
+			return ContentBundle;
+		}
+	}
+
+	return nullptr;
+}
+
 const TArray<TSharedPtr<FContentBundleEditor>>& FContentBundleContainer::GetEditorContentBundles() const
 {
 	check(UseEditorContentBundle());
diff --git a/Engine/Source/Runtime/Engine/Private/WorldPartition/ContentBundle/ContentBundleWorldSubsystem.cpp b/Engine/Source/Runtime/Engine/Private/WorldPartition/ContentBundle/ContentBundleWorldSubsystem.cpp
--- a/Engine/Source/Runtime/Engine/Private/WorldPartition/ContentBundle/ContentBundleWorldSubsystem.cpp
+++ b/Engine/Source/Runtime/Engine/Private/WorldPartition/ContentBundle/ContentBundleWorldSubsystem.cpp
@@ -146,6 +146,20 @@ TSharedPtr<FContentBundleEditor> UContentBundleManager::GetEditorContentBundle(c
 	return nullptr;
 }
 
+TSharedPtr<FContentBundleEditor> UContentBundleManager::GetEditorContentBundle(const FGuid& ContentBundleGuid) const
+{
+	for (const TUniquePtr<FContentBundleContainer>& ContentBundleContainer : ContentBundleContainers)
+	{
+		TSharedPtr<FContentBundleEditor> ContentBundleEditor = ContentBundleContainer->GetEditorContentBundle(ContentBundleGuid);
+		if (ContentBundleEditor)
+		{
+			return ContentBundleEditor;
+		}
+	}
+
+	return nullptr;
+}
+
 #endif
 
 uint32 UContentBundleManager::GetContentBundleContainerIndex(const UWorld* InjectedWorld) const
diff --git a/Engine/Source/Runtime/Engine/Public/WorldPartition/ContentBundle/ContentBundleContainer.h b/Engine/Source/Runtime/Engine/Public/WorldPartition/ContentBundle/ContentBundleContainer.h
--- a/Engine/Source/Runtime/Engine/Public/WorldPartition/ContentBundle/ContentBundleContainer.h
+++ b/Engine/Source/Runtime/Engine/Public/WorldPartition/ContentBundle/ContentBundleContainer.h
@@ -50,6 +50,7 @@ public:
 
 #if WITH_EDITOR
 	TSharedPtr<FContentBundleEditor> GetEditorContentBundle(const UContentBundleDescriptor* Descriptor) const;
+	TSharedPtr<FContentBundleEditor> GetEditorContentBundle(const FGuid& ContentBundleGuid) const;
 
 	const TArray<TSharedPtr<FContentBundleEditor>>& GetEditorContentBundles() const;
 	TArray<TSharedPtr<FContentBundleEditor>>& GetEditorContentBundles();
diff --git a/Engine/Source/Runtime/Engine/Public/WorldPartition/ContentBundle/ContentBundleWorldSubsystem.h b/Engine/Source/Runtime/Engine/Public/WorldPartition/ContentBundle/ContentBundleWorldSubsystem.h
--- a/Engine/Source/Runtime/Engine/Public/WorldPartition/ContentBundle/ContentBundleWorldSubsystem.h
+++ b/Engine/Source/Runtime/Engine/Public/WorldPartition/ContentBundle/ContentBundleWorldSubsystem.h
@@ -37,6 +37,7 @@ public:
 #if WITH_EDITOR
 	bool GetEditorContentBundle(TArray<TSharedPtr<FContentBundleEditor>>& OutContentBundles);
 	TSharedPtr<FContentBundleEditor> GetEditorContentBundle(const UContentBundleDescriptor* Descriptor, const UWorld* ContentBundleWorld) const;
+	TSharedPtr<FContentBundleEditor> GetEditorContentBundle(const FGuid& ContentBundleGuid) const;
 
 	UContentBundleDuplicateForPIEHelper* GetPIEDuplicateHelper() const { return PIEDuplicateHelper; }
 #endif
diff --git a/Engine/Source/Runtime/Engine/Public/WorldPartition/IWorldPartitionEditorModule.h b/Engine/Source/Runtime/Engine/Public/WorldPartition/IWorldPartitionEditorModule.h
--- a/Engine/Source/Runtime/Engine/Public/WorldPartition/IWorldPartitionEditorModule.h
+++ b/Engine/Source/Runtime/Engine/Public/WorldPartition/IWorldPartitionEditorModule.h
@@ -44,6 +44,9 @@ public:
 	virtual bool GetDisableBugIt() const = 0;
 	virtual void SetDisableBugIt(bool bInDisableBugIt) = 0;
 
+	virtual bool IsEditingContentBundle() const = 0;
+	virtual bool IsEditingContentBundle(const FGuid& ContentBundleGuid) const = 0;
+
 	/** Triggered when a world is added. */
 	DECLARE_EVENT_OneParam(IWorldPartitionEditorModule, FWorldPartitionCreated, UWorld*);
 
