Actor Pin/Unpin cleanup:
- Make sure context menu / WorldOutliner respect same rules
- Hide pin for actors that don't support it
- Show pin for Content Bundle actors

#rnx

diff --git a/Engine/Source/Editor/SceneOutliner/Private/ActorBrowsingMode.cpp b/Engine/Source/Editor/SceneOutliner/Private/ActorBrowsingMode.cpp
--- a/Engine/Source/Editor/SceneOutliner/Private/ActorBrowsingMode.cpp
+++ b/Engine/Source/Editor/SceneOutliner/Private/ActorBrowsingMode.cpp
@@ -588,51 +588,55 @@ void FActorBrowsingMode::RegisterContextMenu()
 		UToolMenu* Menu = ToolMenus->RegisterMenu(DefaultContextBaseMenuName);
 
 		Menu->AddDynamicSection("DynamicHierarchySection", FNewToolMenuDelegate::CreateLambda([](UToolMenu* InMenu)
+		{
+			if(USceneOutlinerMenuContext* Context = InMenu->FindContext<USceneOutlinerMenuContext>())
 			{
-				USceneOutlinerMenuContext* Context = InMenu->FindContext<USceneOutlinerMenuContext>();
-				if (!Context || !Context->SceneOutliner.IsValid())
+				if (TSharedPtr<SSceneOutliner> SharedOutliner = Context->SceneOutliner.Pin())
 				{
-					return;
-				}
-
-				// NOTE: the name "Section" is used in many other places
-				FToolMenuSection& Section = InMenu->FindOrAddSection("Section");
-				Section.Label = LOCTEXT("HierarchySectionName", "Hierarchy");
+					SSceneOutliner* SceneOutliner = SharedOutliner.Get();
+					// NOTE: the name "Section" is used in many other places
+					FToolMenuSection& Section = InMenu->FindOrAddSection("Section");
+					Section.Label = LOCTEXT("HierarchySectionName", "Hierarchy");
 
-				SSceneOutliner* SceneOutliner = Context->SceneOutliner.Pin().Get();
-				if (Context->bShowParentTree)
-				{
-					if (Context->NumSelectedItems == 0)
+					if (Context->bShowParentTree)
 					{
-						FSceneOutlinerMenuHelper::AddMenuEntryCreateFolder(Section, *SceneOutliner);
-					}
-					else
-					{
-						if (Context->NumSelectedItems == 1)
+						if (Context->NumSelectedItems == 0)
 						{
-							SceneOutliner->GetTree().GetSelectedItems()[0]->GenerateContextMenu(InMenu, *SceneOutliner);
+							FSceneOutlinerMenuHelper::AddMenuEntryCreateFolder(Section, *SceneOutliner);
 						}
-
-						// If we've only got folders selected, show the selection and edit sub menus
-						if (Context->NumSelectedItems > 0 && Context->NumSelectedFolders == Context->NumSelectedItems)
+						else
 						{
-							Section.AddSubMenu(
-								"SelectSubMenu",
-								LOCTEXT("SelectSubmenu", "Select"),
-								LOCTEXT("SelectSubmenu_Tooltip", "Select the contents of the current selection"),
-								FNewToolMenuDelegate::CreateSP(SceneOutliner, &SSceneOutliner::FillSelectionSubMenu));
+							if (Context->NumSelectedItems == 1)
+							{
+								SceneOutliner->GetTree().GetSelectedItems()[0]->GenerateContextMenu(InMenu, *SceneOutliner);
+							}
+
+							// If we've only got folders selected, show the selection and edit sub menus
+							if (Context->NumSelectedItems > 0 && Context->NumSelectedFolders == Context->NumSelectedItems)
+							{
+								Section.AddSubMenu(
+									"SelectSubMenu",
+									LOCTEXT("SelectSubmenu", "Select"),
+									LOCTEXT("SelectSubmenu_Tooltip", "Select the contents of the current selection"),
+									FNewToolMenuDelegate::CreateSP(SceneOutliner, &SSceneOutliner::FillSelectionSubMenu));
+							}
 						}
 					}
 				}
-			}));
+			}
+		}));
 
 		Menu->AddDynamicSection("DynamicMainSection", FNewToolMenuDelegate::CreateLambda([](UToolMenu* InMenu)
-			{
-				// We always create a section here, even if there is no parent so that clients can still extend the menu
-				FToolMenuSection& Section = InMenu->AddSection("MainSection", LOCTEXT("OutlinerSectionName", "Outliner"));
+		{
+			// We always create a section here, even if there is no parent so that clients can still extend the menu
+			FToolMenuSection& Section = InMenu->AddSection("MainSection", LOCTEXT("OutlinerSectionName", "Outliner"));
 
-				if (USceneOutlinerMenuContext* Context = InMenu->FindContext<USceneOutlinerMenuContext>())
+			if (USceneOutlinerMenuContext* Context = InMenu->FindContext<USceneOutlinerMenuContext>())
+			{
+				if (TSharedPtr<SSceneOutliner> SharedOutliner = Context->SceneOutliner.Pin())
 				{
+					SSceneOutliner* SceneOutliner = SharedOutliner.Get();
+
 					// Don't add any of these menu items if we're not showing the parent tree
 					// Can't move worlds or level blueprints
 					if (Context->bShowParentTree && Context->NumSelectedItems > 0 && Context->NumWorldsSelected == 0 && Context->SceneOutliner.IsValid())
@@ -649,47 +653,58 @@ void FActorBrowsingMode::RegisterContextMenu()
 						// Only add the menu option to wp levels
 						if (!Context->bRepresentingGameWorld && Context->bRepresentingPartitionedWorld)
 						{
-							// If selection contains some unpinned items, show the pin option
-							// If the selection contains folders, always show the pin option
-							if (Context->NumPinnedItems != Context->NumSelectedItems || Context->NumSelectedFolders > 0)
-							{
-								Section.AddMenuEntry(
-									"PinItems",
-									LOCTEXT("Pin", "Pin"),
-									LOCTEXT("PinTooltip", "Keep the selected items loaded in the editor even when they don't overlap a loaded World Partition region"),
-									FSlateIcon(),
-									FUIAction(FExecuteAction::CreateSP(Context->SceneOutliner.Pin().Get(), &SSceneOutliner::PinSelectedItems)));
-							}
-
-							// If the selection contains some pinned items, show the unpin option
-							// If the selection contains folders, always show the unpin option
-							if (Context->NumPinnedItems != 0 || Context->NumSelectedFolders > 0)
-							{
-								Section.AddMenuEntry(
-									"UnpinItems",
-									LOCTEXT("Unpin", "Unpin"),
-									LOCTEXT("UnpinTooltip", "Allow the World Partition system to load and unload the selected items automatically"),
-									FSlateIcon(),
-									FUIAction(FExecuteAction::CreateSP(Context->SceneOutliner.Pin().Get(), &SSceneOutliner::UnpinSelectedItems)));
-							}
+							Section.AddMenuEntry(
+								"PinItems",
+								LOCTEXT("Pin", "Pin"),
+								LOCTEXT("PinTooltip", "Keep the selected items loaded in the editor even when they don't overlap a loaded World Partition region"),
+								FSlateIcon(),
+								FUIAction(
+									FExecuteAction::CreateSP(SceneOutliner, &SSceneOutliner::PinSelectedItems),
+									FCanExecuteAction::CreateLambda([SceneOutliner, Context]()
+									{
+										if(Context->NumPinnedItems != Context->NumSelectedItems || Context->NumSelectedFolders > 0)
+										{
+											return SceneOutliner->CanPinSelectedItems();
+										}
+										return false;
+									})));
+
+							Section.AddMenuEntry(
+								"UnpinItems",
+								LOCTEXT("Unpin", "Unpin"),
+								LOCTEXT("UnpinTooltip", "Allow the World Partition system to load and unload the selected items automatically"),
+								FSlateIcon(),
+								FUIAction(
+									FExecuteAction::CreateSP(SceneOutliner, &SSceneOutliner::UnpinSelectedItems),
+									FCanExecuteAction::CreateLambda([SceneOutliner, Context]()
+									{
+										if (Context->NumPinnedItems != 0 || Context->NumSelectedFolders > 0)
+										{
+											return SceneOutliner->CanUnpinSelectedItems();
+										}
+										return false;
+									})));
+							
 						}
 					}
 
-					if (Context->NumSelectedItems > 0 && Context->SceneOutliner.IsValid())
+					if (Context->NumSelectedItems > 0)
 					{
-						SSceneOutliner* SceneOutliner = Context->SceneOutliner.Pin().Get();
 						SceneOutliner->AddSourceControlMenuOptions(InMenu);
 					}
 				}
-			}));
+			}
+		}));
 
 		Menu->AddDynamicSection("DynamicActorEditorContext", FNewToolMenuDelegate::CreateLambda([](UToolMenu* InMenu)
+		{
+			USceneOutlinerMenuContext* Context = InMenu->FindContext<USceneOutlinerMenuContext>();
+			if (Context && Context->bShowParentTree)
 			{
-				USceneOutlinerMenuContext* Context = InMenu->FindContext<USceneOutlinerMenuContext>();
-				if (Context && Context->SceneOutliner.IsValid() && Context->bShowParentTree)
+				if (TSharedPtr<SSceneOutliner> SharedOutliner = Context->SceneOutliner.Pin())
 				{
 					FToolMenuSection& Section = InMenu->AddSection("ActorEditorContextSection", LOCTEXT("ActorEditorContextSectionName", "Actor Editor Context"));
-					SSceneOutliner* SceneOutliner = Context->SceneOutliner.Pin().Get();
+					SSceneOutliner* SceneOutliner = SharedOutliner.Get();
 
 					if ((Context->NumSelectedItems == 1) && (Context->NumSelectedFolders == 1))
 					{
@@ -726,8 +741,8 @@ void FActorBrowsingMode::RegisterContextMenu()
 										if (FActorFolderTreeItem* FolderItem = Item->CastTo<FActorFolderTreeItem>())
 										{
 											return FolderItem->World.IsValid() &&
-												  (FolderItem->World->GetCurrentLevel() == FolderItem->GetFolder().GetRootObjectAssociatedLevel()) &&
-												  (FActorFolders::Get().GetActorEditorContextFolder(*FolderItem->World) != FolderItem->GetFolder());
+												(FolderItem->World->GetCurrentLevel() == FolderItem->GetFolder().GetRootObjectAssociatedLevel()) &&
+												(FActorFolders::Get().GetActorEditorContextFolder(*FolderItem->World) != FolderItem->GetFolder());
 										}
 									}
 									return false;
@@ -760,7 +775,8 @@ void FActorBrowsingMode::RegisterContextMenu()
 						)
 					);
 				}
-			}));
+			}
+		}));
 	}
 
 	if (!ToolMenus->IsMenuRegistered(DefaultContextMenuName))
@@ -2151,6 +2167,19 @@ void FActorBrowsingMode::SelectFoldersDescendants(const TArray<FFolderTreeItem*>
 	GEditor->NoteSelectionChange();
 }
 
+bool FActorBrowsingMode::CanPinItems(const TArray<FSceneOutlinerTreeItemPtr>& InItems) const
+{
+	for (const FSceneOutlinerTreeItemPtr& Item : InItems)
+	{
+		if (Item->ShouldShowPinnedState())
+		{
+			return true;
+		}
+	}
+
+	return false;
+}
+
 void FActorBrowsingMode::PinItems(const TArray<FSceneOutlinerTreeItemPtr>& InItems)
 {
 	UWorldPartition* const WorldPartition = RepresentingWorld->GetWorldPartition();
@@ -2193,6 +2222,19 @@ void FActorBrowsingMode::PinItems(const TArray<FSceneOutlinerTreeItemPtr>& InIte
 	}
 }
 
+bool FActorBrowsingMode::CanUnpinItems(const TArray<FSceneOutlinerTreeItemPtr>& InItems) const
+{
+	for (const FSceneOutlinerTreeItemPtr& Item : InItems)
+	{
+		if (Item->ShouldShowPinnedState())
+		{
+			return true;
+		}
+	}
+
+	return false;
+}
+
 void FActorBrowsingMode::UnpinItems(const TArray<FSceneOutlinerTreeItemPtr>& InItems)
 {
 	UWorldPartition* const WorldPartition = RepresentingWorld->GetWorldPartition();
@@ -2228,36 +2270,6 @@ void FActorBrowsingMode::UnpinItems(const TArray<FSceneOutlinerTreeItemPtr>& InI
 	}
 }
 
-void FActorBrowsingMode::PinSelectedItems()
-{
-	const FSceneOutlinerItemSelection Selection = SceneOutliner->GetSelection();
-	if (Selection.Num())
-	{
-		TArray<FSceneOutlinerTreeItemPtr> ItemsToPin;
-		Selection.ForEachItem([this, &ItemsToPin](const FSceneOutlinerTreeItemPtr& TreeItem)
-		{
-			ItemsToPin.Add(TreeItem);
-			return true;
-		});
-		PinItems(ItemsToPin);
-	}
-}
-
-void FActorBrowsingMode::UnpinSelectedItems()
-{
-	const FSceneOutlinerItemSelection Selection = SceneOutliner->GetSelection();
-	if(Selection.Num())
-	{
-		TArray<FSceneOutlinerTreeItemPtr> ItemsToUnpin;
-		Selection.ForEachItem([this, &ItemsToUnpin](const FSceneOutlinerTreeItemPtr& TreeItem)
-		{
-			ItemsToUnpin.Add(TreeItem);
-			return true;
-		});
-		UnpinItems(ItemsToUnpin);
-	}
-}
-
 void FActorBrowsingMode::SynchronizeSelection()
 {
 	FActorModeInteractive::SynchronizeSelection();
diff --git a/Engine/Source/Editor/SceneOutliner/Private/ActorDescTreeItem.cpp b/Engine/Source/Editor/SceneOutliner/Private/ActorDescTreeItem.cpp
--- a/Engine/Source/Editor/SceneOutliner/Private/ActorDescTreeItem.cpp
+++ b/Engine/Source/Editor/SceneOutliner/Private/ActorDescTreeItem.cpp
@@ -124,7 +124,7 @@ private:
 				UActorDescContainer* ActorDescContainer = ActorDesc->GetContainer();
 				UWorld* World = ActorDescContainer != nullptr ? ActorDescContainer->GetWorld() : nullptr;
 				UWorldPartition* WorldPartition = World != nullptr ? World->GetWorldPartition() : nullptr;
-				if (WorldPartition && WorldPartition->IsActorPinned(ActorDesc->GetGuid()))
+				if (WorldPartition && (!ActorDesc->GetIsSpatiallyLoaded() || WorldPartition->IsActorPinned(ActorDesc->GetGuid())))
 				{
 					Args.Add(TEXT("UnloadState"), LOCTEXT("UnloadedDataLayer", "(Unloaded DataLayer)"));
 				}
@@ -324,7 +324,27 @@ bool FActorDescTreeItem::ShouldShowPinnedState() const
 		// Pinning of ActorDescs is only supported on the main world partition
 		if (UActorDescContainer* Container = ActorDescHandle->GetContainer())
 		{
-			return Container->IsMainPartitionContainer();
+			if (FWorldPartitionActorDesc const* ActorDesc = ActorDescHandle.Get())
+			{
+				if (!ActorDesc->GetIsSpatiallyLoaded())
+				{
+					return false;
+				}
+			}
+
+			if (Container->IsMainPartitionContainer())
+			{
+				return true;
+			}
+			else if (ActorDescHandle->GetContentBundleGuid().IsValid())
+			{
+				const UWorld* ContainerWorld = Container->GetWorld();
+				const UWorldPartition* ContainerWorldPartition = ContainerWorld ? ContainerWorld->GetWorldPartition() : nullptr;
+				if (ContainerWorldPartition && ContainerWorldPartition->IsMainWorldPartition())
+				{
+					return ActorDescHandle->GetActorSoftPath().GetAssetPath().GetPackageName() == ContainerWorld->GetPackage()->GetFName();
+				}
+			}
 		}
 	}
 
diff --git a/Engine/Source/Editor/SceneOutliner/Private/ActorTreeItem.cpp b/Engine/Source/Editor/SceneOutliner/Private/ActorTreeItem.cpp
--- a/Engine/Source/Editor/SceneOutliner/Private/ActorTreeItem.cpp
+++ b/Engine/Source/Editor/SceneOutliner/Private/ActorTreeItem.cpp
@@ -367,7 +367,7 @@ bool FActorTreeItem::ShouldShowPinnedState() const
 		// Pinning of Actors is only supported on the main world partition
 		const ULevel* Level = ActorPtr->GetLevel();
 		const UWorld* World = Level->GetWorld();
-		return World && !World->IsGameWorld() && !!World->GetWorldPartition() && Level->IsPersistentLevel();
+		return World && !World->IsGameWorld() && !!World->GetWorldPartition() && Level->IsPersistentLevel() && ActorPtr->IsPackageExternal() && ActorPtr->GetIsSpatiallyLoaded();
 	}
 
 	return false;
diff --git a/Engine/Source/Editor/SceneOutliner/Private/SSceneOutliner.cpp b/Engine/Source/Editor/SceneOutliner/Private/SSceneOutliner.cpp
--- a/Engine/Source/Editor/SceneOutliner/Private/SSceneOutliner.cpp
+++ b/Engine/Source/Editor/SceneOutliner/Private/SSceneOutliner.cpp
@@ -2370,14 +2370,42 @@ void SSceneOutliner::UnpinItems(const TArray<FSceneOutlinerTreeItemPtr>& InItems
 	Mode->UnpinItems(InItems);
 }
 
+bool SSceneOutliner::CanPinItems(const TArray<FSceneOutlinerTreeItemPtr>& InItems) const
+{
+	return Mode->CanPinItems(InItems);
+}
+
+bool SSceneOutliner::CanUnpinItems(const TArray<FSceneOutlinerTreeItemPtr>& InItems) const
+{
+	return Mode->CanUnpinItems(InItems);
+}
+
 void SSceneOutliner::PinSelectedItems()
 {
-	Mode->PinSelectedItems();
+	TArray<FSceneOutlinerTreeItemPtr> SelectedItems;
+	GetSelection().Get(SelectedItems);
+	PinItems(SelectedItems);
 }
 
 void SSceneOutliner::UnpinSelectedItems()
 {
-	Mode->UnpinSelectedItems();
+	TArray<FSceneOutlinerTreeItemPtr> SelectedItems;
+	GetSelection().Get(SelectedItems);
+	UnpinItems(SelectedItems);
+}
+
+bool SSceneOutliner::CanPinSelectedItems() const
+{
+	TArray<FSceneOutlinerTreeItemPtr> SelectedItems;
+	GetSelection().Get(SelectedItems);
+	return CanPinItems(SelectedItems);
+}
+
+bool SSceneOutliner::CanUnpinSelectedItems() const
+{
+	TArray<FSceneOutlinerTreeItemPtr> SelectedItems;
+	GetSelection().Get(SelectedItems);
+	return CanUnpinItems(SelectedItems);
 }
 
 FSceneOutlinerTreeItemPtr SSceneOutliner::FindParent(const ISceneOutlinerTreeItem& InItem) const
diff --git a/Engine/Source/Editor/SceneOutliner/Public/ActorBrowsingMode.h b/Engine/Source/Editor/SceneOutliner/Public/ActorBrowsingMode.h
--- a/Engine/Source/Editor/SceneOutliner/Public/ActorBrowsingMode.h
+++ b/Engine/Source/Editor/SceneOutliner/Public/ActorBrowsingMode.h
@@ -56,8 +56,8 @@ public:
 	virtual void SelectFoldersDescendants(const TArray<FFolderTreeItem*>& FolderItems, bool bSelectImmediateChildrenOnly) override;
 	virtual void PinItems(const TArray<FSceneOutlinerTreeItemPtr>& InItems) override;
 	virtual void UnpinItems(const TArray<FSceneOutlinerTreeItemPtr>& InItems) override;
-	virtual void PinSelectedItems() override;
-	virtual void UnpinSelectedItems() override;
+	virtual bool CanPinItems(const TArray<FSceneOutlinerTreeItemPtr>& InItems) const override;
+	virtual bool CanUnpinItems(const TArray<FSceneOutlinerTreeItemPtr>& InItems) const override;
 	virtual void SynchronizeSelection() override;
 
 	/* End ISceneOutlinerMode Interface */
diff --git a/Engine/Source/Editor/SceneOutliner/Public/ISceneOutliner.h b/Engine/Source/Editor/SceneOutliner/Public/ISceneOutliner.h
--- a/Engine/Source/Editor/SceneOutliner/Public/ISceneOutliner.h
+++ b/Engine/Source/Editor/SceneOutliner/Public/ISceneOutliner.h
@@ -139,12 +139,24 @@ public:
 	/**  Unpins an item list in the outliner. */
 	virtual void UnpinItems(const TArray<FSceneOutlinerTreeItemPtr>& InItems) = 0;
 	
+	/** Returns true if any of the items can be pinned. */
+	virtual bool CanPinItems(const TArray<FSceneOutlinerTreeItemPtr>& InItems) const = 0;
+
+	/** Returns true if any of the items can be unpinned. */
+	virtual bool CanUnpinItems(const TArray<FSceneOutlinerTreeItemPtr>& InItems) const = 0;
+
 	/** Pin selected items */
 	virtual void PinSelectedItems() = 0;
 
 	/** Unpins selected items */
 	virtual void UnpinSelectedItems() = 0;
 
+	/** Returns true if any of the selected items can be pinned */
+	virtual bool CanPinSelectedItems() const = 0;
+
+	/** Returns true if any of the selected items can be unpinned */
+	virtual bool CanUnpinSelectedItems() const = 0;
+
 	/** Get the active SceneOutlinerMode */
 	const ISceneOutlinerMode* GetMode() const { return Mode; }
 
diff --git a/Engine/Source/Editor/SceneOutliner/Public/ISceneOutlinerMode.h b/Engine/Source/Editor/SceneOutliner/Public/ISceneOutlinerMode.h
--- a/Engine/Source/Editor/SceneOutliner/Public/ISceneOutlinerMode.h
+++ b/Engine/Source/Editor/SceneOutliner/Public/ISceneOutlinerMode.h
@@ -159,10 +159,10 @@ public:
 	virtual void PinItems(const TArray<FSceneOutlinerTreeItemPtr>& InItems) {}
 	/** Unpins an item list in the outliner */
 	virtual void UnpinItems(const TArray<FSceneOutlinerTreeItemPtr>& InItems) {}
-	/** Pins all selected items */
-	virtual void PinSelectedItems() {}
-	/** Unpins all selected items */
-	virtual void UnpinSelectedItems() {}
+	/** Returns true if any of the items can be pinned. */
+	virtual bool CanPinItems(const TArray<FSceneOutlinerTreeItemPtr>& InItems) const { return false; }
+	/** Returns true if any of the items can be unpinned. */
+	virtual bool CanUnpinItems(const TArray<FSceneOutlinerTreeItemPtr>& InItems) const { return false; }
 
 	/** Function called by the Outliner Filter Bar to compare an item with Type Filters*/
 	virtual bool CompareItemWithClassName(SceneOutliner::FilterBarType InItem, const TSet<FTopLevelAssetPath>&) const { return false; };
diff --git a/Engine/Source/Editor/SceneOutliner/Public/SSceneOutliner.h b/Engine/Source/Editor/SceneOutliner/Public/SSceneOutliner.h
--- a/Engine/Source/Editor/SceneOutliner/Public/SSceneOutliner.h
+++ b/Engine/Source/Editor/SceneOutliner/Public/SSceneOutliner.h
@@ -164,6 +164,17 @@ struct FSceneOutlinerItemSelection
 		}
 	}
 
+	void Get(TArray<FSceneOutlinerTreeItemPtr>& OutArray) const
+	{
+		for (const TWeakPtr<ISceneOutlinerTreeItem>& Item : SelectedItems)
+		{
+			if (const auto ItemPtr = Item.Pin())
+			{
+				OutArray.Add(ItemPtr);
+			}
+		}
+	}
+
 	/** Apply a function to each item of a specified type */
 	template <typename TreeType>
 	void ForEachItem(TFunctionRef<void(TreeType&)> Func) const
@@ -420,6 +431,16 @@ public:
 	 */
 	virtual void UnpinItems(const TArray<FSceneOutlinerTreeItemPtr>& InItems) override;
 
+	/**
+	 * Returns true if any of the items can be pinned.
+	 */
+	virtual bool CanPinItems(const TArray<FSceneOutlinerTreeItemPtr>& InItems) const override;
+
+	/**
+	 * Returns true if any of the items can be unpinned.
+	 */
+	virtual bool CanUnpinItems(const TArray<FSceneOutlinerTreeItemPtr>& InItems) const override;
+
 	/**
 	 * Pin selected items
 	 */
@@ -429,6 +450,16 @@ public:
 	 * Unpins selected items
 	 */
 	virtual void UnpinSelectedItems() override;
+
+	/**
+	 * Returns true if any of the selected items can be pinned
+	 */
+	virtual bool CanPinSelectedItems() const override;
+
+	/**
+	 * Returns true if any of the selected items can be unpinned
+	 */
+	virtual bool CanUnpinSelectedItems() const override;
 	
 	/**
 	 * Returns the parent tree item for a given item if it exists, nullptr otherwise.
