Outliner: Fix an issue where you couldn't delete a custom filter that was active in multiple Outliners

diff --git a/Engine/Source/Developer/ToolWidgets/Public/Filters/SBasicFilterBar.h b/Engine/Source/Developer/ToolWidgets/Public/Filters/SBasicFilterBar.h
--- a/Engine/Source/Developer/ToolWidgets/Public/Filters/SBasicFilterBar.h
+++ b/Engine/Source/Developer/ToolWidgets/Public/Filters/SBasicFilterBar.h
@@ -1047,7 +1047,6 @@ protected:
 		{
 			// Update the frontend filters collection
 			SetFrontendFilterActive(FrontendFilter.ToSharedRef(), false);
-			OnFilterChanged.ExecuteIfBound();
 		}
 
 		RemoveWidgetFromLayout(FilterToRemove);
diff --git a/Engine/Source/Editor/SceneOutliner/Private/SceneOutlinerFilterBar.cpp b/Engine/Source/Editor/SceneOutliner/Private/SceneOutlinerFilterBar.cpp
--- a/Engine/Source/Editor/SceneOutliner/Private/SceneOutlinerFilterBar.cpp
+++ b/Engine/Source/Editor/SceneOutliner/Private/SceneOutlinerFilterBar.cpp
@@ -134,7 +134,10 @@ void SSceneOutlinerFilterBar::LoadSettings()
 	// From the instance settings, get each checked filter and set the checked and active state
 	for(const FCustomTextFilterState& FilterState : InstanceSettings->CustomTextFilters)
 	{
-		RestoreCustomTextFilterState(FilterState);
+		if(!RestoreCustomTextFilterState(FilterState))
+		{
+			UE_LOG(LogSlate, Warning, TEXT("SSceneOutlinerFilterBar was unable to load the following custom text filter: %s"), *FilterState.FilterData.FilterLabel.ToString());
+		}
 	}
 
 	this->OnFilterChanged.ExecuteIfBound();
@@ -201,7 +204,7 @@ void SSceneOutlinerFilterBar::LoadCustomTextFilters(const FFilterBarSettings* Fi
 }
 
 
-void SSceneOutlinerFilterBar::RestoreCustomTextFilterState(const FCustomTextFilterState& InFilterState)
+bool SSceneOutlinerFilterBar::RestoreCustomTextFilterState(const FCustomTextFilterState& InFilterState)
 {
 	// Find the filter associated with the current instance data from our list of custom text filters
 	TSharedRef< ICustomTextFilter<SceneOutliner::FilterBarType> >* Filter =
@@ -210,8 +213,10 @@ void SSceneOutlinerFilterBar::RestoreCustomTextFilterState(const FCustomTextFilt
 		return Element->CreateCustomTextFilterData().FilterLabel.EqualTo(InFilterState.FilterData.FilterLabel);
 	});
 
-	// InFilterState should always be a filter that already exists!
-	check(Filter);
+	if(!Filter)
+	{
+		return false;
+	}
 
 	// Get the actual FFilterBase
 	TSharedRef<FFilterBase<SceneOutliner::FilterBarType>> ActualFilter = Filter->Get().GetFilter().ToSharedRef();
@@ -222,6 +227,8 @@ void SSceneOutlinerFilterBar::RestoreCustomTextFilterState(const FCustomTextFilt
 	// Set the filter as active if it was previously
 	AddedFilter->SetEnabled(InFilterState.bIsActive, false);
 	this->SetFrontendFilterActive(ActualFilter, InFilterState.bIsActive);
+
+	return true;
 }
 
 void SSceneOutlinerFilterBar::OnExternalCustomTextFilterCreated(TSharedPtr<SWidget> BroadcastingFilterBar)
@@ -276,4 +283,7 @@ void SSceneOutlinerFilterBar::OnExternalCustomTextFilterCreated(TSharedPtr<SWidg
 	{
 		RestoreCustomTextFilterState(SavedFilterState);
 	}
+
+	// Save the settings for this instance
+	SaveSettings();
 }
\ No newline at end of file
diff --git a/Engine/Source/Editor/SceneOutliner/Private/SceneOutlinerFilterBar.h b/Engine/Source/Editor/SceneOutliner/Private/SceneOutlinerFilterBar.h
--- a/Engine/Source/Editor/SceneOutliner/Private/SceneOutlinerFilterBar.h
+++ b/Engine/Source/Editor/SceneOutliner/Private/SceneOutlinerFilterBar.h
@@ -89,8 +89,10 @@ private:
 	/** Empty our list of custom text filters, and load from the given config */
 	void LoadCustomTextFilters(const FFilterBarSettings* FilterBarConfig);
 
-	/** Find the custom text filter corresponding to the specified state, and restore it's state to what is specified */
-	void RestoreCustomTextFilterState(const FCustomTextFilterState& InFilterState);
+	/** Find the custom text filter corresponding to the specified state, and restore it's state to what is specified
+	 *  @return True if the filter was restored successfully, false if not
+	 */
+	bool RestoreCustomTextFilterState(const FCustomTextFilterState& InFilterState);
 	
 private:
 	/** The event that executes when a custom text filter is created/modified/deleted in any filter list that is using Shared Settings */
