WorldOutliner: Remove Pin state display from Folders. Still allow Pin/Unpin via context menu (Optimization)

diff --git a/Engine/Source/Editor/SceneOutliner/Private/ActorBrowsingMode.cpp b/Engine/Source/Editor/SceneOutliner/Private/ActorBrowsingMode.cpp
--- a/Engine/Source/Editor/SceneOutliner/Private/ActorBrowsingMode.cpp
+++ b/Engine/Source/Editor/SceneOutliner/Private/ActorBrowsingMode.cpp
@@ -2140,6 +2140,23 @@ namespace ActorBrowsingModeUtils
 			});
 		}
 	};
+
+	bool CanChangePinnedStates(const TArray<FSceneOutlinerTreeItemPtr>& InItems)
+	{
+		for (const FSceneOutlinerTreeItemPtr& Item : InItems)
+		{
+			if (Item->ShouldShowPinnedState())
+			{
+				return true;
+			}
+			else if (const FActorFolderTreeItem* FolderItem = Item->CastTo<FActorFolderTreeItem>(); FolderItem && FolderItem->CanChangeChildrenPinnedState())
+			{
+				return true;
+			}
+		}
+
+		return false;
+	}
 }
 
 void FActorBrowsingMode::SelectFoldersDescendants(const TArray<FFolderTreeItem*>& FolderItems, bool bSelectImmediateChildrenOnly)
@@ -2169,15 +2186,7 @@ void FActorBrowsingMode::SelectFoldersDescendants(const TArray<FFolderTreeItem*>
 
 bool FActorBrowsingMode::CanPinItems(const TArray<FSceneOutlinerTreeItemPtr>& InItems) const
 {
-	for (const FSceneOutlinerTreeItemPtr& Item : InItems)
-	{
-		if (Item->ShouldShowPinnedState())
-		{
-			return true;
-		}
-	}
-
-	return false;
+	return ActorBrowsingModeUtils::CanChangePinnedStates(InItems);
 }
 
 void FActorBrowsingMode::PinItems(const TArray<FSceneOutlinerTreeItemPtr>& InItems)
@@ -2224,15 +2233,7 @@ void FActorBrowsingMode::PinItems(const TArray<FSceneOutlinerTreeItemPtr>& InIte
 
 bool FActorBrowsingMode::CanUnpinItems(const TArray<FSceneOutlinerTreeItemPtr>& InItems) const
 {
-	for (const FSceneOutlinerTreeItemPtr& Item : InItems)
-	{
-		if (Item->ShouldShowPinnedState())
-		{
-			return true;
-		}
-	}
-
-	return false;
+	return ActorBrowsingModeUtils::CanChangePinnedStates(InItems);
 }
 
 void FActorBrowsingMode::UnpinItems(const TArray<FSceneOutlinerTreeItemPtr>& InItems)
diff --git a/Engine/Source/Editor/SceneOutliner/Private/ActorFolderTreeItem.cpp b/Engine/Source/Editor/SceneOutliner/Private/ActorFolderTreeItem.cpp
--- a/Engine/Source/Editor/SceneOutliner/Private/ActorFolderTreeItem.cpp
+++ b/Engine/Source/Editor/SceneOutliner/Private/ActorFolderTreeItem.cpp
@@ -354,7 +354,7 @@ TSharedRef<SWidget> FActorFolderTreeItem::GenerateLabelWidget(ISceneOutliner& Ou
 	return SNew(SActorFolderTreeLabel, *this, Outliner, InRow);
 }
 
-bool FActorFolderTreeItem::ShouldShowPinnedState() const
+bool FActorFolderTreeItem::CanChangeChildrenPinnedState() const
 {
 	if (World.IsValid() && !World->IsGameWorld() && World->IsPartitionedWorld())
 	{
diff --git a/Engine/Source/Editor/SceneOutliner/Public/ActorFolderTreeItem.h b/Engine/Source/Editor/SceneOutliner/Public/ActorFolderTreeItem.h
--- a/Engine/Source/Editor/SceneOutliner/Public/ActorFolderTreeItem.h
+++ b/Engine/Source/Editor/SceneOutliner/Public/ActorFolderTreeItem.h
@@ -22,7 +22,7 @@ public:
 	virtual void Delete(const FFolder& InNewParentFolder) override;
 	virtual bool CanInteract() const override;
 	virtual TSharedRef<SWidget> GenerateLabelWidget(ISceneOutliner& Outliner, const STableRow<FSceneOutlinerTreeItemPtr>& InRow) override;
-	virtual bool ShouldShowPinnedState() const override;
+	virtual bool ShouldShowPinnedState() const override { return false; }
 	virtual bool ShouldShowVisibilityState() const override;
 	virtual FFolder GetFolder() const override;
 	/* End FFolderTreeItem Implementation */
@@ -31,6 +31,7 @@ public:
 	virtual void MoveTo(const FFolder& InNewParentFolder) override;
 	virtual void SetPath(const FName& InNewPath) override;
 	
+	bool CanChangeChildrenPinnedState() const;
 	const UActorFolder* GetActorFolder() const { return ActorFolder.Get(); }
 private:
 	virtual void CreateSubFolder(TWeakPtr<SSceneOutliner> WeakOutliner) override;
