Outliner: Remove an item from the selection when it is removed from the Outliner. This fixes a crash when you try to perform certain actions after ungrouping actors

diff --git a/Engine/Source/Editor/SceneOutliner/Private/SSceneOutliner.cpp b/Engine/Source/Editor/SceneOutliner/Private/SSceneOutliner.cpp
--- a/Engine/Source/Editor/SceneOutliner/Private/SSceneOutliner.cpp
+++ b/Engine/Source/Editor/SceneOutliner/Private/SSceneOutliner.cpp
@@ -887,6 +887,12 @@ void SSceneOutliner::RemoveItemFromTree(FSceneOutlinerTreeItemRef ReferenceItem)
 	{
 		FSceneOutlinerTreeItemRef Item = ItemInTree->ToSharedRef();
 
+		// If the item we are removing is selected, refresh the selection to remove it from there as well
+		if(OutlinerTreeView->GetSelectedItems().Contains(Item))
+		{
+			RefreshSelection();
+		}
+
 		auto Parent = Item->GetParent();
 
 		if (Parent.IsValid())
