Add insight trace to scene outliner creation

#rnx

diff --git a/Engine/Source/Editor/SceneOutliner/Private/SceneOutlinerModule.cpp b/Engine/Source/Editor/SceneOutliner/Private/SceneOutlinerModule.cpp
--- a/Engine/Source/Editor/SceneOutliner/Private/SceneOutlinerModule.cpp
+++ b/Engine/Source/Editor/SceneOutliner/Private/SceneOutlinerModule.cpp
@@ -68,6 +68,8 @@ void FSceneOutlinerModule::ShutdownModule()
 
 TSharedRef<ISceneOutliner> FSceneOutlinerModule::CreateSceneOutliner(const FSceneOutlinerInitializationOptions& InitOptions) const
 {
+	TRACE_CPUPROFILER_EVENT_SCOPE(FSceneOutlinerModule::CreateSceneOutliner);
+
 	return SNew(SSceneOutliner, InitOptions)
 		.IsEnabled(FSlateApplication::Get().GetNormalExecutionAttribute());
 }
