SceneOutliner: Fix Pin not showing for HLOD Actors (RuntimeOnly actors)

#rnx

diff --git a/Engine/Source/Editor/SceneOutliner/Private/ActorDescTreeItem.cpp b/Engine/Source/Editor/SceneOutliner/Private/ActorDescTreeItem.cpp
--- a/Engine/Source/Editor/SceneOutliner/Private/ActorDescTreeItem.cpp
+++ b/Engine/Source/Editor/SceneOutliner/Private/ActorDescTreeItem.cpp
@@ -326,7 +326,8 @@ bool FActorDescTreeItem::ShouldShowPinnedState() const
 		{
 			if (FWorldPartitionActorDesc const* ActorDesc = ActorDescHandle.Get())
 			{
-				if (!ActorDesc->GetIsSpatiallyLoaded())
+				// Only Spatially loaded actors can be pinned with the exception of non spatially loaded, runtime only actors (ex: HLODs)
+				if (!ActorDesc->GetIsSpatiallyLoaded() && !ActorDesc->GetActorIsRuntimeOnly())
 				{
 					return false;
 				}
diff --git a/Engine/Source/Editor/SceneOutliner/Private/ActorTreeItem.cpp b/Engine/Source/Editor/SceneOutliner/Private/ActorTreeItem.cpp
--- a/Engine/Source/Editor/SceneOutliner/Private/ActorTreeItem.cpp
+++ b/Engine/Source/Editor/SceneOutliner/Private/ActorTreeItem.cpp
@@ -367,7 +367,9 @@ bool FActorTreeItem::ShouldShowPinnedState() const
 		// Pinning of Actors is only supported on the main world partition
 		const ULevel* Level = ActorPtr->GetLevel();
 		const UWorld* World = Level->GetWorld();
-		return World && !World->IsGameWorld() && !!World->GetWorldPartition() && Level->IsPersistentLevel() && ActorPtr->IsPackageExternal() && ActorPtr->GetIsSpatiallyLoaded();
+
+		// Only Spatially loaded actors can be pinned with the exception of non spatially loaded, runtime only actors (ex: HLODs)
+		return World && !World->IsGameWorld() && !!World->GetWorldPartition() && Level->IsPersistentLevel() && ActorPtr->IsPackageExternal() && (ActorPtr->GetIsSpatiallyLoaded() || ActorPtr->IsRuntimeOnly());
 	}
 
 	return false;
