World Outliner: scroll unloaded actor into view when clicking on an actor token.

diff --git a/Engine/Source/Editor/SceneOutliner/Private/ActorBrowsingMode.cpp b/Engine/Source/Editor/SceneOutliner/Private/ActorBrowsingMode.cpp
--- a/Engine/Source/Editor/SceneOutliner/Private/ActorBrowsingMode.cpp
+++ b/Engine/Source/Editor/SceneOutliner/Private/ActorBrowsingMode.cpp
@@ -931,6 +931,7 @@ void FActorBrowsingMode::OnSelectUnloadedActors(const TArray<FGuid>& ActorGuids)
 	if (ItemsToSelect.Num())
 	{
 		SceneOutliner->SetItemSelection(ItemsToSelect, true);
+		SceneOutliner->ScrollItemIntoView(ItemsToSelect.Last());
 	}
 }
 
