World Outliner: focus unloaded actor in viewport when clicking on an actor token.

diff --git a/Engine/Source/Editor/SceneOutliner/Private/ActorBrowsingMode.cpp b/Engine/Source/Editor/SceneOutliner/Private/ActorBrowsingMode.cpp
--- a/Engine/Source/Editor/SceneOutliner/Private/ActorBrowsingMode.cpp
+++ b/Engine/Source/Editor/SceneOutliner/Private/ActorBrowsingMode.cpp
@@ -932,6 +932,11 @@ void FActorBrowsingMode::OnSelectUnloadedActors(const TArray<FGuid>& ActorGuids)
 	{
 		SceneOutliner->SetItemSelection(ItemsToSelect, true);
 		SceneOutliner->ScrollItemIntoView(ItemsToSelect.Last());
+
+		if (const FActorDescTreeItem* ActorDescItem = ItemsToSelect.Last()->CastTo<FActorDescTreeItem>())
+		{
+			ActorDescItem->FocusActorBounds();
+		}
 	}
 }
 
