SourceControl: Fixed when using 'Revert Unchanged Only' in the SceneOutliner, the operation ends up reverting all OFPA's, because 'IsModified' is 'false' for all of them because the revert dialog compares against package names - not filenames.

diff --git a/Engine/Source/Editor/SceneOutliner/Private/SceneOutlinerHelpers.cpp b/Engine/Source/Editor/SceneOutliner/Private/SceneOutlinerHelpers.cpp
--- a/Engine/Source/Editor/SceneOutliner/Private/SceneOutlinerHelpers.cpp
+++ b/Engine/Source/Editor/SceneOutliner/Private/SceneOutlinerHelpers.cpp
@@ -5,6 +5,7 @@
 #include "ActorDescTreeItem.h"
 #include "ActorFolderTreeItem.h"
 #include "WorldPartition/WorldPartitionActorDesc.h"
+#include "UObject/Package.h"
 
 #include "SourceControlHelpers.h"
 
@@ -18,7 +19,7 @@ namespace SceneOutliner
 			{
 				if (Actor->IsPackageExternal())
 				{
-					return USourceControlHelpers::PackageFilename(Actor->GetExternalPackage());
+					return Actor->GetExternalPackage()->GetName();
 				}
 
 			}
@@ -29,7 +30,7 @@ namespace SceneOutliner
 			{
 				if (ActorFolder->IsPackageExternal())
 				{
-					return USourceControlHelpers::PackageFilename(ActorFolder->GetExternalPackage());
+					return ActorFolder->GetExternalPackage()->GetName();
 				}
 			}
 		}
@@ -37,7 +38,7 @@ namespace SceneOutliner
 		{
 			if (const FWorldPartitionActorDesc* ActorDesc = ActorDescItem->ActorDescHandle.Get())
 			{
-				return USourceControlHelpers::PackageFilename(ActorDesc->GetActorPackage().ToString());
+				return ActorDesc->GetActorPackage().ToString();
 			}
 		}
 
diff --git a/Engine/Source/Editor/SceneOutliner/Private/SceneOutlinerSCCHandler.cpp b/Engine/Source/Editor/SceneOutliner/Private/SceneOutlinerSCCHandler.cpp
--- a/Engine/Source/Editor/SceneOutliner/Private/SceneOutlinerSCCHandler.cpp
+++ b/Engine/Source/Editor/SceneOutliner/Private/SceneOutlinerSCCHandler.cpp
@@ -279,7 +279,8 @@ void FSceneOutlinerSCCHandler::GetSelectedPackageNames(TArray<FString>& OutPacka
 		}
 
 		FString PackageName = SourceControl->GetPackageName();
-		if (!PackageName.IsEmpty()) {
+		if (!PackageName.IsEmpty())
+		{
 			OutPackageNames.Add(PackageName);
 		}
 	}
@@ -296,7 +297,8 @@ void FSceneOutlinerSCCHandler::GetSelectedPackages(TArray<UPackage*>& OutPackage
 		}
 
 		UPackage* Package = SourceControl->GetPackage();
-		if (Package != nullptr) {
+		if (Package != nullptr)
+		{
 			OutPackages.Add(Package);
 		}
 	}
@@ -368,6 +370,7 @@ void FSceneOutlinerSCCHandler::ExecuteSCCRevert()
 {
 	TArray<FString> PackageNames;
 	GetSelectedPackageNames(PackageNames);
+
 	FSourceControlWindows::PromptForRevert(PackageNames, true);
 }
 
@@ -375,7 +378,8 @@ void FSceneOutlinerSCCHandler::ExecuteSCCHistory()
 {
 	TArray<FString> PackageNames;
 	GetSelectedPackageNames(PackageNames);
-	FSourceControlWindows::DisplayRevisionHistory(SourceControlHelpers::PackageFilenames(PackageNames));
+
+	FSourceControlWindows::DisplayRevisionHistory(PackageNames);
 }
 
 #undef LOCTEXT_NAMESPACE
diff --git a/Engine/Source/Editor/SceneOutliner/Private/SceneOutlinerTreeItemSCC.cpp b/Engine/Source/Editor/SceneOutliner/Private/SceneOutlinerTreeItemSCC.cpp
--- a/Engine/Source/Editor/SceneOutliner/Private/SceneOutlinerTreeItemSCC.cpp
+++ b/Engine/Source/Editor/SceneOutliner/Private/SceneOutlinerTreeItemSCC.cpp
@@ -10,6 +10,7 @@
 #include "ProfilingDebugging/CpuProfilerTrace.h"
 #include "SceneOutlinerHelpers.h"
 #include "WorldPartition/WorldPartitionActorDesc.h"
+#include "UObject/Package.h"
 #include "UncontrolledChangelistsModule.h"
 
 FSceneOutlinerTreeItemSCC::FSceneOutlinerTreeItemSCC(FSceneOutlinerTreeItemPtr InTreeItemPtr)
@@ -19,6 +20,7 @@ FSceneOutlinerTreeItemSCC::FSceneOutlinerTreeItemSCC(FSceneOutlinerTreeItemPtr I
 	if (TreeItemPtr.IsValid())
 	{
 		ExternalPackageName = SceneOutliner::FSceneOutlinerHelpers::GetExternalPackageName(*TreeItemPtr.Get());
+		ExternalPackageFileName = USourceControlHelpers::PackageFilename(ExternalPackageName);
 		ExternalPackage = SceneOutliner::FSceneOutlinerHelpers::GetExternalPackage(*TreeItemPtr.Get());
 		
 		if (FActorTreeItem* ActorItem = TreeItemPtr->CastTo<FActorTreeItem>())
@@ -29,13 +31,15 @@ FSceneOutlinerTreeItemSCC::FSceneOutlinerTreeItemSCC(FSceneOutlinerTreeItemPtr I
 				{
 					if (bExternal)
 					{
-						ExternalPackageName = USourceControlHelpers::PackageFilename(InActor->GetExternalPackage());
+						ExternalPackageName = InActor->GetExternalPackage()->GetName();
+						ExternalPackageFileName = USourceControlHelpers::PackageFilename(ExternalPackageName);
 						ExternalPackage = InActor->GetExternalPackage();
 						ConnectSourceControl();
 					}
 					else
 					{
 						ExternalPackageName = FString();
+						ExternalPackageFileName = FString();
 						ExternalPackage = nullptr;
 						DisconnectSourceControl();
 					}
@@ -43,7 +47,7 @@ FSceneOutlinerTreeItemSCC::FSceneOutlinerTreeItemSCC(FSceneOutlinerTreeItemPtr I
 			}
 		}
 		
-		if (!ExternalPackageName.IsEmpty())
+		if (!ExternalPackageFileName.IsEmpty())
 		{
 			ConnectSourceControl();
 		}
@@ -67,31 +71,31 @@ FSceneOutlinerTreeItemSCC::~FSceneOutlinerTreeItemSCC()
 
 FSourceControlStatePtr FSceneOutlinerTreeItemSCC::GetSourceControlState()
 {
-	return ISourceControlModule::Get().GetProvider().GetState(ExternalPackageName, EStateCacheUsage::Use);
+	return ISourceControlModule::Get().GetProvider().GetState(ExternalPackageFileName, EStateCacheUsage::Use);
 }
 
 FSourceControlStatePtr FSceneOutlinerTreeItemSCC::RefreshSourceControlState()
 {
-	return ISourceControlModule::Get().GetProvider().GetState(ExternalPackageName, EStateCacheUsage::ForceUpdate);
+	return ISourceControlModule::Get().GetProvider().GetState(ExternalPackageFileName, EStateCacheUsage::ForceUpdate);
 }
 
 void FSceneOutlinerTreeItemSCC::ConnectSourceControl()
 {
-	check(!ExternalPackageName.IsEmpty());
+	check(!ExternalPackageFileName.IsEmpty());
 
 	ISourceControlModule& SCCModule = ISourceControlModule::Get();
 	SourceControlProviderChangedDelegateHandle = SCCModule.RegisterProviderChanged(FSourceControlProviderChanged::FDelegate::CreateRaw(this, &FSceneOutlinerTreeItemSCC::HandleSourceControlProviderChanged));
 	SourceControlStateChangedDelegateHandle = SCCModule.GetProvider().RegisterSourceControlStateChanged_Handle(FSourceControlStateChanged::FDelegate::CreateRaw(this, &FSceneOutlinerTreeItemSCC::HandleSourceControlStateChanged, EStateCacheUsage::Use));
 
 	// Check if there is already a cached state for this item
-	FSourceControlStatePtr SourceControlState = ISourceControlModule::Get().GetProvider().GetState(ExternalPackageName, EStateCacheUsage::Use);
+	FSourceControlStatePtr SourceControlState = ISourceControlModule::Get().GetProvider().GetState(ExternalPackageFileName, EStateCacheUsage::Use);
 	if (SourceControlState.IsValid() && !SourceControlState->IsUnknown())
 	{
 		BroadcastNewState(SourceControlState);
 	}
 	else
 	{
-		SCCModule.QueueStatusUpdate(ExternalPackageName);
+		SCCModule.QueueStatusUpdate(ExternalPackageFileName);
 	}
 }
 
@@ -115,7 +119,7 @@ void FSceneOutlinerTreeItemSCC::HandleSourceControlStateChanged(EStateCacheUsage
 {
 	TRACE_CPUPROFILER_EVENT_SCOPE(FSceneOutlinerTreeItemSCC::HandleSourceControlStateChanged);
 
-	FSourceControlStatePtr SourceControlState = ISourceControlModule::Get().GetProvider().GetState(ExternalPackageName, CacheUsage);
+	FSourceControlStatePtr SourceControlState = ISourceControlModule::Get().GetProvider().GetState(ExternalPackageFileName, CacheUsage);
 	if (SourceControlState.IsValid())
 	{
 		BroadcastNewState(SourceControlState);
@@ -129,7 +133,7 @@ void FSceneOutlinerTreeItemSCC::HandleSourceControlProviderChanged(ISourceContro
 	
 	BroadcastNewState(nullptr);
 
-	ISourceControlModule::Get().QueueStatusUpdate(ExternalPackageName);
+	ISourceControlModule::Get().QueueStatusUpdate(ExternalPackageFileName);
 }
 
 void FSceneOutlinerTreeItemSCC::BroadcastNewState(FSourceControlStatePtr SourceControlState)
@@ -147,7 +151,7 @@ void FSceneOutlinerTreeItemSCC::HandleUncontrolledChangelistsStateChanged()
 
 	for (const TSharedRef<FUncontrolledChangelistState>& UncontrolledChangelistStateRef : UncontrolledChangelistStates)
 	{
-		if(UncontrolledChangelistStateRef->GetFilenames().Contains(ExternalPackageName))
+		if(UncontrolledChangelistStateRef->GetFilenames().Contains(ExternalPackageFileName))
 		{
 			UncontrolledChangelistState = UncontrolledChangelistStateRef;
 			break;
diff --git a/Engine/Source/Editor/SceneOutliner/Public/SceneOutlinerTreeItemSCC.h b/Engine/Source/Editor/SceneOutliner/Public/SceneOutlinerTreeItemSCC.h
--- a/Engine/Source/Editor/SceneOutliner/Public/SceneOutlinerTreeItemSCC.h
+++ b/Engine/Source/Editor/SceneOutliner/Public/SceneOutlinerTreeItemSCC.h
@@ -28,6 +28,8 @@ public:
 
 	FString GetPackageName() { return ExternalPackageName; }
 
+	FString GetPackageFileName() { return ExternalPackageFileName; }
+
 	UPackage* GetPackage() { return ExternalPackage; }
 
 	FSourceControlStateChangedDelegate OnSourceControlStateChanged;
@@ -53,8 +55,9 @@ private:
 	/** The tree item we relate to */
 	FSceneOutlinerTreeItemPtr TreeItemPtr;
 
-	/** Cache the items external package name */
+	/** Cache the items external package name and filename */
 	FString ExternalPackageName;
+	FString ExternalPackageFileName;
 	UPackage* ExternalPackage = nullptr;
 
 	/** Source control state changed delegate handle */
