Slate:

-Add additional Header Row-related options to FSceneOutlinerInitializationOptions and FSceneOutlinerColumnInfo

diff --git a/Engine/Source/Editor/SceneOutliner/Private/SSceneOutliner.cpp b/Engine/Source/Editor/SceneOutliner/Private/SSceneOutliner.cpp
--- a/Engine/Source/Editor/SceneOutliner/Private/SSceneOutliner.cpp
+++ b/Engine/Source/Editor/SceneOutliner/Private/SSceneOutliner.cpp
@@ -2,42 +2,30 @@
 
 #include "SSceneOutliner.h"
 
-#include "EdMode.h"
 #include "Editor.h"
-#include "Editor/UnrealEdEngine.h"
-#include "EditorModeManager.h"
 #include "Styling/AppStyle.h"
 #include "Engine/GameViewportClient.h"
-#include "Engine/Selection.h"
-#include "EngineUtils.h"
 #include "Framework/Application/SlateApplication.h"
 #include "Framework/Commands/UIAction.h"
 #include "Framework/MultiBox/MultiBoxBuilder.h"
 #include "HAL/PlatformApplicationMisc.h"
 #include "ISceneOutlinerColumn.h"
-#include "Kismet2/ComponentEditorUtils.h"
 #include "Layout/WidgetPath.h"
 #include "Modules/ModuleManager.h"
-#include "SceneOutlinerDelegates.h"
 #include "SceneOutlinerFilters.h"
 #include "SceneOutlinerModule.h"
 #include "ScopedTransaction.h"
 #include "Textures/SlateIcon.h"
 #include "ToolMenus.h"
-#include "UnrealEdGlobals.h"
 #include "UObject/PackageReload.h"
 #include "Widgets/Images/SImage.h"
 #include "Widgets/Input/SButton.h"
 #include "Widgets/Input/SComboButton.h"
-#include "Widgets/Input/SSearchBox.h"
-#include "Widgets/Layout/SSeparator.h"
 #include "Widgets/SOverlay.h"
-#include "SceneOutlinerMenuContext.h"
 #include "ISceneOutlinerMode.h"
 #include "FolderTreeItem.h"
 #include "EditorFolderUtils.h"
 #include "SceneOutlinerConfig.h"
-#include "Algo/ForEach.h"
 #include "SceneOutlinerFilterBar.h"
 #include "Framework/Notifications/NotificationManager.h"
 #include "Widgets/Notifications/SNotificationList.h"
@@ -133,7 +121,7 @@ void SSceneOutliner::Construct(const FArguments& InArgs, const FSceneOutlinerIni
 		SNew( SHeaderRow )
 			// Only show the list header if the user configured the outliner for that
 			.Visibility(InInitOptions.bShowHeaderRow ? EVisibility::Visible : EVisibility::Collapsed)
-			.CanSelectGeneratedColumn(true)
+			.CanSelectGeneratedColumn(InInitOptions.bCanSelectGeneratedColumns)
 			.OnHiddenColumnsListChanged(this, &SSceneOutliner::HandleHiddenColumnsChanged);
 
 	SetupColumns(*HeaderRowWidget);
@@ -507,27 +495,37 @@ void SSceneOutliner::SetupColumns(SHeaderRow& HeaderRow)
 				ColumnArgs.ShouldGenerateWidget(true);
 			}
 
-			if (SharedData->ColumnMap[ID].FillSize.IsSet())
+			if (FilteredColumnMap[ID].FillSize.IsSet())
 			{
-				ColumnArgs.FillWidth(SharedData->ColumnMap[ID].FillSize.GetValue());
+				ColumnArgs.FillWidth(FilteredColumnMap[ID].FillSize.GetValue());
 			}
 
+			if (FilteredColumnMap[ID].OnGetHeaderContextMenuContent.IsBound())
+			{
+				ColumnArgs.MenuContent()
+				[
+					FilteredColumnMap[ID].OnGetHeaderContextMenuContent.Execute()
+				];
+			}
+
+			ColumnArgs.HideHeaderMenuButton(FilteredColumnMap[ID].bHideMenuContentButton);
+
 			HeaderRow.AddColumn(ColumnArgs);
-			HeaderRowWidget->SetShowGeneratedColumn(ID, bIsVisible);
+			HeaderRow.SetShowGeneratedColumn(ID, bIsVisible);
 		}
 	}
 	Columns.Shrink();
 	bNeedsColumRefresh = false;
 }
 
-void SSceneOutliner::RefreshColums()
+void SSceneOutliner::RefreshColumns()
 {
 	bNeedsColumRefresh = true;
 }
 
 void SSceneOutliner::OnColumnPermissionListChanged()
 {
-	RefreshColums();
+	RefreshColumns();
 	FullRefresh();
 }
 
@@ -1160,7 +1158,7 @@ void SSceneOutliner::AddColumn(FName ColumId, const FSceneOutlinerColumnInfo& Co
 	if (!SharedData->ColumnMap.Contains(ColumId))
 	{
 		SharedData->ColumnMap.Add(ColumId, ColumInfo);
-		RefreshColums();
+		RefreshColumns();
 	}
 }
 
@@ -1169,7 +1167,7 @@ void SSceneOutliner::RemoveColumn(FName ColumId)
 	if (SharedData->ColumnMap.Contains(ColumId))
 	{
 		SharedData->ColumnMap.Remove(ColumId);
-		RefreshColums();
+		RefreshColumns();
 	}
 }
 
@@ -2554,4 +2552,4 @@ bool SSceneOutliner::IsFilterActive(const FString& FilterName) const
 	return false;
 }
 
-#undef LOCTEXT_NAMESPACE
\ No newline at end of file
+#undef LOCTEXT_NAMESPACE
diff --git a/Engine/Source/Editor/SceneOutliner/Public/SSceneOutliner.h b/Engine/Source/Editor/SceneOutliner/Public/SSceneOutliner.h
--- a/Engine/Source/Editor/SceneOutliner/Public/SSceneOutliner.h
+++ b/Engine/Source/Editor/SceneOutliner/Public/SSceneOutliner.h
@@ -1,7 +1,6 @@
 // Copyright Epic Games, Inc. All Rights Reserved.
 #pragma once
 
-#include "CoreMinimal.h"
 #include "EditorUndoClient.h"
 #include "Framework/Commands/UICommandList.h"
 #include "Framework/Views/ITypedTableView.h"
@@ -633,7 +632,7 @@ private:
 	void HandleHiddenColumnsChanged();
 
 	/** Refresh the scene outliner for when a colum was added or removed */
-	void RefreshColums();
+	void RefreshColumns();
 
 	/** Populates OutSearchStrings with the strings associated with TreeItem that should be used in searching */
 	void PopulateSearchStrings( const ISceneOutlinerTreeItem& TreeItem, OUT TArray< FString >& OutSearchStrings ) const;
diff --git a/Engine/Source/Editor/SceneOutliner/Public/SceneOutlinerPublicTypes.h b/Engine/Source/Editor/SceneOutliner/Public/SceneOutlinerPublicTypes.h
--- a/Engine/Source/Editor/SceneOutliner/Public/SceneOutlinerPublicTypes.h
+++ b/Engine/Source/Editor/SceneOutliner/Public/SceneOutlinerPublicTypes.h
@@ -2,7 +2,6 @@
 
 #pragma once
 
-#include "CoreMinimal.h"
 #include "Layout/Margin.h"
 #include "Framework/SlateDelegates.h"
 #include "Engine/World.h"
@@ -75,15 +74,20 @@ enum class ESceneOutlinerColumnVisibility : uint8
 struct FSceneOutlinerColumnInfo
 {
 	FSceneOutlinerColumnInfo(ESceneOutlinerColumnVisibility InVisibility, int32 InPriorityIndex, const FCreateSceneOutlinerColumn& InFactory = FCreateSceneOutlinerColumn(), bool inCanBeHidden = true, TOptional<float> InFillSize = TOptional<float>()
-		, TAttribute<FText> InColumnLabel = TAttribute<FText>())
-		: Visibility(InVisibility), PriorityIndex(InPriorityIndex), bCanBeHidden(inCanBeHidden), Factory(InFactory), FillSize(InFillSize), ColumnLabel(InColumnLabel)
+		, TAttribute<FText> InColumnLabel = TAttribute<FText>(), bool bInHideMenuContentButton = false, FOnGetContent InOnGetMenuContent = FOnGetContent())
+		: Visibility(InVisibility), PriorityIndex(InPriorityIndex), bCanBeHidden(inCanBeHidden)
+		, Factory(InFactory), FillSize(InFillSize), ColumnLabel(InColumnLabel)
+		, bHideMenuContentButton(bInHideMenuContentButton), OnGetHeaderContextMenuContent(InOnGetMenuContent)
 	{
 	}
 
 	FSceneOutlinerColumnInfo() {}
 
 	FSceneOutlinerColumnInfo(const FSceneOutlinerColumnInfo& InColumnInfo)
-		: Visibility(InColumnInfo.Visibility), PriorityIndex(InColumnInfo.PriorityIndex), bCanBeHidden(InColumnInfo.bCanBeHidden), Factory(InColumnInfo.Factory), FillSize(InColumnInfo.FillSize), ColumnLabel(InColumnInfo.ColumnLabel)
+		: Visibility(InColumnInfo.Visibility), PriorityIndex(InColumnInfo.PriorityIndex)
+		, bCanBeHidden(InColumnInfo.bCanBeHidden), Factory(InColumnInfo.Factory)
+		, FillSize(InColumnInfo.FillSize), ColumnLabel(InColumnInfo.ColumnLabel)
+		, bHideMenuContentButton(InColumnInfo.bHideMenuContentButton), OnGetHeaderContextMenuContent(InColumnInfo.OnGetHeaderContextMenuContent)
 	{}
 
 	ESceneOutlinerColumnVisibility 	Visibility;
@@ -92,6 +96,10 @@ struct FSceneOutlinerColumnInfo
 	FCreateSceneOutlinerColumn	Factory;
 	TOptional< float > FillSize;
 	TAttribute<FText> ColumnLabel; // Override for the column name used instead of ID if specified (use this if you want the column name to be localizable)
+
+	/** Hides the button on each header cell which normally shows up when OnGetHeaderContextMenuContent is bound */
+	bool bHideMenuContentButton;
+	FOnGetContent OnGetHeaderContextMenuContent;
 };
 
 /** Settings for the scene outliner which can be quieried publicly */
@@ -128,7 +136,6 @@ public:
 	void UseDefaultColumns();
 };
 
-
 /* Settings for the Filter Bar attached to the Scene Outliner. Can be specified through FSceneOutlinerInitializationOptions */
 struct FSceneOutlinerFilterBarOptions
 {
@@ -160,10 +167,17 @@ struct FSceneOutlinerFilterBarOptions
 	* of actors.
 	*/
 struct FSceneOutlinerInitializationOptions : FSharedSceneOutlinerData
-{
+{	
 	/** True if we should draw the header row above the tree view */
 	bool bShowHeaderRow : 1;
-
+	
+	/**
+	 * Can select the columns generated by right clicking on the header menu.
+	 * FColumn with ShouldGenerateWidget set can not be selected.
+	 * FColumn with MenuContent will still be displayed.
+	 */
+	bool bCanSelectGeneratedColumns : 1;
+	
 	/** Whether the Scene Outliner should expose its searchbox */
 	bool bShowSearchBox : 1;
 
@@ -189,6 +203,7 @@ public:
 	/** Constructor */
 	FSceneOutlinerInitializationOptions()
 		: bShowHeaderRow( true )
+		, bCanSelectGeneratedColumns( true )
 		, bShowSearchBox( true )
 		, bFocusSearchBoxWhenOpened( false )
 		, bShowCreateNewFolder( true )
@@ -218,4 +233,4 @@ struct SCENEOUTLINER_API FSceneOutlinerVisibilityCache
 	bool GetVisibility(const ISceneOutlinerTreeItem& Item) const;
 };
 
-#undef LOCTEXT_NAMESPACE
\ No newline at end of file
+#undef LOCTEXT_NAMESPACE
