PR #9948: Prevent a rare crash in actor hierarchy when representing world is invalid. (Contributed by [REDACTED_PERSON])

diff --git a/Engine/Source/Editor/SceneOutliner/Private/ActorHierarchy.cpp b/Engine/Source/Editor/SceneOutliner/Private/ActorHierarchy.cpp
--- a/Engine/Source/Editor/SceneOutliner/Private/ActorHierarchy.cpp
+++ b/Engine/Source/Editor/SceneOutliner/Private/ActorHierarchy.cpp
@@ -344,14 +344,19 @@ FSceneOutlinerTreeItemPtr FActorHierarchy::FindOrCreateParentItem(const ISceneOu
 	}
 
 	// If we get here return world item
-	if (const FSceneOutlinerTreeItemPtr* ParentItem = Items.Find(RepresentingWorld.Get()))
+	if (UWorld* RepresentingWorldPtr = RepresentingWorld.Get())
 	{
-		return *ParentItem;
-	}
-	else
-	{
-		return bCreate ? Mode->CreateItemFor<FWorldTreeItem>(RepresentingWorld.Get(), true) : nullptr;
+		if (const FSceneOutlinerTreeItemPtr* ParentItem = Items.Find(RepresentingWorldPtr))
+		{
+			return *ParentItem;
+		}
+		else
+		{
+			return bCreate ? Mode->CreateItemFor<FWorldTreeItem>(RepresentingWorldPtr, true) : nullptr;
+		}
 	}
+
+	return nullptr;
 }
 
 void FActorHierarchy::CreateComponentItems(const AActor* Actor, TArray<FSceneOutlinerTreeItemPtr>& OutItems) const
