SceneOutlinerTreeItemSCC and SSourceControlWidget: Change all delegates to use lambdas capturing weak references as a speculative fix to a crash on exit

diff --git a/Engine/Source/Editor/SceneOutliner/Private/SSourceControlWidget.cpp b/Engine/Source/Editor/SceneOutliner/Private/SSourceControlWidget.cpp
--- a/Engine/Source/Editor/SceneOutliner/Private/SSourceControlWidget.cpp
+++ b/Engine/Source/Editor/SceneOutliner/Private/SSourceControlWidget.cpp
@@ -10,10 +10,22 @@ void SSourceControlWidget::Construct(const FArguments& InArgs, TSharedPtr<FScene
 	check(InItemSourceControl.IsValid());
 
 	ItemSourceControl = InItemSourceControl;
-
-	ItemSourceControl->OnSourceControlStateChanged.BindSP(this, &SSourceControlWidget::UpdateSourceControlState);
-	ItemSourceControl->OnUncontrolledChangelistsStateChanged.BindSP(this, &SSourceControlWidget::UpdateUncontrolledChangelistState);
-
+	
+	ItemSourceControl->OnSourceControlStateChanged.BindLambda([this, WeakThis = AsWeak()](FSourceControlStatePtr SourceControlState)
+	{
+		if (WeakThis.IsValid())
+		{
+			UpdateSourceControlState(SourceControlState);
+		}
+	});
+	
+	ItemSourceControl->OnUncontrolledChangelistsStateChanged.BindLambda([this, WeakThis = AsWeak()](TSharedPtr<FUncontrolledChangelistState> UncontrolledChangelistState)
+	{
+		if (WeakThis.IsValid())
+		{
+			UpdateUncontrolledChangelistState(UncontrolledChangelistState);
+		}
+	});
 	
 	SImage::Construct(
 		SImage::FArguments()
diff --git a/Engine/Source/Editor/SceneOutliner/Private/SceneOutlinerSCCHandler.cpp b/Engine/Source/Editor/SceneOutliner/Private/SceneOutlinerSCCHandler.cpp
--- a/Engine/Source/Editor/SceneOutliner/Private/SceneOutlinerSCCHandler.cpp
+++ b/Engine/Source/Editor/SceneOutliner/Private/SceneOutlinerSCCHandler.cpp
@@ -19,6 +19,7 @@ TSharedPtr<FSceneOutlinerTreeItemSCC> FSceneOutlinerSCCHandler::GetItemSourceCon
 	if (!Result)
 	{
 		Result = &ItemSourceControls.Add(InItem, MakeShared<FSceneOutlinerTreeItemSCC>(InItem));
+		(*Result)->Initialize();
 	}
 	
 	check(Result);
diff --git a/Engine/Source/Editor/SceneOutliner/Private/SceneOutlinerTreeItemSCC.cpp b/Engine/Source/Editor/SceneOutliner/Private/SceneOutlinerTreeItemSCC.cpp
--- a/Engine/Source/Editor/SceneOutliner/Private/SceneOutlinerTreeItemSCC.cpp
+++ b/Engine/Source/Editor/SceneOutliner/Private/SceneOutlinerTreeItemSCC.cpp
@@ -16,7 +16,18 @@
 FSceneOutlinerTreeItemSCC::FSceneOutlinerTreeItemSCC(FSceneOutlinerTreeItemPtr InTreeItemPtr)
 {
 	TreeItemPtr = InTreeItemPtr;
+}
+
+FSceneOutlinerTreeItemSCC::~FSceneOutlinerTreeItemSCC()
+{
+	FUncontrolledChangelistsModule& UncontrolledChangelistModule = FUncontrolledChangelistsModule::Get();
+	UncontrolledChangelistModule.OnUncontrolledChangelistModuleChanged.Remove(UncontrolledChangelistChangedHandle);
+	
+	DisconnectSourceControl();
+}
 
+void FSceneOutlinerTreeItemSCC::Initialize()
+{
 	if (TreeItemPtr.IsValid())
 	{
 		ExternalPackageName = SceneOutliner::FSceneOutlinerHelpers::GetExternalPackageName(*TreeItemPtr.Get());
@@ -54,19 +65,17 @@ FSceneOutlinerTreeItemSCC::FSceneOutlinerTreeItemSCC(FSceneOutlinerTreeItemPtr I
 	}
 
 	FUncontrolledChangelistsModule& UncontrolledChangelistModule = FUncontrolledChangelistsModule::Get();
-	UncontrolledChangelistChangedHandle = UncontrolledChangelistModule.OnUncontrolledChangelistModuleChanged.AddRaw(this, &FSceneOutlinerTreeItemSCC::HandleUncontrolledChangelistsStateChanged);
-
-	// Call the delegate to update the initial uncontrolled state
-	HandleUncontrolledChangelistsStateChanged();
-
-}
 
-FSceneOutlinerTreeItemSCC::~FSceneOutlinerTreeItemSCC()
-{
-	FUncontrolledChangelistsModule& UncontrolledChangelistModule = FUncontrolledChangelistsModule::Get();
-	UncontrolledChangelistModule.OnUncontrolledChangelistModuleChanged.Remove(UncontrolledChangelistChangedHandle);
+	UncontrolledChangelistChangedHandle = UncontrolledChangelistModule.OnUncontrolledChangelistModuleChanged.AddLambda([this, WeakThis = AsWeak()]()
+	{
+		if (WeakThis.IsValid())
+		{
+			HandleUncontrolledChangelistsStateChanged();
+		}
+	});
 
-	DisconnectSourceControl();
+	// Call the delegate to update the initial uncontrolled state
+	HandleUncontrolledChangelistsStateChanged();	
 }
 
 FSourceControlStatePtr FSceneOutlinerTreeItemSCC::GetSourceControlState()
@@ -84,9 +93,23 @@ void FSceneOutlinerTreeItemSCC::ConnectSourceControl()
 	check(!ExternalPackageFileName.IsEmpty());
 
 	ISourceControlModule& SCCModule = ISourceControlModule::Get();
-	SourceControlProviderChangedDelegateHandle = SCCModule.RegisterProviderChanged(FSourceControlProviderChanged::FDelegate::CreateRaw(this, &FSceneOutlinerTreeItemSCC::HandleSourceControlProviderChanged));
-	SourceControlStateChangedDelegateHandle = SCCModule.GetProvider().RegisterSourceControlStateChanged_Handle(FSourceControlStateChanged::FDelegate::CreateRaw(this, &FSceneOutlinerTreeItemSCC::HandleSourceControlStateChanged, EStateCacheUsage::Use));
 
+	SourceControlProviderChangedDelegateHandle = SCCModule.RegisterProviderChanged(FSourceControlProviderChanged::FDelegate::CreateLambda([this, WeakThis = AsWeak()](ISourceControlProvider& OldProvider, ISourceControlProvider& NewProvider)
+	{
+		if (WeakThis.IsValid())
+		{
+			HandleSourceControlProviderChanged(OldProvider, NewProvider);
+		}
+	}));
+	
+	SourceControlStateChangedDelegateHandle = SCCModule.GetProvider().RegisterSourceControlStateChanged_Handle(FSourceControlStateChanged::FDelegate::CreateLambda([this, WeakThis = AsWeak()]()
+	{
+		if (WeakThis.IsValid())
+		{
+			HandleSourceControlStateChanged(EStateCacheUsage::Use);
+		}
+	}));
+	
 	// Check if there is already a cached state for this item
 	FSourceControlStatePtr SourceControlState = ISourceControlModule::Get().GetProvider().GetState(ExternalPackageFileName, EStateCacheUsage::Use);
 	if (SourceControlState.IsValid() && !SourceControlState->IsUnknown())
@@ -129,7 +152,14 @@ void FSceneOutlinerTreeItemSCC::HandleSourceControlStateChanged(EStateCacheUsage
 void FSceneOutlinerTreeItemSCC::HandleSourceControlProviderChanged(ISourceControlProvider& OldProvider, ISourceControlProvider& NewProvider)
 {
 	OldProvider.UnregisterSourceControlStateChanged_Handle(SourceControlStateChangedDelegateHandle);
-	SourceControlStateChangedDelegateHandle = NewProvider.RegisterSourceControlStateChanged_Handle(FSourceControlStateChanged::FDelegate::CreateRaw(this, &FSceneOutlinerTreeItemSCC::HandleSourceControlStateChanged, EStateCacheUsage::Use));
+
+	SourceControlStateChangedDelegateHandle = NewProvider.RegisterSourceControlStateChanged_Handle(FSourceControlStateChanged::FDelegate::CreateLambda([this, WeakThis = AsWeak()]()
+	{
+		if (WeakThis.IsValid())
+		{
+			HandleSourceControlStateChanged(EStateCacheUsage::Use);
+		}
+	}));
 	
 	BroadcastNewState(nullptr);
 
@@ -151,7 +181,7 @@ void FSceneOutlinerTreeItemSCC::HandleUncontrolledChangelistsStateChanged()
 
 	for (const TSharedRef<FUncontrolledChangelistState>& UncontrolledChangelistStateRef : UncontrolledChangelistStates)
 	{
-		if(UncontrolledChangelistStateRef->GetFilenames().Contains(ExternalPackageFileName))
+		if (UncontrolledChangelistStateRef->GetFilenames().Contains(ExternalPackageFileName))
 		{
 			UncontrolledChangelistState = UncontrolledChangelistStateRef;
 			break;
@@ -159,7 +189,7 @@ void FSceneOutlinerTreeItemSCC::HandleUncontrolledChangelistsStateChanged()
 	}
 
 	// Broadcast the delegate if our uncontrolled status was changed
-	if(UncontrolledChangelistState != PrevUncontrolledChangelistState)
+	if (UncontrolledChangelistState != PrevUncontrolledChangelistState)
 	{
 		OnUncontrolledChangelistsStateChanged.ExecuteIfBound(UncontrolledChangelistState);
 	}
diff --git a/Engine/Source/Editor/SceneOutliner/Public/SceneOutlinerTreeItemSCC.h b/Engine/Source/Editor/SceneOutliner/Public/SceneOutlinerTreeItemSCC.h
--- a/Engine/Source/Editor/SceneOutliner/Public/SceneOutlinerTreeItemSCC.h
+++ b/Engine/Source/Editor/SceneOutliner/Public/SceneOutlinerTreeItemSCC.h
@@ -20,6 +20,9 @@ public:
 
 	~FSceneOutlinerTreeItemSCC();
 
+	// Actually attempt to connect to Revision Control and get the item's status
+	void Initialize();
+
 	FSourceControlStatePtr GetSourceControlState();
 
 	FSourceControlStatePtr RefreshSourceControlState();
