WorldOutliner: Fix FActorDescTreeItem in different Level Instances with clashing Guids (were ignored)

diff --git a/Engine/Source/Editor/SceneOutliner/Private/ActorDescTreeItem.cpp b/Engine/Source/Editor/SceneOutliner/Private/ActorDescTreeItem.cpp
--- a/Engine/Source/Editor/SceneOutliner/Private/ActorDescTreeItem.cpp
+++ b/Engine/Source/Editor/SceneOutliner/Private/ActorDescTreeItem.cpp
@@ -24,6 +24,8 @@
 #include "ToolMenus.h"
 #include "LevelEditorViewport.h"
 #include "HAL/PlatformApplicationMisc.h"
+#include "Misc/ArchiveMD5.h"
+#include "UObject/ObjectKey.h"
 
 #define LOCTEXT_NAMESPACE "SceneOutliner_ActorDescTreeItem"
 
@@ -231,17 +233,18 @@ FActorDescTreeItem::FActorDescTreeItem(const FGuid& InActorGuid, UActorDescConta
 	Initialize();
 }
 
-FActorDescTreeItem::FActorDescTreeItem(const FGuid& InActorGuid, FActorDescContainerCollection* ContainerCollection)
-	: IActorBaseTreeItem(Type)
-	, ActorDescHandle(ContainerCollection, InActorGuid)
-	, ActorGuid(InActorGuid)
-{
-	Initialize();
-}
-
 void FActorDescTreeItem::Initialize()
 {
-	ID = FSceneOutlinerTreeItemID(ActorGuid);
+	FArchiveMD5 Ar;
+	Ar << ActorGuid;
+
+	FObjectKey ContainerKey(ActorDescHandle.Container.Get());
+	Ar << ContainerKey;
+
+	FMD5Hash MD5Hash;
+	Ar.GetHash(MD5Hash);
+	
+	ID = FSceneOutlinerTreeItemID(MD5HashToGuid(MD5Hash));
 
 	if (const FWorldPartitionActorDesc* const ActorDesc = ActorDescHandle.Get())
 	{
diff --git a/Engine/Source/Editor/SceneOutliner/Private/ActorHierarchy.cpp b/Engine/Source/Editor/SceneOutliner/Private/ActorHierarchy.cpp
--- a/Engine/Source/Editor/SceneOutliner/Private/ActorHierarchy.cpp
+++ b/Engine/Source/Editor/SceneOutliner/Private/ActorHierarchy.cpp
@@ -653,11 +653,11 @@ void FActorHierarchy::OnLoadedActorRemoved(AActor& InActor)
 		if (UWorldPartition* WorldPartition = RepresentingWorld->GetWorldPartition())
 		{
 			const FGuid& ActorGuid = InActor.GetActorGuid();
-			if (WorldPartition->GetActorDesc(ActorGuid) != nullptr)
+			if (FWorldPartitionActorDesc* ActorDesc = WorldPartition->GetActorDesc(ActorGuid))
 			{
 				FSceneOutlinerHierarchyChangedData EventData;
 				EventData.Type = FSceneOutlinerHierarchyChangedData::Added;
-				EventData.Items.Add(Mode->CreateItemFor<FActorDescTreeItem>(FActorDescTreeItem(ActorGuid, WorldPartition)));
+				EventData.Items.Add(Mode->CreateItemFor<FActorDescTreeItem>(FActorDescTreeItem(ActorGuid, ActorDesc->GetContainer())));
 				HierarchyChangedEvent.Broadcast(EventData);
 			}
 		}
@@ -672,7 +672,7 @@ void FActorHierarchy::OnActorDescAdded(FWorldPartitionActorDesc* ActorDesc)
 		{
 			FSceneOutlinerHierarchyChangedData EventData;
 			EventData.Type = FSceneOutlinerHierarchyChangedData::Added;
-			EventData.Items.Add(Mode->CreateItemFor<FActorDescTreeItem>(FActorDescTreeItem(ActorDesc->GetGuid(), WorldPartition)));
+			EventData.Items.Add(Mode->CreateItemFor<FActorDescTreeItem>(FActorDescTreeItem(ActorDesc->GetGuid(), ActorDesc->GetContainer())));
 			HierarchyChangedEvent.Broadcast(EventData);
 		}
 	}
diff --git a/Engine/Source/Editor/SceneOutliner/Public/ActorDescTreeItem.h b/Engine/Source/Editor/SceneOutliner/Public/ActorDescTreeItem.h
--- a/Engine/Source/Editor/SceneOutliner/Public/ActorDescTreeItem.h
+++ b/Engine/Source/Editor/SceneOutliner/Public/ActorDescTreeItem.h
@@ -41,9 +41,6 @@ public:
 	/** Construct this item from an actor desc */
 	FActorDescTreeItem(const FGuid& InActorGuid, UActorDescContainer* Container);
 
-	/** Construct this item from an actor desc */
-	FActorDescTreeItem(const FGuid& InActorGuid, FActorDescContainerCollection* ContainerCollection);
-
 	/* Begin ISceneOutlinerTreeItem Implementation */
 	virtual bool IsValid() const override { return ActorDescHandle.Get() != nullptr; }
 	virtual FSceneOutlinerTreeItemID GetID() const override;
