Outliner: FActorMode and SceneOutlinerModule now both query the Level Editor to pick the correct world based on context, preferring PIE over Editor worlds

diff --git a/Engine/Source/Editor/SceneOutliner/Private/ActorMode.cpp b/Engine/Source/Editor/SceneOutliner/Private/ActorMode.cpp
--- a/Engine/Source/Editor/SceneOutliner/Private/ActorMode.cpp
+++ b/Engine/Source/Editor/SceneOutliner/Private/ActorMode.cpp
@@ -15,6 +15,7 @@
 #include "FolderTreeItem.h"
 #include "ComponentTreeItem.h"
 #include "ActorDescTreeItem.h"
+#include "EditorModeManager.h"
 #include "WorldTreeItem.h"
 #include "LevelInstance/LevelInstanceInterface.h"
 #include "LevelInstance/LevelInstanceSubsystem.h"
@@ -223,19 +224,12 @@ void FActorMode::ChooseRepresentingWorld()
 
 	if (RepresentingWorld == nullptr)
 	{
-		// still no world so fallback to old logic where we just prefer PIE over Editor
+		// If there is still no world, we query the Level Editor, which prefers the PIE world over the Editor world
+		TWeakPtr<ILevelEditor> LevelEditor = FModuleManager::GetModuleChecked<FLevelEditorModule>(TEXT("LevelEditor")).GetLevelEditorInstance();
 
-		for (const FWorldContext& Context : GEngine->GetWorldContexts())
+		if (TSharedPtr<ILevelEditor> LevelEditorPin = LevelEditor.Pin())
 		{
-			if (Context.WorldType == EWorldType::PIE)
-			{
-				RepresentingWorld = Context.World();
-				break;
-			}
-			else if (Context.WorldType == EWorldType::Editor)
-			{
-				RepresentingWorld = Context.World();
-			}
+			RepresentingWorld = LevelEditorPin->GetEditorModeManager().GetWorld();
 		}
 	}
 }
diff --git a/Engine/Source/Editor/SceneOutliner/Private/SceneOutlinerModule.cpp b/Engine/Source/Editor/SceneOutliner/Private/SceneOutlinerModule.cpp
--- a/Engine/Source/Editor/SceneOutliner/Private/SceneOutlinerModule.cpp
+++ b/Engine/Source/Editor/SceneOutliner/Private/SceneOutlinerModule.cpp
@@ -31,6 +31,8 @@
 #include "LevelInstance/LevelInstanceSubsystem.h"
 #include "LevelInstance/LevelInstanceInterface.h"
 #include "Editor.h"
+#include "EditorModeManager.h"
+#include "LevelEditor.h"
 
 #define LOCTEXT_NAMESPACE "SceneOutlinerModule"
 
@@ -168,18 +170,23 @@ TSharedRef< ISceneOutliner > FSceneOutlinerModule::CreateActorBrowser(const FSce
 
 		ESceneOutlinerColumnVisibility UnsavedColumnVisibility = ESceneOutlinerColumnVisibility::Invisible;
 		
-		UWorld* WorldPtr = GEditor->GetEditorWorldContext().World();
-		if (WorldPtr)
+		// Query the Level Editor to get the correct world based on context
+		TWeakPtr<ILevelEditor> LevelEditor = FModuleManager::GetModuleChecked<FLevelEditorModule>(TEXT("LevelEditor")).GetLevelEditorInstance();
+
+		if (TSharedPtr<ILevelEditor> LevelEditorPin = LevelEditor.Pin())
 		{
-			if (WorldPtr->IsPartitionedWorld())
+			if (UWorld* WorldPtr = LevelEditorPin->GetEditorModeManager().GetWorld())
 			{
-				// We don't want the pinned column in non wp levels
-				InitOptions.ColumnMap.Add(FSceneOutlinerBuiltInColumnTypes::Pinned(), FSceneOutlinerColumnInfo(ESceneOutlinerColumnVisibility::Visible, 5, FCreateSceneOutlinerColumn(), true, TOptional<float>(), FSceneOutlinerBuiltInColumnTypes::Pinned_Localized()));
+				if (WorldPtr->IsPartitionedWorld())
+				{
+					// We don't want the pinned column in non wp levels
+					InitOptions.ColumnMap.Add(FSceneOutlinerBuiltInColumnTypes::Pinned(), FSceneOutlinerColumnInfo(ESceneOutlinerColumnVisibility::Visible, 5, FCreateSceneOutlinerColumn(), true, TOptional<float>(), FSceneOutlinerBuiltInColumnTypes::Pinned_Localized()));
 
-				// We want the unsaved column to be visible by default in partitioned levels
-				UnsavedColumnVisibility = ESceneOutlinerColumnVisibility::Visible;
+					// We want the unsaved column to be visible by default in partitioned levels
+					UnsavedColumnVisibility = ESceneOutlinerColumnVisibility::Visible;
+				}
+				CreateActorInfoColumns(InitOptions, WorldPtr);
 			}
-			CreateActorInfoColumns(InitOptions, WorldPtr);
 		}
 
 		InitOptions.ColumnMap.Add(FSceneOutlinerBuiltInColumnTypes::Unsaved(), FSceneOutlinerColumnInfo(UnsavedColumnVisibility, 1, FCreateSceneOutlinerColumn(), true, TOptional<float>(), FSceneOutlinerBuiltInColumnTypes::Unsaved_Localized()));
