SkeinUE: Allow to revert an unsaved file with in-memory changes from the SceneOutliner.

diff --git a/Engine/Source/Editor/SceneOutliner/Private/SceneOutlinerSCCHandler.cpp b/Engine/Source/Editor/SceneOutliner/Private/SceneOutlinerSCCHandler.cpp
--- a/Engine/Source/Editor/SceneOutliner/Private/SceneOutlinerSCCHandler.cpp
+++ b/Engine/Source/Editor/SceneOutliner/Private/SceneOutlinerSCCHandler.cpp
@@ -114,7 +114,7 @@ void FSceneOutlinerSCCHandler::CacheCanExecuteVars()
 
 			// Check the SCC state for each package in the selected paths
 			FSourceControlStatePtr SourceControlState = SourceControl->GetSourceControlState();
-			if(SourceControlState.IsValid())
+			if (SourceControlState.IsValid())
 			{
 				if ( SourceControlState->CanCheckout() )
 				{
@@ -135,6 +135,17 @@ void FSceneOutlinerSCCHandler::CacheCanExecuteVars()
 				{
 					bCanExecuteSCCRevert = true;
 				}
+				else
+				{
+					// If the package is dirty, allow a revert of the in-memory changes that have not yet been saved to disk.
+					if (UPackage* Package = SourceControl->GetPackage())
+					{
+						if (Package->IsDirty())
+						{
+							bCanExecuteSCCRevert = true;
+						}
+					}
+				}
 			}
 		}
 
@@ -369,9 +380,23 @@ void FSceneOutlinerSCCHandler::ExecuteSCCCheckIn()
 
 void FSceneOutlinerSCCHandler::ExecuteSCCRevert()
 {
-	TArray<FString> PackageNames;
-	GetSelectedPackageNames(PackageNames);
-	FSourceControlWindows::PromptForRevert(PackageNames, /*bReloadWorld=*/true);
+	TArray<UPackage*> PackagesToSave;
+	GetSelectedPackages(PackagesToSave);
+
+	if (PackagesToSave.Num() > 0)
+	{
+		// To prevent a 'Save Dialog' from popping up during RevertAndReloadPackages we save the dirty packages here,
+		// as the 'Save Dialog' is confusing during a revert operation, unlike the check in operation.
+		UEditorLoadingAndSavingUtils::SavePackages(PackagesToSave, /*bOnlyDirty=*/true);
+	}
+
+	TArray<FString> PackagesToRevert;
+	GetSelectedPackageNames(PackagesToRevert);
+
+	if (PackagesToRevert.Num() > 0)
+	{
+		SourceControlHelpers::RevertAndReloadPackages(PackagesToRevert, /*bRevertAll=*/false, /*bReloadWorld=*/true);
+	}
 }
 
 void FSceneOutlinerSCCHandler::ExecuteSCCHistory()
