Fixed toggling outliner folder visibility not working properly if some children are Pinned and others are not.

diff --git a/Engine/Source/Editor/SceneOutliner/Private/ActorDescTreeItem.cpp b/Engine/Source/Editor/SceneOutliner/Private/ActorDescTreeItem.cpp
--- a/Engine/Source/Editor/SceneOutliner/Private/ActorDescTreeItem.cpp
+++ b/Engine/Source/Editor/SceneOutliner/Private/ActorDescTreeItem.cpp
@@ -308,7 +308,7 @@ void FActorDescTreeItem::GenerateContextMenu(UToolMenu* Menu, SSceneOutliner&)
 
 bool FActorDescTreeItem::GetVisibility() const
 {
-	return true;
+	return false;
 }
 
 bool FActorDescTreeItem::ShouldShowPinnedState() const
