Outliner: Remove the foldercount variable as it was only counting folders that are visible instead of all folders that exist

diff --git a/Engine/Source/Editor/SceneOutliner/Private/SSceneOutliner.cpp b/Engine/Source/Editor/SceneOutliner/Private/SSceneOutliner.cpp
--- a/Engine/Source/Editor/SceneOutliner/Private/SSceneOutliner.cpp
+++ b/Engine/Source/Editor/SceneOutliner/Private/SSceneOutliner.cpp
@@ -765,8 +765,6 @@ void SSceneOutliner::EmptyTreeItems()
 	TreeItemMap.Reset();
 	PendingTreeItemMap.Empty();
 
-	FolderCount = 0;
-
 	RootTreeItems.Empty();
 }
 
@@ -1005,11 +1003,6 @@ void SSceneOutliner::AddUnfilteredItemToTree(FSceneOutlinerTreeItemRef Item)
 		RootTreeItems.Add(Item);
 	}
 
-	// keep track of the number of active folders
-	if (Item->IsA<FFolderTreeItem>())
-	{
-		++FolderCount;
-	}
 	Mode->OnItemAdded(Item);
 }
 
@@ -1322,8 +1315,8 @@ TSharedRef<TSet<FFolder>> SSceneOutliner::GatherInvalidMoveToDestinations() cons
 	
 void SSceneOutliner::AddMoveToFolderOutliner(UToolMenu* Menu) const
 {
-	// We don't show this if there aren't any folders in the world and if the mode is showing folders
-	if (!Mode->ShouldShowFolders() || FolderCount == 0)
+	// We don't show this if the mode does not show folders
+	if (!Mode->ShouldShowFolders())
 	{
 		return;
 	}
diff --git a/Engine/Source/Editor/SceneOutliner/Public/SSceneOutliner.h b/Engine/Source/Editor/SceneOutliner/Public/SSceneOutliner.h
--- a/Engine/Source/Editor/SceneOutliner/Public/SSceneOutliner.h
+++ b/Engine/Source/Editor/SceneOutliner/Public/SSceneOutliner.h
@@ -919,9 +919,6 @@ private:
 
 	TMap<FName, const FSlateBrush*> CachedIcons;
 
-	/** Maintain a count of the number of folders active in the outliner */
-	uint32 FolderCount = 0;
-
 	SceneOutliner::FTreeItemPtrEvent OnDoubleClickOnTreeEvent;
 
 	SceneOutliner::FOnItemSelectionChanged OnItemSelectionChanged;
