World Outliner Level column name shows "Package Short Name" when level uses OFPA (same as partitioned levels).

diff --git a/Engine/Source/Editor/SceneOutliner/Private/SceneOutlinerModule.cpp b/Engine/Source/Editor/SceneOutliner/Private/SceneOutlinerModule.cpp
--- a/Engine/Source/Editor/SceneOutliner/Private/SceneOutlinerModule.cpp
+++ b/Engine/Source/Editor/SceneOutliner/Private/SceneOutlinerModule.cpp
@@ -494,7 +494,7 @@ void FSceneOutlinerModule::CreateActorInfoColumns(FSceneOutlinerInitializationOp
 	// The "Level" column should be named "Package Short Name" in wp enabled levels
 	auto LevelColumnName = TAttribute<FText>::CreateLambda([WorldPtr]() -> FText
 	{
-		if (WorldPtr && WorldPtr->IsPartitionedWorld())
+		if (WorldPtr && WorldPtr->PersistentLevel->IsUsingExternalActors())
 		{
 			return FSceneOutlinerBuiltInColumnTypes::PackageShortName_Localized();
 		}
