Add more scope traces to investigate Outliner hitches during PIE

#swarm 24390675

diff --git a/Engine/Source/Editor/SceneOutliner/Private/ActorHierarchy.cpp b/Engine/Source/Editor/SceneOutliner/Private/ActorHierarchy.cpp
--- a/Engine/Source/Editor/SceneOutliner/Private/ActorHierarchy.cpp
+++ b/Engine/Source/Editor/SceneOutliner/Private/ActorHierarchy.cpp
@@ -380,6 +380,8 @@ void FActorHierarchy::CreateComponentItems(const AActor* Actor, TArray<FSceneOut
 
 void FActorHierarchy::CreateWorldChildren(UWorld* World, TArray<FSceneOutlinerTreeItemPtr>& OutItems) const
 {
+	TRACE_CPUPROFILER_EVENT_SCOPE(FActorHierarchy::CreateWorldChildren);
+
 	check(World);
 
 	if (Mode->ShouldShowFolders() && bShowingEmptyFolders)
@@ -473,6 +475,8 @@ void FActorHierarchy::CreateWorldChildren(UWorld* World, TArray<FSceneOutlinerTr
 
 void FActorHierarchy::CreateItems(TArray<FSceneOutlinerTreeItemPtr>& OutItems) const
 {
+	TRACE_CPUPROFILER_EVENT_SCOPE(FActorHierarchy::CreateItems);
+
 	if (RepresentingWorld.IsValid())
 	{
 		UWorld* RepresentingWorldPtr = RepresentingWorld.Get();
diff --git a/Engine/Source/Editor/SceneOutliner/Private/SSceneOutliner.cpp b/Engine/Source/Editor/SceneOutliner/Private/SSceneOutliner.cpp
--- a/Engine/Source/Editor/SceneOutliner/Private/SSceneOutliner.cpp
+++ b/Engine/Source/Editor/SceneOutliner/Private/SSceneOutliner.cpp
@@ -640,6 +640,8 @@ void SSceneOutliner::RefreshSelection()
 
 void SSceneOutliner::Populate()
 {
+	TRACE_CPUPROFILER_EVENT_SCOPE(SSceneOutliner::Populate);
+
 	// Block events while we clear out the list
  	TGuardValue<bool> ReentrantGuard(bIsReentrant, true);
 
@@ -801,6 +803,8 @@ void SSceneOutliner::AddPendingItemAndChildren(FSceneOutlinerTreeItemPtr Item)
 
 void SSceneOutliner::RepopulateEntireTree()
 {
+	TRACE_CPUPROFILER_EVENT_SCOPE(SSceneOutliner::RepopulateEntireTree);
+
 	EmptyTreeItems();
 
 	// Rebuild the hierarchy
@@ -2198,6 +2202,8 @@ FReply SSceneOutliner::OnKeyDown( const FGeometry& MyGeometry, const FKeyEvent&
 
 void SSceneOutliner::Tick(const FGeometry& AllottedGeometry, const double InCurrentTime, const float InDeltaTime)
 {	
+	TRACE_CPUPROFILER_EVENT_SCOPE(SSceneOutliner::Tick);
+
 	for (auto& Pair : Columns)
 	{
 		Pair.Value->Tick(InCurrentTime, InDeltaTime);
