Enable revision control column in outliner by default when source control and external actors are enabled.

diff --git a/Engine/Source/Editor/SceneOutliner/Private/SSceneOutliner.cpp b/Engine/Source/Editor/SceneOutliner/Private/SSceneOutliner.cpp
--- a/Engine/Source/Editor/SceneOutliner/Private/SSceneOutliner.cpp
+++ b/Engine/Source/Editor/SceneOutliner/Private/SSceneOutliner.cpp
@@ -67,6 +67,7 @@ void SSceneOutliner::Construct(const FArguments& InArgs, const FSceneOutlinerIni
 	bFullRefresh = true;
 	bNeedsRefresh = true;
 	bNeedsColumRefresh = true;
+	bShouldCacheColumnVisibility = true;
 	bIsReentrant = false;
 	bSortDirty = true;
 	bSelectionDirty = true;
@@ -126,6 +127,8 @@ void SSceneOutliner::Construct(const FArguments& InArgs, const FSceneOutlinerIni
 
 	SetupColumns(*HeaderRowWidget);
 
+	CacheHiddenColumns = TSet(HeaderRowWidget->GetHiddenColumnIds());
+
 	ChildSlot
 	[
 		VerticalBox
@@ -343,7 +346,12 @@ void SSceneOutliner::Construct(const FArguments& InArgs, const FSceneOutlinerIni
 
 void SSceneOutliner::HandleHiddenColumnsChanged()
 {
-	const TArray<FName> HiddenColumns = HeaderRowWidget->GetHiddenColumnIds();
+	if (!bShouldCacheColumnVisibility)
+	{
+		return;
+	}
+	
+	TSet<FName> HiddenColumns = TSet(HeaderRowWidget->GetHiddenColumnIds());
 	FSceneOutlinerConfig* OutlinerConfig = GetMutableConfig();
 
 	if (OutlinerConfig != nullptr)
@@ -351,23 +359,16 @@ void SSceneOutliner::HandleHiddenColumnsChanged()
 		TMap<FName, bool> ColumnVisibilities = OutlinerConfig->ColumnVisibilities;
 
 		bool bAnyColumnVisibilityChanged = false;
-		
+	
 		for (const TPair<FName, TSharedPtr<ISceneOutlinerColumn>>& Pair : Columns)
 		{
-			const bool bIsColumnVisible = HiddenColumns.Find(Pair.Key) == INDEX_NONE;
+			const bool bWasColumnVisible = CacheHiddenColumns.Find(Pair.Key) != nullptr;
+			const bool bIsColumnVisible = HiddenColumns.Find(Pair.Key) != nullptr;
 
-			// If this column already has a visibility saved in the config, only update if it changed
-			if (bool* ExistingColumnVisibility = ColumnVisibilities.Find(Pair.Key))
+			// Only update column visibility if it changed
+			if (bWasColumnVisible != bIsColumnVisible)
 			{
-				if(*ExistingColumnVisibility != bIsColumnVisible)
-				{
-					*ExistingColumnVisibility = bIsColumnVisible;
-					bAnyColumnVisibilityChanged = true;
-				}
-			}
-			else
-			{
-				ColumnVisibilities.Add(Pair.Key, bIsColumnVisible);
+				ColumnVisibilities.FindOrAdd(Pair.Key) = bIsColumnVisible;
 				bAnyColumnVisibilityChanged = true;
 			}
 		}
@@ -375,11 +376,13 @@ void SSceneOutliner::HandleHiddenColumnsChanged()
 		// Only call SaveConfig if something actually changed
 		if(bAnyColumnVisibilityChanged)
 		{
-			OutlinerConfig->ColumnVisibilities = ColumnVisibilities;
-			SaveConfig();
-		}
+		OutlinerConfig->ColumnVisibilities = ColumnVisibilities;
+		SaveConfig();
 	}
 }
+	
+	CacheHiddenColumns = MoveTemp(HiddenColumns);
+}
 
 void SSceneOutliner::GetSortedColumnIDs(TArray<FName>& OutColumnIDs) const
 {
@@ -438,6 +441,10 @@ void SSceneOutliner::SetupColumns(SHeaderRow& HeaderRow)
 		ColumnVisibilities = OutlinerConfig->ColumnVisibilities;
 	}
 
+	// Avoid caching column visibilities while building the columns
+	bool const bPreviousShouldCacheColumnVisibility = bShouldCacheColumnVisibility;
+	bShouldCacheColumnVisibility = false;
+	
 	for (const FName& ID : SortedIDs)
 	{
 		bool bIsVisible = true;
@@ -516,6 +523,7 @@ void SSceneOutliner::SetupColumns(SHeaderRow& HeaderRow)
 	}
 	Columns.Shrink();
 	bNeedsColumRefresh = false;
+	bShouldCacheColumnVisibility = bPreviousShouldCacheColumnVisibility;
 }
 
 void SSceneOutliner::RefreshColumns()
@@ -2551,4 +2559,4 @@ bool SSceneOutliner::IsFilterActive(const FString& FilterName) const
 	return false;
 }
 
-#undef LOCTEXT_NAMESPACE
+#undef LOCTEXT_NAMESPACE
\ No newline at end of file
diff --git a/Engine/Source/Editor/SceneOutliner/Private/SceneOutlinerModule.cpp b/Engine/Source/Editor/SceneOutliner/Private/SceneOutlinerModule.cpp
--- a/Engine/Source/Editor/SceneOutliner/Private/SceneOutlinerModule.cpp
+++ b/Engine/Source/Editor/SceneOutliner/Private/SceneOutlinerModule.cpp
@@ -31,6 +31,7 @@
 #include "LevelInstance/LevelInstanceSubsystem.h"
 #include "LevelInstance/LevelInstanceInterface.h"
 #include "Editor.h"
+#include "ISourceControlModule.h"
 #include "EditorModeManager.h"
 #include "LevelEditor.h"
 
@@ -170,28 +171,38 @@ TSharedRef< ISceneOutliner > FSceneOutlinerModule::CreateActorBrowser(const FSce
 		InitOptions.UseDefaultColumns();
 		InitOptions.ColumnMap.Add(FSceneOutlinerBuiltInColumnTypes::Gutter(), FSceneOutlinerColumnInfo(ESceneOutlinerColumnVisibility::Visible, 0, FCreateSceneOutlinerColumn(), true, TOptional<float>(), FSceneOutlinerBuiltInColumnTypes::Gutter_Localized()));
 		InitOptions.ColumnMap.Add(FSceneOutlinerBuiltInColumnTypes::ActorInfo(), FSceneOutlinerColumnInfo(ESceneOutlinerColumnVisibility::Visible, 20, FCreateSceneOutlinerColumn(), true, TOptional<float>(), FSceneOutlinerBuiltInColumnTypes::ActorInfo_Localized()));
-		InitOptions.ColumnMap.Add(FSceneOutlinerBuiltInColumnTypes::SourceControl(), FSceneOutlinerColumnInfo(ESceneOutlinerColumnVisibility::Invisible, 30, FCreateSceneOutlinerColumn(), true, TOptional<float>(), FSceneOutlinerBuiltInColumnTypes::SourceControl_Localized()));
 
-		ESceneOutlinerColumnVisibility UnsavedColumnVisibility = ESceneOutlinerColumnVisibility::Invisible;
+		ESceneOutlinerColumnVisibility SourceControlColumnVisibility = ESceneOutlinerColumnVisibility::Invisible;
 		
+		UWorld* WorldPtr = nullptr;
 		// Query the Level Editor to get the correct world based on context
 		TWeakPtr<ILevelEditor> LevelEditor = FModuleManager::GetModuleChecked<FLevelEditorModule>(TEXT("LevelEditor")).GetLevelEditorInstance();
-
 		if (TSharedPtr<ILevelEditor> LevelEditorPin = LevelEditor.Pin())
 		{
-			if (UWorld* WorldPtr = LevelEditorPin->GetEditorModeManager().GetWorld())
+			WorldPtr = LevelEditorPin->GetEditorModeManager().GetWorld();
+		}
+		ISourceControlModule& SourceControlModule = ISourceControlModule::Get();
+		if (SourceControlModule.IsEnabled())
+		{
+			if (WorldPtr && WorldPtr->PersistentLevel->IsUsingExternalActors())
 			{
-				if (WorldPtr->IsPartitionedWorld())
-				{
-					// We don't want the pinned column in non wp levels
-					InitOptions.ColumnMap.Add(FSceneOutlinerBuiltInColumnTypes::Pinned(), FSceneOutlinerColumnInfo(ESceneOutlinerColumnVisibility::Visible, 5, FCreateSceneOutlinerColumn(), true, TOptional<float>(), FSceneOutlinerBuiltInColumnTypes::Pinned_Localized()));
-
-					// We want the unsaved column to be visible by default in partitioned levels
-					UnsavedColumnVisibility = ESceneOutlinerColumnVisibility::Visible;
-				}
-				CreateActorInfoColumns(InitOptions, WorldPtr);
+				// The source control column should be visible by default in source-controlled levels using external actors
+				SourceControlColumnVisibility = ESceneOutlinerColumnVisibility::Visible;
 			}
 		}
+		InitOptions.ColumnMap.Add(FSceneOutlinerBuiltInColumnTypes::SourceControl(), FSceneOutlinerColumnInfo(SourceControlColumnVisibility, 30, FCreateSceneOutlinerColumn(), true, TOptional<float>(), FSceneOutlinerBuiltInColumnTypes::SourceControl_Localized()));
+
+		ESceneOutlinerColumnVisibility UnsavedColumnVisibility = ESceneOutlinerColumnVisibility::Invisible;
+		
+		if (WorldPtr && WorldPtr->IsPartitionedWorld())
+		{
+			// We don't want the pinned column in non wp levels
+			InitOptions.ColumnMap.Add(FSceneOutlinerBuiltInColumnTypes::Pinned(), FSceneOutlinerColumnInfo(ESceneOutlinerColumnVisibility::Visible, 5, FCreateSceneOutlinerColumn(), true, TOptional<float>(), FSceneOutlinerBuiltInColumnTypes::Pinned_Localized()));
+
+			// We want the unsaved column to be visible by default in partitioned levels
+			UnsavedColumnVisibility = ESceneOutlinerColumnVisibility::Visible;
+		}
+		CreateActorInfoColumns(InitOptions, WorldPtr);
 
 		InitOptions.ColumnMap.Add(FSceneOutlinerBuiltInColumnTypes::Unsaved(), FSceneOutlinerColumnInfo(UnsavedColumnVisibility, 1, FCreateSceneOutlinerColumn(), true, TOptional<float>(), FSceneOutlinerBuiltInColumnTypes::Unsaved_Localized()));
 		
diff --git a/Engine/Source/Editor/SceneOutliner/Public/SSceneOutliner.h b/Engine/Source/Editor/SceneOutliner/Public/SSceneOutliner.h
--- a/Engine/Source/Editor/SceneOutliner/Public/SSceneOutliner.h
+++ b/Engine/Source/Editor/SceneOutliner/Public/SSceneOutliner.h
@@ -758,6 +758,9 @@ private:
 	/** Maps pre-existing children during paste or duplicate */
 	TMap<FFolder, TArray<FSceneOutlinerTreeItemID>> CachePasteFolderExistingChildrenMap;
 
+	/** Cache which columns are hidden to know what changed */
+	TSet<FName> CacheHiddenColumns;
+
 private:
 	bool GetCommonRootObjectFromSelection(FFolder::FRootObject& OutCommonRootObject) const;
 
@@ -887,6 +890,9 @@ private:
 
 	uint8 bNeedsColumRefresh : 1;
 
+	/** True if the outliner should cache changes to column visibility into the config */
+	uint8 bShouldCacheColumnVisibility : 1;
+
 	/** Reentrancy guard */
 	bool bIsReentrant;
 
