WorldPartition - HLOD: Fix HLODs missing from the scene outliner since CL 24125400

diff --git a/Engine/Source/Editor/SceneOutliner/Private/ActorHierarchy.cpp b/Engine/Source/Editor/SceneOutliner/Private/ActorHierarchy.cpp
--- a/Engine/Source/Editor/SceneOutliner/Private/ActorHierarchy.cpp
+++ b/Engine/Source/Editor/SceneOutliner/Private/ActorHierarchy.cpp
@@ -458,7 +458,7 @@ void FActorHierarchy::CreateWorldChildren(UWorld* World, TArray<FSceneOutlinerTr
 				}
 				FWorldPartitionHelpers::ForEachActorDesc(WorldPartition, [this, &OutItems](const FWorldPartitionActorDesc* ActorDesc)
 				{
-					if (ActorDesc != nullptr && !ActorDesc->IsLoaded(true) && ActorDesc->IsEditorRelevant())
+					if (ActorDesc != nullptr && !ActorDesc->IsLoaded(true))
 					{
 						if (const FSceneOutlinerTreeItemPtr ActorDescItem = Mode->CreateItemFor<FActorDescTreeItem>(FActorDescTreeItem(ActorDesc->GetGuid(), ActorDesc->GetContainer())))
 						{
