SceneOutliner: Fixed 'DoesPackageExist called on PackageName that will always return false.' error showing up in the log that happened when the PackageName was an empty string.

diff --git a/Engine/Source/Editor/SceneOutliner/Private/SceneOutlinerTreeItemSCC.cpp b/Engine/Source/Editor/SceneOutliner/Private/SceneOutlinerTreeItemSCC.cpp
--- a/Engine/Source/Editor/SceneOutliner/Private/SceneOutlinerTreeItemSCC.cpp
+++ b/Engine/Source/Editor/SceneOutliner/Private/SceneOutlinerTreeItemSCC.cpp
@@ -31,7 +31,7 @@ void FSceneOutlinerTreeItemSCC::Initialize()
 	if (TreeItemPtr.IsValid())
 	{
 		ExternalPackageName = SceneOutliner::FSceneOutlinerHelpers::GetExternalPackageName(*TreeItemPtr.Get());
-		ExternalPackageFileName = USourceControlHelpers::PackageFilename(ExternalPackageName);
+		ExternalPackageFileName = !ExternalPackageName.IsEmpty() ? USourceControlHelpers::PackageFilename(ExternalPackageName) : FString();
 		ExternalPackage = SceneOutliner::FSceneOutlinerHelpers::GetExternalPackage(*TreeItemPtr.Get());
 		
 		if (FActorTreeItem* ActorItem = TreeItemPtr->CastTo<FActorTreeItem>())
@@ -75,7 +75,7 @@ void FSceneOutlinerTreeItemSCC::Initialize()
 	});
 
 	// Call the delegate to update the initial uncontrolled state
-	HandleUncontrolledChangelistsStateChanged();	
+	HandleUncontrolledChangelistsStateChanged();
 }
 
 FSourceControlStatePtr FSceneOutlinerTreeItemSCC::GetSourceControlState()
