Fix duplicate 'Create Subfolder' menu entry appearing in outliner

diff --git a/Engine/Source/Editor/SceneOutliner/Private/ActorBrowsingMode.cpp b/Engine/Source/Editor/SceneOutliner/Private/ActorBrowsingMode.cpp
--- a/Engine/Source/Editor/SceneOutliner/Private/ActorBrowsingMode.cpp
+++ b/Engine/Source/Editor/SceneOutliner/Private/ActorBrowsingMode.cpp
@@ -705,8 +705,6 @@ void FActorBrowsingMode::RegisterContextMenu()
 
 					if ((Context->NumSelectedItems == 1) && (Context->NumSelectedFolders == 1))
 					{
-						SceneOutliner->GetTree().GetSelectedItems()[0]->GenerateContextMenu(InMenu, *SceneOutliner);
-
 						Section.AddMenuEntry(
 							"MakeCurrentFolder",
 							LOCTEXT("MakeCurrentFolder", "Make Current Folder"),
