Refactor RegisterContextMenu to avoid confusion
#rnx

diff --git a/Engine/Source/Editor/SceneOutliner/Private/ActorBrowsingMode.cpp b/Engine/Source/Editor/SceneOutliner/Private/ActorBrowsingMode.cpp
--- a/Engine/Source/Editor/SceneOutliner/Private/ActorBrowsingMode.cpp
+++ b/Engine/Source/Editor/SceneOutliner/Private/ActorBrowsingMode.cpp
@@ -583,200 +583,192 @@ void FActorBrowsingMode::RegisterContextMenu()
 	if (!ToolMenus->IsMenuRegistered(DefaultContextBaseMenuName))
 	{
 		UToolMenu* Menu = ToolMenus->RegisterMenu(DefaultContextBaseMenuName);
+		Menu->AddDynamicSection("DynamicHierarchySection", FNewToolMenuDelegate::CreateStatic(&FActorBrowsingMode::FillDefaultContextBaseMenu));
+	}
 
-		Menu->AddDynamicSection("DynamicHierarchySection", FNewToolMenuDelegate::CreateLambda([](UToolMenu* InMenu)
+	if (!ToolMenus->IsMenuRegistered(DefaultContextMenuName))
+	{
+		ToolMenus->RegisterMenu(DefaultContextMenuName, DefaultContextBaseMenuName);
+	}
+}
+
+void FActorBrowsingMode::FillDefaultContextBaseMenu(UToolMenu* InMenu)
+{
+	USceneOutlinerMenuContext* Context = InMenu->FindContext<USceneOutlinerMenuContext>();
+	if (!Context)
+	{
+		return;
+	}
+
+	TSharedPtr<SSceneOutliner> SharedOutliner = Context->SceneOutliner.Pin();
+	SSceneOutliner* SceneOutliner = SharedOutliner.Get();
+	if (!SceneOutliner)
+	{
+		return;
+	}
+
+	{
+		// NOTE: the name "Section" is used in many other places
+		FToolMenuSection& Section = InMenu->FindOrAddSection("Section");
+		Section.Label = LOCTEXT("HierarchySectionName", "Hierarchy");
+
+		if (Context->bShowParentTree)
 		{
-			if(USceneOutlinerMenuContext* Context = InMenu->FindContext<USceneOutlinerMenuContext>())
+			if (Context->NumSelectedItems == 0)
+			{
+				FSceneOutlinerMenuHelper::AddMenuEntryCreateFolder(Section, *SceneOutliner);
+			}
+			else
 			{
-				if (TSharedPtr<SSceneOutliner> SharedOutliner = Context->SceneOutliner.Pin())
+				if (Context->NumSelectedItems == 1)
 				{
-					SSceneOutliner* SceneOutliner = SharedOutliner.Get();
-					// NOTE: the name "Section" is used in many other places
-					FToolMenuSection& Section = InMenu->FindOrAddSection("Section");
-					Section.Label = LOCTEXT("HierarchySectionName", "Hierarchy");
-
-					if (Context->bShowParentTree)
-					{
-						if (Context->NumSelectedItems == 0)
-						{
-							FSceneOutlinerMenuHelper::AddMenuEntryCreateFolder(Section, *SceneOutliner);
-						}
-						else
-						{
-							if (Context->NumSelectedItems == 1)
-							{
-								SceneOutliner->GetTree().GetSelectedItems()[0]->GenerateContextMenu(InMenu, *SceneOutliner);
-							}
+					SceneOutliner->GetTree().GetSelectedItems()[0]->GenerateContextMenu(InMenu, *SceneOutliner);
+				}
 
-							// If we've only got folders selected, show the selection and edit sub menus
-							if (Context->NumSelectedItems > 0 && Context->NumSelectedFolders == Context->NumSelectedItems)
-							{
-								Section.AddSubMenu(
-									"SelectSubMenu",
-									LOCTEXT("SelectSubmenu", "Select"),
-									LOCTEXT("SelectSubmenu_Tooltip", "Select the contents of the current selection"),
-									FNewToolMenuDelegate::CreateSP(SceneOutliner, &SSceneOutliner::FillSelectionSubMenu));
-							}
-						}
-					}
+				// If we've only got folders selected, show the selection and edit sub menus
+				if (Context->NumSelectedItems > 0 && Context->NumSelectedFolders == Context->NumSelectedItems)
+				{
+					Section.AddSubMenu(
+						"SelectSubMenu",
+						LOCTEXT("SelectSubmenu", "Select"),
+						LOCTEXT("SelectSubmenu_Tooltip", "Select the contents of the current selection"),
+						FNewToolMenuDelegate::CreateSP(SceneOutliner, &SSceneOutliner::FillSelectionSubMenu));
 				}
 			}
-		}));
+		}
+	}
+	
+	{
+		// We always create a section here, even if there is no parent so that clients can still extend the menu
+		FToolMenuSection& MainSection = InMenu->AddSection("MainSection", LOCTEXT("OutlinerSectionName", "Outliner"));
 
-		Menu->AddDynamicSection("DynamicMainSection", FNewToolMenuDelegate::CreateLambda([](UToolMenu* InMenu)
+		// Don't add any of these menu items if we're not showing the parent tree
+		// Can't move worlds or level blueprints
+		if (Context->bShowParentTree && Context->NumSelectedItems > 0 && Context->NumWorldsSelected == 0)
 		{
-			// We always create a section here, even if there is no parent so that clients can still extend the menu
-			FToolMenuSection& Section = InMenu->AddSection("MainSection", LOCTEXT("OutlinerSectionName", "Outliner"));
+			MainSection.AddSubMenu(
+				"MoveActorsTo",
+				LOCTEXT("MoveActorsTo", "Move To"),
+				LOCTEXT("MoveActorsTo_Tooltip", "Move selection to another folder"),
+				FNewToolMenuDelegate::CreateSP(SceneOutliner, &SSceneOutliner::FillFoldersSubMenu));
+		}
 
-			if (USceneOutlinerMenuContext* Context = InMenu->FindContext<USceneOutlinerMenuContext>())
+		if (Context->bShowParentTree && Context->NumSelectedItems > 0)
+		{
+			// Only add the menu option to wp levels
+			if (!Context->bRepresentingGameWorld && Context->bRepresentingPartitionedWorld)
 			{
-				if (TSharedPtr<SSceneOutliner> SharedOutliner = Context->SceneOutliner.Pin())
+				MainSection.AddMenuEntry(
+					"PinItems",
+					LOCTEXT("Pin", "Pin"),
+					LOCTEXT("PinTooltip", "Keep the selected items loaded in the editor even when they don't overlap a loaded World Partition region"),
+					FSlateIcon(),
+					FUIAction(
+						FExecuteAction::CreateSP(SceneOutliner, &SSceneOutliner::PinSelectedItems),
+						FCanExecuteAction::CreateLambda([SceneOutliner, Context]()
 				{
-					SSceneOutliner* SceneOutliner = SharedOutliner.Get();
-
-					// Don't add any of these menu items if we're not showing the parent tree
-					// Can't move worlds or level blueprints
-					if (Context->bShowParentTree && Context->NumSelectedItems > 0 && Context->NumWorldsSelected == 0 && Context->SceneOutliner.IsValid())
+					if (Context->NumPinnedItems != Context->NumSelectedItems || Context->NumSelectedFolders > 0)
 					{
-						Section.AddSubMenu(
-							"MoveActorsTo",
-							LOCTEXT("MoveActorsTo", "Move To"),
-							LOCTEXT("MoveActorsTo_Tooltip", "Move selection to another folder"),
-							FNewToolMenuDelegate::CreateSP(Context->SceneOutliner.Pin().Get(), &SSceneOutliner::FillFoldersSubMenu));
+						return SceneOutliner->CanPinSelectedItems();
 					}
-
-					if (Context->bShowParentTree && Context->NumSelectedItems > 0 && Context->SceneOutliner.IsValid())
+					return false;
+				})));
+
+				MainSection.AddMenuEntry(
+					"UnpinItems",
+					LOCTEXT("Unpin", "Unpin"),
+					LOCTEXT("UnpinTooltip", "Allow the World Partition system to load and unload the selected items automatically"),
+					FSlateIcon(),
+					FUIAction(
+						FExecuteAction::CreateSP(SceneOutliner, &SSceneOutliner::UnpinSelectedItems),
+						FCanExecuteAction::CreateLambda([SceneOutliner, Context]()
+				{
+					if (Context->NumPinnedItems != 0 || Context->NumSelectedFolders > 0)
 					{
-						// Only add the menu option to wp levels
-						if (!Context->bRepresentingGameWorld && Context->bRepresentingPartitionedWorld)
-						{
-							Section.AddMenuEntry(
-								"PinItems",
-								LOCTEXT("Pin", "Pin"),
-								LOCTEXT("PinTooltip", "Keep the selected items loaded in the editor even when they don't overlap a loaded World Partition region"),
-								FSlateIcon(),
-								FUIAction(
-									FExecuteAction::CreateSP(SceneOutliner, &SSceneOutliner::PinSelectedItems),
-									FCanExecuteAction::CreateLambda([SceneOutliner, Context]()
-									{
-										if(Context->NumPinnedItems != Context->NumSelectedItems || Context->NumSelectedFolders > 0)
-										{
-											return SceneOutliner->CanPinSelectedItems();
-										}
-										return false;
-									})));
-
-							Section.AddMenuEntry(
-								"UnpinItems",
-								LOCTEXT("Unpin", "Unpin"),
-								LOCTEXT("UnpinTooltip", "Allow the World Partition system to load and unload the selected items automatically"),
-								FSlateIcon(),
-								FUIAction(
-									FExecuteAction::CreateSP(SceneOutliner, &SSceneOutliner::UnpinSelectedItems),
-									FCanExecuteAction::CreateLambda([SceneOutliner, Context]()
-									{
-										if (Context->NumPinnedItems != 0 || Context->NumSelectedFolders > 0)
-										{
-											return SceneOutliner->CanUnpinSelectedItems();
-										}
-										return false;
-									})));
-							
-						}
+						return SceneOutliner->CanUnpinSelectedItems();
 					}
+					return false;
+				})));
 
-					if (Context->NumSelectedItems > 0)
-					{
-						SceneOutliner->AddSourceControlMenuOptions(InMenu);
-					}
-				}
 			}
-		}));
+		}
 
-		Menu->AddDynamicSection("DynamicActorEditorContext", FNewToolMenuDelegate::CreateLambda([](UToolMenu* InMenu)
+		if (Context->NumSelectedItems > 0)
 		{
-			USceneOutlinerMenuContext* Context = InMenu->FindContext<USceneOutlinerMenuContext>();
-			if (Context && Context->bShowParentTree)
-			{
-				if (TSharedPtr<SSceneOutliner> SharedOutliner = Context->SceneOutliner.Pin())
-				{
-					FToolMenuSection& Section = InMenu->AddSection("ActorEditorContextSection", LOCTEXT("ActorEditorContextSectionName", "Actor Editor Context"));
-					SSceneOutliner* SceneOutliner = SharedOutliner.Get();
-
-					if ((Context->NumSelectedItems == 1) && (Context->NumSelectedFolders == 1))
-					{
-						Section.AddMenuEntry(
-							"MakeCurrentFolder",
-							LOCTEXT("MakeCurrentFolder", "Make Current Folder"),
-							FText(),
-							FSlateIcon(),
-							FUIAction(
-								FExecuteAction::CreateLambda([SceneOutliner]()
-								{
-									const FSceneOutlinerItemSelection& Selection = SceneOutliner->GetSelection();
-									if (Selection.SelectedItems.Num() == 1)
-									{
-										FSceneOutlinerTreeItemPtr Item = Selection.SelectedItems[0].Pin();
-										if (FActorFolderTreeItem* FolderItem = Item->CastTo<FActorFolderTreeItem>())
-										{
-											if (FolderItem->World.IsValid())
-											{
-												const FScopedTransaction Transaction(NSLOCTEXT("UnrealEd", "MakeCurrentActorFolder", "Make Current Actor Folder"));
-												FActorFolders::Get().SetActorEditorContextFolder(*FolderItem->World, FolderItem->GetFolder());
-											}
-										}
-									}
-								}),
-								FCanExecuteAction::CreateLambda([SceneOutliner]
-								{
-									const FSceneOutlinerItemSelection& Selection = SceneOutliner->GetSelection();
-									if (Selection.SelectedItems.Num() == 1)
-									{
-										FSceneOutlinerTreeItemPtr Item = Selection.SelectedItems[0].Pin();
-										if (FActorFolderTreeItem* FolderItem = Item->CastTo<FActorFolderTreeItem>())
-										{
-											return FolderItem->World.IsValid() &&
-												(FolderItem->World->GetCurrentLevel() == FolderItem->GetFolder().GetRootObjectAssociatedLevel()) &&
-												(FActorFolders::Get().GetActorEditorContextFolder(*FolderItem->World) != FolderItem->GetFolder());
-										}
-									}
-									return false;
-								})
-							)
-						);
-					}
+			SceneOutliner->AddSourceControlMenuOptions(InMenu);
+		}
+	}
 
-					const FActorBrowsingMode* Mode = static_cast<const FActorBrowsingMode*>(SceneOutliner->GetMode());
-					check(Mode);
+	if (Context->bShowParentTree)
+	{
+		FToolMenuSection& ActorEditorContextSection = InMenu->AddSection("ActorEditorContextSection", LOCTEXT("ActorEditorContextSectionName", "Actor Editor Context"));
 
-					Section.AddMenuEntry(
-						"ClearCurrentFolder",
-						LOCTEXT("ClearCurrentFolder", "Clear Current Folder"),
-						FText(),
-						FSlateIcon(),
-						FUIAction(
-							FExecuteAction::CreateLambda([Mode]()
+		if ((Context->NumSelectedItems == 1) && (Context->NumSelectedFolders == 1))
+		{
+			ActorEditorContextSection.AddMenuEntry(
+				"MakeCurrentFolder",
+				LOCTEXT("MakeCurrentFolder", "Make Current Folder"),
+				FText(),
+				FSlateIcon(),
+				FUIAction(
+					FExecuteAction::CreateLambda([SceneOutliner]()
+					{
+						const FSceneOutlinerItemSelection& Selection = SceneOutliner->GetSelection();
+						if (Selection.SelectedItems.Num() == 1)
+						{
+							FSceneOutlinerTreeItemPtr Item = Selection.SelectedItems[0].Pin();
+							if (FActorFolderTreeItem* FolderItem = Item->CastTo<FActorFolderTreeItem>())
 							{
-								if (Mode->RepresentingWorld.IsValid())
+								if (FolderItem->World.IsValid())
 								{
-									const FScopedTransaction Transaction(NSLOCTEXT("UnrealEd", "ClearCurrentActorFolder", "Clear Current Actor Folder"));
-									FActorFolders::Get().SetActorEditorContextFolder(*Mode->RepresentingWorld.Get(), FFolder::GetWorldRootFolder(Mode->RepresentingWorld.Get()));
+									const FScopedTransaction Transaction(NSLOCTEXT("UnrealEd", "MakeCurrentActorFolder", "Make Current Actor Folder"));
+									FActorFolders::Get().SetActorEditorContextFolder(*FolderItem->World, FolderItem->GetFolder());
 								}
-							}),
-							FCanExecuteAction::CreateLambda([Mode]
+							}
+						}
+					}),
+					FCanExecuteAction::CreateLambda([SceneOutliner]
+					{
+						const FSceneOutlinerItemSelection& Selection = SceneOutliner->GetSelection();
+						if (Selection.SelectedItems.Num() == 1)
+						{
+							FSceneOutlinerTreeItemPtr Item = Selection.SelectedItems[0].Pin();
+							if (FActorFolderTreeItem* FolderItem = Item->CastTo<FActorFolderTreeItem>())
 							{
-								return Mode->RepresentingWorld.IsValid() && !FActorFolders::Get().GetActorEditorContextFolder(*Mode->RepresentingWorld.Get()).IsNone();
-							})
-						)
-					);
-				}
-			}
-		}));
-	}
+								return FolderItem->World.IsValid() &&
+									(FolderItem->World->GetCurrentLevel() == FolderItem->GetFolder().GetRootObjectAssociatedLevel()) &&
+									(FActorFolders::Get().GetActorEditorContextFolder(*FolderItem->World) != FolderItem->GetFolder());
+							}
+						}
+						return false;
+					})
+				)
+			);
+		}
 
-	if (!ToolMenus->IsMenuRegistered(DefaultContextMenuName))
-	{
-		ToolMenus->RegisterMenu(DefaultContextMenuName, DefaultContextBaseMenuName);
+		const FActorBrowsingMode* Mode = static_cast<const FActorBrowsingMode*>(SceneOutliner->GetMode());
+		check(Mode);
+
+		ActorEditorContextSection.AddMenuEntry(
+			"ClearCurrentFolder",
+			LOCTEXT("ClearCurrentFolder", "Clear Current Folder"),
+			FText(),
+			FSlateIcon(),
+			FUIAction(
+				FExecuteAction::CreateLambda([Mode]()
+				{
+					if (Mode->RepresentingWorld.IsValid())
+					{
+						const FScopedTransaction Transaction(NSLOCTEXT("UnrealEd", "ClearCurrentActorFolder", "Clear Current Actor Folder"));
+						FActorFolders::Get().SetActorEditorContextFolder(*Mode->RepresentingWorld.Get(), FFolder::GetWorldRootFolder(Mode->RepresentingWorld.Get()));
+					}
+				}),
+				FCanExecuteAction::CreateLambda([Mode]
+				{
+					return Mode->RepresentingWorld.IsValid() && !FActorFolders::Get().GetActorEditorContextFolder(*Mode->RepresentingWorld.Get()).IsNone();
+				})
+			)
+		);
 	}
 }
 
diff --git a/Engine/Source/Editor/SceneOutliner/Public/ActorBrowsingMode.h b/Engine/Source/Editor/SceneOutliner/Public/ActorBrowsingMode.h
--- a/Engine/Source/Editor/SceneOutliner/Public/ActorBrowsingMode.h
+++ b/Engine/Source/Editor/SceneOutliner/Public/ActorBrowsingMode.h
@@ -112,6 +112,7 @@ private:
 	TSharedPtr<SWidget> BuildContextMenu();
 	/** Register the context menu with the engine */
 	static void RegisterContextMenu();
+	static void FillDefaultContextBaseMenu(UToolMenu* InMenu);
 	bool CanPasteFoldersOnlyFromClipboard() const;
 
 	bool GetFolderNamesFromPayload(const FSceneOutlinerDragDropPayload& InPayload, TArray<FName>& OutFolders, FFolder::FRootObject& OutCommonRootObject) const;
