PR #10203: Bugfix scene outliner repair errors (Contributed by [REDACTED_PERSON])

diff --git a/Engine/Source/Editor/SceneOutliner/Private/ActorBrowsingMode.cpp b/Engine/Source/Editor/SceneOutliner/Private/ActorBrowsingMode.cpp
--- a/Engine/Source/Editor/SceneOutliner/Private/ActorBrowsingMode.cpp
+++ b/Engine/Source/Editor/SceneOutliner/Private/ActorBrowsingMode.cpp
@@ -1255,15 +1255,12 @@ void FActorBrowsingMode::RepairErrors() const
 			ISourceControlModule& SourceControlModule = ISourceControlModule::Get();
 			ISourceControlProvider& SourceControlProvider = SourceControlModule.GetProvider();
 
-			TArray<FWorldPartitionActorDesc*> InvalidActorDescs;
-			WorldPartition->ForEachActorDescContainer([&InvalidActorDescs](UActorDescContainer* ActorDescContainer)
+			TArray<FAssetData> InvalidActorAssets;
+			WorldPartition->ForEachActorDescContainer([&InvalidActorAssets](UActorDescContainer* ActorDescContainer)
 			{
-				for (const TUniquePtr<FWorldPartitionActorDesc>& InvalidActor : ActorDescContainer->GetInvalidActors())
+				for (const FAssetData& InvalidActor : ActorDescContainer->GetInvalidActors())
 				{
-					if (FWorldPartitionActorDesc* InvalidActorDesc = InvalidActor.Get())
-					{
-						InvalidActorDescs.Add(InvalidActorDesc);
-					}
+					InvalidActorAssets.Add(InvalidActor);
 				}
 				ActorDescContainer->ClearInvalidActors();
 			});
@@ -1271,13 +1268,13 @@ void FActorBrowsingMode::RepairErrors() const
 			TArray<FString> ActorFilesToDelete;
 			TArray<FString> ActorFilesToRevert;
 			{
-				FScopedSlowTask SlowTask(InvalidActorDescs.Num(), LOCTEXT("UpdatingSourceControlStatus", "Updating revision control status..."));
+				FScopedSlowTask SlowTask(InvalidActorAssets.Num(), LOCTEXT("UpdatingSourceControlStatus", "Updating source control status..."));
 				SlowTask.MakeDialogDelayed(1.0f);
 
-				for (const FWorldPartitionActorDesc* InvalidActorDesc : InvalidActorDescs)
+				for (const FAssetData& InvalidActorAsset : InvalidActorAssets)
 				{
 					FPackagePath PackagePath;
-					if (FPackagePath::TryFromPackageName(InvalidActorDesc->GetActorPackage(), PackagePath))
+					if (FPackagePath::TryFromPackageName(InvalidActorAsset.PackageName, PackagePath))
 					{
 						const FString ActorFile = PackagePath.GetLocalFullPath();
 						FSourceControlStatePtr SCState = SourceControlProvider.GetState(ActorFile, EStateCacheUsage::ForceUpdate);
diff --git a/Engine/Source/Runtime/Engine/Private/WorldPartition/ActorDescContainer.cpp b/Engine/Source/Runtime/Engine/Private/WorldPartition/ActorDescContainer.cpp
--- a/Engine/Source/Runtime/Engine/Private/WorldPartition/ActorDescContainer.cpp
+++ b/Engine/Source/Runtime/Engine/Private/WorldPartition/ActorDescContainer.cpp
@@ -90,7 +90,7 @@ void UActorDescContainer::Initialize(const FInitializeParams& InitParams)
 				FActorDescList::GetActorDesc(ActorDesc->GetGuid()) ||
 				(InitParams.FilterActorDesc && !InitParams.FilterActorDesc(ActorDesc.Get())))
 			{
-				InvalidActors.Emplace(MoveTemp(ActorDesc));
+				InvalidActors.Emplace(Asset);
 				continue;
 			}
 
diff --git a/Engine/Source/Runtime/Engine/Public/WorldPartition/ActorDescContainer.h b/Engine/Source/Runtime/Engine/Public/WorldPartition/ActorDescContainer.h
--- a/Engine/Source/Runtime/Engine/Public/WorldPartition/ActorDescContainer.h
+++ b/Engine/Source/Runtime/Engine/Public/WorldPartition/ActorDescContainer.h
@@ -87,7 +87,7 @@ public:
 	const FTransform& GetInstanceTransform() const;
 
 	bool HasInvalidActors() const { return InvalidActors.Num() > 0; }
-	const TArray<TUniquePtr<FWorldPartitionActorDesc>>& GetInvalidActors() const { return InvalidActors; }
+	const TArray<FAssetData>& GetInvalidActors() const { return InvalidActors; }
 	void ClearInvalidActors() { InvalidActors.Empty(); }
 
 	void RegisterActorDescriptor(FWorldPartitionActorDesc* ActorDesc, UWorld* InWorldContext);
@@ -108,7 +108,7 @@ public:
 	FName ContainerPackageName;
 	FGuid ContentBundleGuid;
 
-	TArray<TUniquePtr<FWorldPartitionActorDesc>> InvalidActors;
+	TArray<FAssetData> InvalidActors;
 
 protected:
 	TNameActorDescMap ActorsByName;
@@ -125,4 +125,4 @@ private:
 	void RegisterEditorDelegates();
 	void UnregisterEditorDelegates();
 #endif
-};
\ No newline at end of file
+};
