SceneOutlinerTreeItemSCC: Add an early exit to HandleSourceControlProviderChanged when the engine is shutting down to prevent crashes due to lingering items

diff --git a/Engine/Source/Editor/SceneOutliner/Private/SceneOutlinerTreeItemSCC.cpp b/Engine/Source/Editor/SceneOutliner/Private/SceneOutlinerTreeItemSCC.cpp
--- a/Engine/Source/Editor/SceneOutliner/Private/SceneOutlinerTreeItemSCC.cpp
+++ b/Engine/Source/Editor/SceneOutliner/Private/SceneOutlinerTreeItemSCC.cpp
@@ -152,6 +152,14 @@ void FSceneOutlinerTreeItemSCC::HandleSourceControlStateChanged(EStateCacheUsage
 void FSceneOutlinerTreeItemSCC::HandleSourceControlProviderChanged(ISourceControlProvider& OldProvider, ISourceControlProvider& NewProvider)
 {
 	OldProvider.UnregisterSourceControlStateChanged_Handle(SourceControlStateChangedDelegateHandle);
+	
+	/* Early exit if the engine is shutting down, in case there are any lingering SCC items in the Outliner that haven't
+	 * been destroyed yet calling into modules that have been unloaded
+	 */
+	if (IsEngineExitRequested())
+	{
+		return;
+	}
 
 	SourceControlStateChangedDelegateHandle = NewProvider.RegisterSourceControlStateChanged_Handle(FSourceControlStateChanged::FDelegate::CreateLambda([this, WeakThis = AsWeak()]()
 	{
