Never show unloaded actors in the outliner when running PIE

diff --git a/Engine/Source/Editor/SceneOutliner/Private/ActorHierarchy.cpp b/Engine/Source/Editor/SceneOutliner/Private/ActorHierarchy.cpp
--- a/Engine/Source/Editor/SceneOutliner/Private/ActorHierarchy.cpp
+++ b/Engine/Source/Editor/SceneOutliner/Private/ActorHierarchy.cpp
@@ -378,6 +378,11 @@ void FActorHierarchy::CreateComponentItems(const AActor* Actor, TArray<FSceneOut
 	}
 }
 
+bool FActorHierarchy::IsShowingUnloadedActors() const
+{
+	return bShowingUnloadedActors && RepresentingWorld.IsValid() && !RepresentingWorld->IsPlayInEditor();
+}
+
 void FActorHierarchy::CreateWorldChildren(UWorld* World, TArray<FSceneOutlinerTreeItemPtr>& OutItems) const
 {
 	TRACE_CPUPROFILER_EVENT_SCOPE(FActorHierarchy::CreateWorldChildren);
@@ -439,7 +444,7 @@ void FActorHierarchy::CreateWorldChildren(UWorld* World, TArray<FSceneOutlinerTr
 		}
 	}
 
-	if (bShowingUnloadedActors)
+	if (IsShowingUnloadedActors())
 	{
 		if (UWorldPartitionSubsystem* WorldPartitionSubsystem = UWorld::GetSubsystem<UWorldPartitionSubsystem>(World))
 		{
@@ -660,7 +665,7 @@ void FActorHierarchy::OnLoadedActorRemoved(AActor& InActor)
 {
 	OnLevelActorDeleted(&InActor);
 
-	if (bShowingUnloadedActors)
+	if (IsShowingUnloadedActors())
 	{
 		if (UWorldPartition* WorldPartition = FWorldPartitionHelpers::GetWorldPartition(&InActor))
 		{
@@ -678,7 +683,7 @@ void FActorHierarchy::OnLoadedActorRemoved(AActor& InActor)
 
 void FActorHierarchy::OnActorDescAdded(FWorldPartitionActorDesc* ActorDesc)
 {
-	if (bShowingUnloadedActors && ActorDesc && !ActorDesc->IsLoaded(true) && FActorDescTreeItem::ShouldDisplayInOutliner(ActorDesc))
+	if (IsShowingUnloadedActors() && ActorDesc && !ActorDesc->IsLoaded(true) && FActorDescTreeItem::ShouldDisplayInOutliner(ActorDesc))
 	{
 		FSceneOutlinerHierarchyChangedData EventData;
 		EventData.Type = FSceneOutlinerHierarchyChangedData::Added;
@@ -689,7 +694,7 @@ void FActorHierarchy::OnActorDescAdded(FWorldPartitionActorDesc* ActorDesc)
 
 void FActorHierarchy::OnActorDescRemoved(FWorldPartitionActorDesc* ActorDesc)
 {
-	if (bShowingUnloadedActors && ActorDesc)
+	if (IsShowingUnloadedActors() && ActorDesc)
 	{
 		FSceneOutlinerHierarchyChangedData EventData;
 		EventData.Type = FSceneOutlinerHierarchyChangedData::Removed;
diff --git a/Engine/Source/Editor/SceneOutliner/Public/ActorHierarchy.h b/Engine/Source/Editor/SceneOutliner/Public/ActorHierarchy.h
--- a/Engine/Source/Editor/SceneOutliner/Public/ActorHierarchy.h
+++ b/Engine/Source/Editor/SceneOutliner/Public/ActorHierarchy.h
@@ -31,6 +31,7 @@ public:
 	void SetShowingEmptyFolders(bool bInShowingEmptyFolders) { bShowingEmptyFolders = bInShowingEmptyFolders; }
 
 private:
+	bool IsShowingUnloadedActors() const;
 	/** Adds all the direct and indirect children of a world to OutItems */
 	void CreateWorldChildren(UWorld* World, TArray<FSceneOutlinerTreeItemPtr>& OutItems) const;
 	/** Create all component items for an actor if we are showing components and place them in OutItems */
