Outliner: Update the unsaved column and unsaved and uncontrolled actors filters to use the externalpackagefilename correctly

diff --git a/Engine/Source/Editor/LevelEditor/Private/LevelEditorOutlinerSettings.cpp b/Engine/Source/Editor/LevelEditor/Private/LevelEditorOutlinerSettings.cpp
--- a/Engine/Source/Editor/LevelEditor/Private/LevelEditorOutlinerSettings.cpp
+++ b/Engine/Source/Editor/LevelEditor/Private/LevelEditorOutlinerSettings.cpp
@@ -28,6 +28,7 @@
 #include "Landscape.h"
 #include "LevelInstance/LevelInstanceActor.h"
 #include "Sound/AmbientSound.h"
+#include "SourceControlHelpers.h"
 
 #define LOCTEXT_NAMESPACE "LevelEditorOutlinerSettings"
 
@@ -294,16 +295,16 @@ void FLevelEditorOutlinerSettings::CreateDefaultFilters()
 
 bool FLevelEditorOutlinerSettings::DoesActorPassUnsavedFilter(const ISceneOutlinerTreeItem& InItem)
 {
-	return UnsavedAssets.Contains(SceneOutliner::FSceneOutlinerHelpers::GetExternalPackageName(InItem));
+	return UnsavedAssets.Contains(USourceControlHelpers::PackageFilename(SceneOutliner::FSceneOutlinerHelpers::GetExternalPackageName(InItem)));
 }
 
 bool FLevelEditorOutlinerSettings::DoesActorPassUncontrolledFilter(const ISceneOutlinerTreeItem& InItem)
 {
-	FString ExternalPackageName = SceneOutliner::FSceneOutlinerHelpers::GetExternalPackageName(InItem);
+	FString ExternalPackageFilename = USourceControlHelpers::PackageFilename(SceneOutliner::FSceneOutlinerHelpers::GetExternalPackageName(InItem));
 	
 	for (const TSharedRef<FUncontrolledChangelistState>& UncontrolledChangelistState : UncontrolledChangelistStates)
 	{
-		return UncontrolledChangelistState->GetFilenames().Contains(ExternalPackageName);
+		return UncontrolledChangelistState->GetFilenames().Contains(ExternalPackageFilename);
 	}
 	
 	return false;
diff --git a/Engine/Source/Editor/SceneOutliner/Private/SceneOutlinerUnsavedColumn.cpp b/Engine/Source/Editor/SceneOutliner/Private/SceneOutlinerUnsavedColumn.cpp
--- a/Engine/Source/Editor/SceneOutliner/Private/SceneOutlinerUnsavedColumn.cpp
+++ b/Engine/Source/Editor/SceneOutliner/Private/SceneOutlinerUnsavedColumn.cpp
@@ -21,7 +21,7 @@ public:
 	SLATE_END_ARGS()
 
 	/** Construct this widget */
-	void Construct(const FArguments& InArgs, const FString& InActorExternalPackageName);
+	void Construct(const FArguments& InArgs, const FString& InActorExternalPackageFilename);
 	
 	bool IsUnsaved() const;
 	
@@ -31,13 +31,13 @@ private:
 	void OnUnsavedAssetRemoved(const FString& FileAbsPathname);
 
 private:
-	FString ExternalPackageName;
+	FString ExternalPackageFilename;
 	bool bIsUnsaved;
 };
 
-void SUnsavedActorWidget::Construct(const FArguments& InArgs, const FString& InActorExternalPackageName)
+void SUnsavedActorWidget::Construct(const FArguments& InArgs, const FString& InActorExternalPackageFilename)
 {
-	ExternalPackageName = InActorExternalPackageName;
+	ExternalPackageFilename = InActorExternalPackageFilename;
 
 	SImage::Construct(
 			SImage::FArguments()
@@ -48,7 +48,7 @@ void SUnsavedActorWidget::Construct(const FArguments& InArgs, const FString& InA
 	UnsavedAssetsTrackerModule.OnUnsavedAssetAdded.AddSP(this, &SUnsavedActorWidget::OnUnsavedAssetAdded);
 	UnsavedAssetsTrackerModule.OnUnsavedAssetRemoved.AddSP(this, &SUnsavedActorWidget::OnUnsavedAssetRemoved);
 
-	bIsUnsaved = UnsavedAssetsTrackerModule.IsAssetUnsaved(ExternalPackageName);
+	bIsUnsaved = UnsavedAssetsTrackerModule.IsAssetUnsaved(ExternalPackageFilename);
 
 	UpdateImage();
 }
@@ -60,7 +60,7 @@ bool SUnsavedActorWidget::IsUnsaved() const
 
 void SUnsavedActorWidget::OnUnsavedAssetAdded(const FString& FileAbsPathname)
 {
-	if (FileAbsPathname == ExternalPackageName)
+	if (FileAbsPathname == ExternalPackageFilename)
 	{
 		// We should never be desynced, i.e if this item was added as an unsaved asset bIsUnsavedAsset MUST be false before
 		check(!bIsUnsaved)
@@ -72,7 +72,7 @@ void SUnsavedActorWidget::OnUnsavedAssetAdded(const FString& FileAbsPathname)
 
 void SUnsavedActorWidget::OnUnsavedAssetRemoved(const FString& FileAbsPathname)
 {
-	if (FileAbsPathname == ExternalPackageName)
+	if (FileAbsPathname == ExternalPackageFilename)
 	{
 		// We should never be desynced, i.e if this item was removed from the unsaved asset list bIsUnsavedAsset MUST be true before
 		check(bIsUnsaved)
@@ -118,7 +118,7 @@ SHeaderRow::FColumn::FArguments FSceneOutlinerActorUnsavedColumn::ConstructHeade
 
 const TSharedRef<SWidget> FSceneOutlinerActorUnsavedColumn::ConstructRowWidget(FSceneOutlinerTreeItemRef TreeItem, const STableRow<FSceneOutlinerTreeItemPtr>& Row)
 {
-	FString AssetPath = SceneOutliner::FSceneOutlinerHelpers::GetExternalPackageName(TreeItem.Get());
+	FString AssetPath =  USourceControlHelpers::PackageFilename(SceneOutliner::FSceneOutlinerHelpers::GetExternalPackageName(TreeItem.Get()));
 
 	if (AssetPath.IsEmpty())
 	{
@@ -146,7 +146,7 @@ void FSceneOutlinerActorUnsavedColumn::SortItems(TArray<FSceneOutlinerTreeItemPt
 		/** Sort by unsaved first */
 		.Primary([this](const ISceneOutlinerTreeItem& Item)
 		{
-			if (const TSharedRef<SUnsavedActorWidget>* FoundWidget = UnsavedActorWidgets.Find(SceneOutliner::FSceneOutlinerHelpers::GetExternalPackageName(Item)))
+			if (const TSharedRef<SUnsavedActorWidget>* FoundWidget = UnsavedActorWidgets.Find(USourceControlHelpers::PackageFilename(SceneOutliner::FSceneOutlinerHelpers::GetExternalPackageName(Item))))
 			{
 				// return the inverse because we want items with Unsaved = TRUE to show up on top in ascending sort (default)
 				return !(*FoundWidget)->IsUnsaved();
