Outliner: Fix an bug where column visibility changes would not persist properly

diff --git a/Engine/Source/Editor/SceneOutliner/Private/SSceneOutliner.cpp b/Engine/Source/Editor/SceneOutliner/Private/SSceneOutliner.cpp
--- a/Engine/Source/Editor/SceneOutliner/Private/SSceneOutliner.cpp
+++ b/Engine/Source/Editor/SceneOutliner/Private/SSceneOutliner.cpp
@@ -362,8 +362,8 @@ void SSceneOutliner::HandleHiddenColumnsChanged()
 	
 		for (const TPair<FName, TSharedPtr<ISceneOutlinerColumn>>& Pair : Columns)
 		{
-			const bool bWasColumnVisible = CacheHiddenColumns.Find(Pair.Key) != nullptr;
-			const bool bIsColumnVisible = HiddenColumns.Find(Pair.Key) != nullptr;
+			const bool bWasColumnVisible = CacheHiddenColumns.Find(Pair.Key) == nullptr;
+			const bool bIsColumnVisible = HiddenColumns.Find(Pair.Key) == nullptr;
 
 			// Only update column visibility if it changed
 			if (bWasColumnVisible != bIsColumnVisible)
