Outliner: Optimize AddColumn and RemoveColumn to stop re-creating all columns when called and only updating to add/remove the new column

diff --git a/Engine/Source/Editor/SceneOutliner/Private/SSceneOutliner.cpp b/Engine/Source/Editor/SceneOutliner/Private/SSceneOutliner.cpp
--- a/Engine/Source/Editor/SceneOutliner/Private/SSceneOutliner.cpp
+++ b/Engine/Source/Editor/SceneOutliner/Private/SSceneOutliner.cpp
@@ -125,7 +125,7 @@ void SSceneOutliner::Construct(const FArguments& InArgs, const FSceneOutlinerIni
 			.CanSelectGeneratedColumn(InInitOptions.bCanSelectGeneratedColumns)
 			.OnHiddenColumnsListChanged(this, &SSceneOutliner::HandleHiddenColumnsChanged);
 
-	SetupColumns(*HeaderRowWidget);
+	SetupColumns();
 
 	CacheHiddenColumns = TSet(HeaderRowWidget->GetHiddenColumnIds());
 
@@ -408,122 +408,157 @@ void SSceneOutliner::GetSortedColumnIDs(TArray<FName>& OutColumnIDs) const
 		});
 }
 
-void SSceneOutliner::SetupColumns(SHeaderRow& HeaderRow)
+void SSceneOutliner::AddColumn_Internal(const FName& ColumnId, const FSceneOutlinerColumnInfo& ColumnInfo, const TMap<FName, bool>& ColumnVisibilities, int32 InsertPosition)
 {
+	if(!HeaderRowWidget)
+	{
+		return;
+	}
+
+	SHeaderRow& HeaderRow = *HeaderRowWidget;
+
 	FSceneOutlinerModule& SceneOutlinerModule = FModuleManager::LoadModuleChecked<FSceneOutlinerModule>("SceneOutliner");
+	
+	// Avoid caching column visibilities while building the columns
+	bool const bPreviousShouldCacheColumnVisibility = bShouldCacheColumnVisibility;
+	bShouldCacheColumnVisibility = false;
+	
+	bool bIsVisible = true;
 
-	if (SharedData->ColumnMap.Num() == 0)
+	// If there is a config saved for this column, ignore the default visibility
+	if (const bool *ColumnVisibility = ColumnVisibilities.Find(ColumnId))
 	{
-		SharedData->UseDefaultColumns();
+		bIsVisible = *ColumnVisibility;
 	}
-
-	TMap<FName, FSceneOutlinerColumnInfo> FilteredColumnMap;
-	for (auto It(SharedData->ColumnMap.CreateIterator()); It; ++It)
+	else if (ColumnInfo.Visibility == ESceneOutlinerColumnVisibility::Invisible)
 	{
-		if (SceneOutlinerModule.GetColumnPermissionList()->PassesFilter(It.Key()))
-		{
-			FilteredColumnMap.Add(It.Key(), It.Value());
-		}
+		bIsVisible = false;
 	}
 
-	Columns.Empty(FilteredColumnMap.Num());
-	HeaderRow.ClearColumns();
-
-	TArray<FName> SortedIDs;
-	GetSortedColumnIDs(SortedIDs);
+	TSharedPtr<ISceneOutlinerColumn> Column;
 
-	TMap<FName, bool> ColumnVisibilities;
-	const FSceneOutlinerConfig* OutlinerConfig = GetConstConfig();
-
-	// Try to load visibility of columns from the config file
-	if (OutlinerConfig)
+	if (ColumnInfo.Factory.IsBound())
 	{
-		ColumnVisibilities = OutlinerConfig->ColumnVisibilities;
+		Column = ColumnInfo.Factory.Execute(*this);
+	}
+	else
+	{
+		Column = SceneOutlinerModule.FactoryColumn(ColumnId, *this);
 	}
 
-	// Avoid caching column visibilities while building the columns
-	bool const bPreviousShouldCacheColumnVisibility = bShouldCacheColumnVisibility;
-	bShouldCacheColumnVisibility = false;
-	
-	for (const FName& ID : SortedIDs)
+	if (ensure(Column.IsValid()))
 	{
-		bool bIsVisible = true;
+		Columns.Add(ColumnId, Column);
+
+		auto ColumnArgs = Column->ConstructHeaderRowColumn();
+
+		if (Column->SupportsSorting())
+		{
+			ColumnArgs
+				.SortMode(this, &SSceneOutliner::GetColumnSortMode, ColumnId)
+				.OnSort(this, &SSceneOutliner::OnColumnSortModeChanged);
+		}
 
-		// If there is a config saved for this column, ignore the default visibility
-		if (bool *ColumnVisibility = ColumnVisibilities.Find(ID))
+		if (ColumnInfo.ColumnLabel.IsSet())
 		{
-			bIsVisible = *ColumnVisibility;
+			ColumnArgs.DefaultLabel(ColumnInfo.ColumnLabel);
 		}
-		else if (FilteredColumnMap[ID].Visibility == ESceneOutlinerColumnVisibility::Invisible)
+		else
 		{
-			bIsVisible = false;
+			if (HeaderRow.GetVisibility() == EVisibility::Visible)
+			{
+				UE_LOG(LogSceneOutliner, Log, TEXT("Outliner Column %s does not have a localizable name, please specify one to FSceneOutlinerColumnInfo"), *ColumnId.ToString());
+			}
+				
+			ColumnArgs.DefaultLabel(FText::FromName(ColumnId));
+		}
+			
+		if (!ColumnInfo.bCanBeHidden)
+		{
+			ColumnArgs.ShouldGenerateWidget(true);
 		}
 
-		TSharedPtr<ISceneOutlinerColumn> Column;
+		if (ColumnInfo.FillSize.IsSet())
+		{
+			ColumnArgs.FillWidth(ColumnInfo.FillSize.GetValue());
+		}
 
-		if (FilteredColumnMap[ID].Factory.IsBound())
+		if (ColumnInfo.OnGetHeaderContextMenuContent.IsBound())
 		{
-			Column = FilteredColumnMap[ID].Factory.Execute(*this);
+			ColumnArgs.MenuContent()
+			[
+				ColumnInfo.OnGetHeaderContextMenuContent.Execute()
+			];
+		}
+
+		ColumnArgs.HeaderComboVisibility(ColumnInfo.HeaderComboVisibility);
+
+		if(InsertPosition == INDEX_NONE)
+		{
+			HeaderRow.AddColumn(ColumnArgs);
 		}
 		else
 		{
-			Column = SceneOutlinerModule.FactoryColumn(ID, *this);
+			HeaderRow.InsertColumn(ColumnArgs, InsertPosition);
 		}
+		
+		HeaderRow.SetShowGeneratedColumn(ColumnId, bIsVisible);
+	}
 
-		if (ensure(Column.IsValid()))
-		{
-			Columns.Add(ID, Column);
+	bShouldCacheColumnVisibility = bPreviousShouldCacheColumnVisibility;
+}
 
-			auto ColumnArgs = Column->ConstructHeaderRowColumn();
+void SSceneOutliner::RemoveColumn_Internal(const FName& ColumnId)
+{
+	Columns.Remove(ColumnId);
+	HeaderRowWidget->RemoveColumn(ColumnId);
+}
 
-			if (Column->SupportsSorting())
-			{
-				ColumnArgs
-					.SortMode(this, &SSceneOutliner::GetColumnSortMode, ID)
-					.OnSort(this, &SSceneOutliner::OnColumnSortModeChanged);
-			}
+void SSceneOutliner::SetupColumns()
+{
+	if(!HeaderRowWidget)
+	{
+		return;
+	}
+	
+	FSceneOutlinerModule& SceneOutlinerModule = FModuleManager::LoadModuleChecked<FSceneOutlinerModule>("SceneOutliner");
 
-			if (FilteredColumnMap[ID].ColumnLabel.IsSet())
-			{
-				ColumnArgs.DefaultLabel(FilteredColumnMap[ID].ColumnLabel);
-			}
-			else
-			{
-				if (HeaderRow.GetVisibility() == EVisibility::Visible)
-				{
-					UE_LOG(LogSceneOutliner, Log, TEXT("Outliner Column %s does not have a localizable name, please specify one to FSceneOutlinerColumnInfo"), *ID.ToString());
-				}
-				
-				ColumnArgs.DefaultLabel(FText::FromName(ID));
-			}
-			
-			if (!FilteredColumnMap[ID].bCanBeHidden)
-			{
-				ColumnArgs.ShouldGenerateWidget(true);
-			}
+	if (SharedData->ColumnMap.Num() == 0)
+	{
+		SharedData->UseDefaultColumns();
+	}
 
-			if (FilteredColumnMap[ID].FillSize.IsSet())
-			{
-				ColumnArgs.FillWidth(FilteredColumnMap[ID].FillSize.GetValue());
-			}
+	TMap<FName, FSceneOutlinerColumnInfo> FilteredColumnMap;
+	for (auto It(SharedData->ColumnMap.CreateIterator()); It; ++It)
+	{
+		if (SceneOutlinerModule.GetColumnPermissionList()->PassesFilter(It.Key()))
+		{
+			FilteredColumnMap.Add(It.Key(), It.Value());
+		}
+	}
 
-			if (FilteredColumnMap[ID].OnGetHeaderContextMenuContent.IsBound())
-			{
-				ColumnArgs.MenuContent()
-				[
-					FilteredColumnMap[ID].OnGetHeaderContextMenuContent.Execute()
-				];
-			}
+	Columns.Empty(FilteredColumnMap.Num());
+	HeaderRowWidget->ClearColumns();
 
-			ColumnArgs.HeaderComboVisibility(FilteredColumnMap[ID].HeaderComboVisibility);
+	TArray<FName> SortedIDs;
+	GetSortedColumnIDs(SortedIDs);
 
-			HeaderRow.AddColumn(ColumnArgs);
-			HeaderRow.SetShowGeneratedColumn(ID, bIsVisible);
-		}
+	TMap<FName, bool> ColumnVisibilities;
+	const FSceneOutlinerConfig* OutlinerConfig = GetConstConfig();
+
+	// Try to load visibility of columns from the config file
+	if (OutlinerConfig)
+	{
+		ColumnVisibilities = OutlinerConfig->ColumnVisibilities;
+	}
+
+	for (const FName& ID : SortedIDs)
+	{
+		AddColumn_Internal(ID, FilteredColumnMap[ID], ColumnVisibilities);
 	}
+	
 	Columns.Shrink();
 	bNeedsColumRefresh = false;
-	bShouldCacheColumnVisibility = bPreviousShouldCacheColumnVisibility;
 }
 
 void SSceneOutliner::RefreshColumns()
@@ -1158,21 +1193,33 @@ int32 SSceneOutliner::GetFilterCount() const
 	return Filters->Num();
 }
 
-void SSceneOutliner::AddColumn(FName ColumId, const FSceneOutlinerColumnInfo& ColumInfo)
+void SSceneOutliner::AddColumn(FName ColumnId, const FSceneOutlinerColumnInfo& ColumInfo)
 {
-	if (!SharedData->ColumnMap.Contains(ColumId))
+	if (!SharedData->ColumnMap.Contains(ColumnId))
 	{
-		SharedData->ColumnMap.Add(ColumId, ColumInfo);
-		RefreshColumns();
+		SharedData->ColumnMap.Add(ColumnId, ColumInfo);
+
+		// Get the new sorted list of columns to make sure this is added in the right position
+		TArray<FName> SortedColumnIDs;
+		GetSortedColumnIDs(SortedColumnIDs);
+
+		TMap<FName, bool> ColumnVisibilities;
+		const FSceneOutlinerConfig* OutlinerConfig = GetConstConfig();
+		if(OutlinerConfig)
+		{
+			ColumnVisibilities = OutlinerConfig->ColumnVisibilities;
+		}
+
+		AddColumn_Internal(ColumnId, ColumInfo, ColumnVisibilities, SortedColumnIDs.Find(ColumnId));
 	}
 }
 
-void SSceneOutliner::RemoveColumn(FName ColumId)
+void SSceneOutliner::RemoveColumn(FName ColumnId)
 {
-	if (SharedData->ColumnMap.Contains(ColumId))
+	if (SharedData->ColumnMap.Contains(ColumnId))
 	{
-		SharedData->ColumnMap.Remove(ColumId);
-		RefreshColumns();
+		SharedData->ColumnMap.Remove(ColumnId);
+		RemoveColumn_Internal(ColumnId);
 	}
 }
 
@@ -2227,7 +2274,7 @@ void SSceneOutliner::Tick(const FGeometry& AllottedGeometry, const double InCurr
 
 	if ( bNeedsColumRefresh )
 	{
-		SetupColumns(*HeaderRowWidget);
+		SetupColumns();
 	}
 
 	if( bNeedsRefresh )
diff --git a/Engine/Source/Editor/SceneOutliner/Public/SSceneOutliner.h b/Engine/Source/Editor/SceneOutliner/Public/SSceneOutliner.h
--- a/Engine/Source/Editor/SceneOutliner/Public/SSceneOutliner.h
+++ b/Engine/Source/Editor/SceneOutliner/Public/SSceneOutliner.h
@@ -336,13 +336,13 @@ public:
 	 * @param ColumnId The id of the column to add
 	 * @param ColumInfo The struct that contains the information on how to present and retrieve the column
 	 */
-	virtual void AddColumn(FName ColumId, const FSceneOutlinerColumnInfo& ColumInfo) override;
+	virtual void AddColumn(FName ColumnId, const FSceneOutlinerColumnInfo& ColumInfo) override;
 
 	/**
 	 * Remove a column of the scene outliner
 	 * @param ColumnId The name of the column to remove
 	 */
-	virtual void RemoveColumn(FName ColumId) override;
+	virtual void RemoveColumn(FName ColumnId) override;
 
 	void SetColumnVisibility(FName ColumnId, bool bIsVisible);
 
@@ -628,13 +628,16 @@ private:
 	TMap<FName, TSharedPtr<ISceneOutlinerColumn>> Columns;
 
 	/** Set up the columns required for this outliner */
-	void SetupColumns(SHeaderRow& HeaderRow);
+	void SetupColumns();
 
 	void HandleHiddenColumnsChanged();
 
-	/** Refresh the scene outliner for when a colum was added or removed */
+	/** Refresh the scene outliner columns */
 	void RefreshColumns();
 
+	void AddColumn_Internal(const FName& ColumnId, const FSceneOutlinerColumnInfo& ColumnInfo, const TMap<FName, bool>& ColumnVisibilities, int32 InsertPosition = INDEX_NONE);
+	void RemoveColumn_Internal(const FName& ColumnId);
+	
 	/** Populates OutSearchStrings with the strings associated with TreeItem that should be used in searching */
 	void PopulateSearchStrings( const ISceneOutlinerTreeItem& TreeItem, OUT TArray< FString >& OutSearchStrings ) const;
 
