SceneOutliner: Fixed 'unable to check out' warning message appearing when checking out a number of files where one of the files is already checked out (by you).

diff --git a/Engine/Source/Editor/SceneOutliner/Private/SceneOutlinerSCCHandler.cpp b/Engine/Source/Editor/SceneOutliner/Private/SceneOutlinerSCCHandler.cpp
--- a/Engine/Source/Editor/SceneOutliner/Private/SceneOutlinerSCCHandler.cpp
+++ b/Engine/Source/Editor/SceneOutliner/Private/SceneOutlinerSCCHandler.cpp
@@ -387,7 +387,7 @@ void FSceneOutlinerSCCHandler::ExecuteSCCCheckOut()
 
 	if ( PackagesToCheckOut.Num() > 0 )
 	{
-		FEditorFileUtils::CheckoutPackages(PackagesToCheckOut);
+		FEditorFileUtils::CheckoutPackages(PackagesToCheckOut, nullptr, /*bErrorIfAlreadyCheckedOut=*/false);
 	}
 }
 
