Filtering empty strings from PackageName requests to reduce output log spam.
Note this will not remove the `Exist called on PackageName that will always return false. Reason: Input '/Memory <snip> is in a memory-only mount point.` spam.

diff --git a/Engine/Source/Editor/SceneOutliner/Private/SceneOutlinerUnsavedColumn.cpp b/Engine/Source/Editor/SceneOutliner/Private/SceneOutlinerUnsavedColumn.cpp
--- a/Engine/Source/Editor/SceneOutliner/Private/SceneOutlinerUnsavedColumn.cpp
+++ b/Engine/Source/Editor/SceneOutliner/Private/SceneOutlinerUnsavedColumn.cpp
@@ -118,7 +118,14 @@ SHeaderRow::FColumn::FArguments FSceneOutlinerActorUnsavedColumn::ConstructHeade
 
 const TSharedRef<SWidget> FSceneOutlinerActorUnsavedColumn::ConstructRowWidget(FSceneOutlinerTreeItemRef TreeItem, const STableRow<FSceneOutlinerTreeItemPtr>& Row)
 {
-	FString AssetPath =  USourceControlHelpers::PackageFilename(SceneOutliner::FSceneOutlinerHelpers::GetExternalPackageName(TreeItem.Get()));
+	FString ItemPath = SceneOutliner::FSceneOutlinerHelpers::GetExternalPackageName(TreeItem.Get());
+
+	if (ItemPath.IsEmpty()) 
+	{
+		return SNullWidget::NullWidget;
+	}
+
+	FString AssetPath =  USourceControlHelpers::PackageFilename(ItemPath);
 
 	if (AssetPath.IsEmpty())
 	{
