Alternate fix for log spam
* Added memory package filtering in SourceControlHelpers.
* Added empty string filtering when sorting unsaved assets in the outliner.

diff --git a/Engine/Source/Developer/SourceControl/Private/SourceControlHelpers.cpp b/Engine/Source/Developer/SourceControl/Private/SourceControlHelpers.cpp
--- a/Engine/Source/Developer/SourceControl/Private/SourceControlHelpers.cpp
+++ b/Engine/Source/Developer/SourceControl/Private/SourceControlHelpers.cpp
@@ -1645,7 +1645,7 @@ static FString PackageFilename_Internal( const FString& InPackageName )
 	FString Filename = InPackageName;
 
 	// Get the filename by finding it on disk first
-	if ( !FPackageName::DoesPackageExist(InPackageName, &Filename) )
+	if ( !FPackageName::IsMemoryPackage(InPackageName) && !FPackageName::DoesPackageExist(InPackageName, &Filename) )
 	{
 		// The package does not exist on disk, see if we can find it in memory and predict the file extension
 		// Only do this if the supplied package name is valid
diff --git a/Engine/Source/Editor/SceneOutliner/Private/SceneOutlinerUnsavedColumn.cpp b/Engine/Source/Editor/SceneOutliner/Private/SceneOutlinerUnsavedColumn.cpp
--- a/Engine/Source/Editor/SceneOutliner/Private/SceneOutlinerUnsavedColumn.cpp
+++ b/Engine/Source/Editor/SceneOutliner/Private/SceneOutlinerUnsavedColumn.cpp
@@ -153,10 +153,14 @@ void FSceneOutlinerActorUnsavedColumn::SortItems(TArray<FSceneOutlinerTreeItemPt
 		/** Sort by unsaved first */
 		.Primary([this](const ISceneOutlinerTreeItem& Item)
 		{
-			if (const TSharedRef<SUnsavedActorWidget>* FoundWidget = UnsavedActorWidgets.Find(USourceControlHelpers::PackageFilename(SceneOutliner::FSceneOutlinerHelpers::GetExternalPackageName(Item))))
+			FString LongPackageName = SceneOutliner::FSceneOutlinerHelpers::GetExternalPackageName(Item);
+			if (!LongPackageName.IsEmpty())
 			{
-				// return the inverse because we want items with Unsaved = TRUE to show up on top in ascending sort (default)
-				return !(*FoundWidget)->IsUnsaved();
+				if (const TSharedRef<SUnsavedActorWidget>* FoundWidget = UnsavedActorWidgets.Find(USourceControlHelpers::PackageFilename(LongPackageName)))
+				{
+					// return the inverse because we want items with Unsaved = TRUE to show up on top in ascending sort (default)
+					return !(*FoundWidget)->IsUnsaved();
+				}
 			}
 			return true;
 		}, SortMode)
