Fixed mismatched copy constructor/assignment operator in FSceneOutlinerTreeItemType.

diff --git a/Engine/Source/Editor/SceneOutliner/Public/SceneOutlinerStandaloneTypes.h b/Engine/Source/Editor/SceneOutliner/Public/SceneOutlinerStandaloneTypes.h
--- a/Engine/Source/Editor/SceneOutliner/Public/SceneOutlinerStandaloneTypes.h
+++ b/Engine/Source/Editor/SceneOutliner/Public/SceneOutlinerStandaloneTypes.h
@@ -183,7 +183,6 @@ struct SCENEOUTLINER_API FSceneOutlinerTreeItemType
 {
 public:
 	explicit FSceneOutlinerTreeItemType(const FSceneOutlinerTreeItemType* Parent = nullptr) : ID(++NextUniqueID), ParentType(Parent) {}
-	FSceneOutlinerTreeItemType(const FSceneOutlinerTreeItemType& Src) : ID(Src.ID), ParentType(Src.ParentType) {}
 
 	bool operator==(const FSceneOutlinerTreeItemType& Other) const
 	{
