SceneOutliner: Maintain expansion state for all items (not just on repopulate)

- Fixes Enter/Exit PIE resetting expansion states

diff --git a/Engine/Plugins/Editor/ObjectMixer/ObjectMixer/Source/ObjectMixer/Private/Views/List/Modes/ObjectMixerOutlinerMode.cpp b/Engine/Plugins/Editor/ObjectMixer/ObjectMixer/Source/ObjectMixer/Private/Views/List/Modes/ObjectMixerOutlinerMode.cpp
--- a/Engine/Plugins/Editor/ObjectMixer/ObjectMixer/Source/ObjectMixer/Private/Views/List/Modes/ObjectMixerOutlinerMode.cpp
+++ b/Engine/Plugins/Editor/ObjectMixer/ObjectMixer/Source/ObjectMixer/Private/Views/List/Modes/ObjectMixerOutlinerMode.cpp
@@ -1759,13 +1759,6 @@ void FObjectMixerOutlinerMode::OnItemAdded(FSceneOutlinerTreeItemPtr Item)
 			}
 		}
 	}
-	else if (FActorFolderTreeItem* FolderItem = Item->CastTo<FActorFolderTreeItem>())
-	{
-		if (FolderItem->World.IsValid())
-		{
-			FolderItem->Flags.bIsExpanded = FActorFolders::Get().IsFolderExpanded(*FolderItem->World, FolderItem->GetFolder());
-		}
-	}
 	else if (Item->IsA<FActorDescTreeItem>())
 	{
 		if (!Item->Flags.bIsFilteredOut)
diff --git a/Engine/Source/Editor/SceneOutliner/Private/ActorBrowsingMode.cpp b/Engine/Source/Editor/SceneOutliner/Private/ActorBrowsingMode.cpp
--- a/Engine/Source/Editor/SceneOutliner/Private/ActorBrowsingMode.cpp
+++ b/Engine/Source/Editor/SceneOutliner/Private/ActorBrowsingMode.cpp
@@ -857,13 +857,6 @@ void FActorBrowsingMode::OnItemAdded(FSceneOutlinerTreeItemPtr Item)
 			}
 		}
 	}
-	else if (FActorFolderTreeItem* FolderItem = Item->CastTo<FActorFolderTreeItem>())
-	{
-		if (FolderItem->World.IsValid())
-		{
-			FolderItem->Flags.bIsExpanded = FActorFolders::Get().IsFolderExpanded(*FolderItem->World, FolderItem->GetFolder());
-		}
-	}
 	else if (Item->IsA<FActorDescTreeItem>())
 	{
 		if (!Item->Flags.bIsFilteredOut)
diff --git a/Engine/Source/Editor/SceneOutliner/Private/ActorFolderTreeItem.cpp b/Engine/Source/Editor/SceneOutliner/Private/ActorFolderTreeItem.cpp
--- a/Engine/Source/Editor/SceneOutliner/Private/ActorFolderTreeItem.cpp
+++ b/Engine/Source/Editor/SceneOutliner/Private/ActorFolderTreeItem.cpp
@@ -257,6 +257,11 @@ FActorFolderTreeItem::FActorFolderTreeItem(const FFolder& InFolder, const TWeakO
 	FFolderTreeItem::SetPath(InFolder.GetPath());
 	// Initialize ActorFolder
 	ActorFolder = InFolder.GetActorFolder();
+
+	if (World.IsValid())
+	{
+		Flags.bIsExpanded = FActorFolders::Get().IsFolderExpanded(*World, GetFolder());
+	}
 }
 
 void FActorFolderTreeItem::OnExpansionChanged()
diff --git a/Engine/Source/Editor/SceneOutliner/Private/SSceneOutliner.cpp b/Engine/Source/Editor/SceneOutliner/Private/SSceneOutliner.cpp
--- a/Engine/Source/Editor/SceneOutliner/Private/SSceneOutliner.cpp
+++ b/Engine/Source/Editor/SceneOutliner/Private/SSceneOutliner.cpp
@@ -688,12 +688,6 @@ void SSceneOutliner::Populate()
 	// Block events while we clear out the list
  	TGuardValue<bool> ReentrantGuard(bIsReentrant, true);
 
-	// Get a collection of items and folders which were formerly collapsed
-	if (CachedExpansionStateInfo.Num() == 0)
-	{
-		CachedExpansionStateInfo.Append(GetParentsExpansionState());
-	}
-
 	bool bMadeAnySignificantChanges = false;
 	if (bFullRefresh)
 	{
@@ -763,9 +757,8 @@ void SSceneOutliner::Populate()
 	bool bFinalSort = false;
 	if (PendingOperations.Num() == 0)
 	{
-		SetParentsExpansionState(CachedExpansionStateInfo);
-		CachedExpansionStateInfo.Empty();
-		
+		// Update expansion state based on item states
+		SetParentsExpansionState();
 		// When done processing a FullRefresh Scroll to First item in selection as it may have been
 		// scrolled out of view by the Refresh
 		if (bProcessingFullRefresh)
@@ -1050,38 +1043,19 @@ void SSceneOutliner::AddUnfilteredItemToTree(FSceneOutlinerTreeItemRef Item)
 		RootTreeItems.Add(Item);
 	}
 
+	Item->Flags.bIsExpanded = CachedExpansionStateInfo.FindOrAdd(Item->GetID(), Item->Flags.bIsExpanded);
+	
 	Mode->OnItemAdded(Item);
 }
 
-SSceneOutliner::FParentsExpansionState SSceneOutliner::GetParentsExpansionState() const
-{
-	FParentsExpansionState States;
-	for (const auto& Pair : TreeItemMap)
-	{
-		if (Pair.Value->GetChildren().Num())
-		{
-			States.Add(Pair.Key, Pair.Value->Flags.bIsExpanded);
-		}
-	}
-	return States;
-}
-
-void SSceneOutliner::SetParentsExpansionState(const FParentsExpansionState& ExpansionStateInfo) const
+void SSceneOutliner::SetParentsExpansionState() const
 {
 	for (const auto& Pair : TreeItemMap)
 	{
 		auto& Item = Pair.Value;
 		if (Item->GetChildren().Num())
 		{
-			const bool* bIsExpanded = ExpansionStateInfo.Find(Pair.Key);
-			if (bIsExpanded)
-			{
-				OutlinerTreeView->SetItemExpansion(Item, *bIsExpanded);
-			}
-			else
-			{
-				OutlinerTreeView->SetItemExpansion(Item, Item->Flags.bIsExpanded);
-			}
+			OutlinerTreeView->SetItemExpansion(Item, Item->Flags.bIsExpanded);
 		}
 	}
 }
@@ -1984,6 +1958,9 @@ void SSceneOutliner::OnItemExpansionChanged(FSceneOutlinerTreeItemPtr TreeItem,
 			OutlinerTreeView->SetItemExpansion(Child, true);
 		}
 	}
+
+	// Notify Mode
+	CachedExpansionStateInfo.Add(TreeItem->GetID(), TreeItem->Flags.bIsExpanded);
 }
 
 void SSceneOutliner::OnHierarchyChangedEvent(FSceneOutlinerHierarchyChangedData Event)
diff --git a/Engine/Source/Editor/SceneOutliner/Public/SSceneOutliner.h b/Engine/Source/Editor/SceneOutliner/Public/SSceneOutliner.h
--- a/Engine/Source/Editor/SceneOutliner/Public/SSceneOutliner.h
+++ b/Engine/Source/Editor/SceneOutliner/Public/SSceneOutliner.h
@@ -864,14 +864,11 @@ private:
 	/** Structure containing information relating to the expansion state of parent items in the tree */
 	typedef TMap<FSceneOutlinerTreeItemID, bool> FParentsExpansionState;
 	
-	/** Cached expansion state info, in case we need to process >500 items so we don't re-fetch from the partially rebuilt tree */
-	FParentsExpansionState CachedExpansionStateInfo;
+	/** Cached expansion state info */
+	mutable FParentsExpansionState CachedExpansionStateInfo;
 
-	/** Gets the current expansion state of parent items */
-	FParentsExpansionState GetParentsExpansionState() const;
-
-	/** Updates the expansion state of parent items after a repopulate, according to the previous state */
-	void SetParentsExpansionState(const FParentsExpansionState& ExpansionStateInfo) const;
+	/** Updates the expansion state of parent items after a repopulate */
+	void SetParentsExpansionState() const;
 
 private:
 
