SceneOutliner: Fix "Only in Current Level" filter
- Filter out folders from other levels
- Filter in currently edited Level Instance

diff --git a/Engine/Source/Editor/SceneOutliner/Private/ActorBrowsingMode.cpp b/Engine/Source/Editor/SceneOutliner/Private/ActorBrowsingMode.cpp
--- a/Engine/Source/Editor/SceneOutliner/Private/ActorBrowsingMode.cpp
+++ b/Engine/Source/Editor/SceneOutliner/Private/ActorBrowsingMode.cpp
@@ -489,13 +489,41 @@ TSharedRef<FSceneOutlinerFilter> FActorBrowsingMode::CreateHideTemporaryActorsFi
 
 TSharedRef<FSceneOutlinerFilter> FActorBrowsingMode::CreateIsInCurrentLevelFilter()
 {
-	return MakeShareable(new FActorFilter(FActorTreeItem::FFilterPredicate::CreateStatic([](const AActor* InActor)
+	return MakeShareable(new TSceneOutlinerPredicateFilter<ISceneOutlinerTreeItem>(ISceneOutlinerTreeItem::FFilterPredicate::CreateStatic([](const ISceneOutlinerTreeItem& InItem)
 		{
-			if (InActor->GetWorld())
+			if (const FActorTreeItem* ActorItem = InItem.CastTo<FActorTreeItem>())
 			{
-				return InActor->GetLevel() == InActor->GetWorld()->GetCurrentLevel();
-			}
+				if (const AActor* Actor = ActorItem->Actor.Get(); Actor && Actor->GetWorld())
+				{
+					if (const ILevelInstanceInterface* LevelInstance = Cast<ILevelInstanceInterface>(Actor))
+					{
+						if (LevelInstance->IsEditing())
+						{
+							return Actor->GetWorld()->GetCurrentLevel() == LevelInstance->GetLoadedLevel();
+						}
+					}
 
+					return Actor->GetLevel()->IsCurrentLevel();
+				}
+			}
+			else if (const FActorFolderTreeItem* FolderItem = InItem.CastTo<FActorFolderTreeItem>())
+			{
+				if (const UWorld* FolderWorld = FolderItem->World.Get())
+				{
+					if (ULevel* FolderLevel = FolderItem->GetFolder().GetRootObjectAssociatedLevel())
+					{
+						return FolderLevel->IsCurrentLevel();
+					}
+				}
+			}
+			else if (const FComponentTreeItem* ComponentItem = InItem.CastTo<FComponentTreeItem>())
+			{
+				if (const UActorComponent* Component = ComponentItem->Component.Get(); Component && Component->GetComponentLevel())
+				{
+					return Component->GetComponentLevel()->IsCurrentLevel();
+				}
+			}
+			
 			return false;
 		}), FSceneOutlinerFilter::EDefaultBehaviour::Pass));
 }
diff --git a/Engine/Source/Editor/SceneOutliner/Public/ISceneOutlinerTreeItem.h b/Engine/Source/Editor/SceneOutliner/Public/ISceneOutlinerTreeItem.h
--- a/Engine/Source/Editor/SceneOutliner/Public/ISceneOutlinerTreeItem.h
+++ b/Engine/Source/Editor/SceneOutliner/Public/ISceneOutlinerTreeItem.h
@@ -96,6 +96,19 @@ public:
 	{
 		return Children;
 	}
+
+	DECLARE_DELEGATE_RetVal_OneParam(bool, FFilterPredicate, const ISceneOutlinerTreeItem& Item);
+	DECLARE_DELEGATE_RetVal_OneParam(bool, FInteractivePredicate, const ISceneOutlinerTreeItem& Item);
+
+	bool Filter(FFilterPredicate Pred) const
+	{
+		return Pred.Execute(*this);
+	}
+
+	bool GetInteractiveState(FInteractivePredicate Pred) const
+	{
+		return Pred.Execute(*this);
+	}
 public:
 	/** Attempt to cast this item to another type if it is of that type. Returns nullptr if it fails */
 	template <typename T>
