Add a tooltip to the toggle visibility button in the scene outliner and the Pin button in the scene outliner.

diff --git a/Engine/Source/Editor/SceneOutliner/Private/SceneOutlinerGutter.cpp b/Engine/Source/Editor/SceneOutliner/Private/SceneOutlinerGutter.cpp
--- a/Engine/Source/Editor/SceneOutliner/Private/SceneOutlinerGutter.cpp
+++ b/Engine/Source/Editor/SceneOutliner/Private/SceneOutlinerGutter.cpp
@@ -346,6 +346,7 @@ const TSharedRef<SWidget> FSceneOutlinerGutter::ConstructRowWidget(FSceneOutline
 			.VAlign(VAlign_Center)
 			[
 				SNew(SVisibilityWidget, SharedThis(this), WeakOutliner, TreeItem, &Row)
+				.ToolTipText(LOCTEXT("SceneOutlinerVisibilityToggleTooltip", "Toggles the visibility of this object in the level editor"))
 			];
 	}
 	return SNullWidget::NullWidget;
diff --git a/Engine/Source/Editor/SceneOutliner/Private/SceneOutlinerPinnedColumn.cpp b/Engine/Source/Editor/SceneOutliner/Private/SceneOutlinerPinnedColumn.cpp
--- a/Engine/Source/Editor/SceneOutliner/Private/SceneOutlinerPinnedColumn.cpp
+++ b/Engine/Source/Editor/SceneOutliner/Private/SceneOutlinerPinnedColumn.cpp
@@ -206,6 +206,7 @@ const TSharedRef<SWidget> FSceneOutlinerPinnedColumn::ConstructRowWidget(FSceneO
 		.VAlign(VAlign_Center)
 		[
 			SNew(SPinnedWidget, WeakSceneOutliner, TreeItem, SharedThis(this), &Row)
+			.ToolTipText(LOCTEXT("SceneOutlinerPinnedWidgetTooltip", "Toggles if this object is pinned (always loaded) in the editor"))
 		];
 	}
 	return SNullWidget::NullWidget;
