update tooltip text to change the wording and add periods

#rnx

diff --git a/Engine/Source/Editor/SceneOutliner/Private/SceneOutlinerGutter.cpp b/Engine/Source/Editor/SceneOutliner/Private/SceneOutlinerGutter.cpp
--- a/Engine/Source/Editor/SceneOutliner/Private/SceneOutlinerGutter.cpp
+++ b/Engine/Source/Editor/SceneOutliner/Private/SceneOutlinerGutter.cpp
@@ -346,7 +346,7 @@ const TSharedRef<SWidget> FSceneOutlinerGutter::ConstructRowWidget(FSceneOutline
 			.VAlign(VAlign_Center)
 			[
 				SNew(SVisibilityWidget, SharedThis(this), WeakOutliner, TreeItem, &Row)
-				.ToolTipText(LOCTEXT("SceneOutlinerVisibilityToggleTooltip", "Toggles the visibility of this object in the level editor"))
+				.ToolTipText(LOCTEXT("SceneOutlinerVisibilityToggleTooltip", "Toggles the visibility of this object in the level editor."))
 			];
 	}
 	return SNullWidget::NullWidget;
diff --git a/Engine/Source/Editor/SceneOutliner/Private/SceneOutlinerPinnedColumn.cpp b/Engine/Source/Editor/SceneOutliner/Private/SceneOutlinerPinnedColumn.cpp
--- a/Engine/Source/Editor/SceneOutliner/Private/SceneOutlinerPinnedColumn.cpp
+++ b/Engine/Source/Editor/SceneOutliner/Private/SceneOutlinerPinnedColumn.cpp
@@ -206,7 +206,7 @@ const TSharedRef<SWidget> FSceneOutlinerPinnedColumn::ConstructRowWidget(FSceneO
 		.VAlign(VAlign_Center)
 		[
 			SNew(SPinnedWidget, WeakSceneOutliner, TreeItem, SharedThis(this), &Row)
-			.ToolTipText(LOCTEXT("SceneOutlinerPinnedWidgetTooltip", "Toggles if this object is pinned (always loaded) in the editor"))
+			.ToolTip(IDocumentation::Get()->CreateToolTip(LOCTEXT("SceneOutlinerPinnedWidgetTooltip", "Toggles whether this object is pinned (always loaded) in the editor."), nullptr, "Shared/MenuEntries/SceneOutliner_ActorBrowsingMode", "PinTooltip"))
 		];
 	}
 	return SNullWidget::NullWidget;
