LevelInstance: Fix perf issue with dirtyness check.

Remove '*' next to Level Instance display name until we implement a proper tracking that doesn't involve calling FEditorFileUtils::GetDirtyPackages which is slow

#rnx

diff --git a/Engine/Source/Editor/SceneOutliner/Private/ActorTreeItem.cpp b/Engine/Source/Editor/SceneOutliner/Private/ActorTreeItem.cpp
--- a/Engine/Source/Editor/SceneOutliner/Private/ActorTreeItem.cpp
+++ b/Engine/Source/Editor/SceneOutliner/Private/ActorTreeItem.cpp
@@ -114,9 +114,8 @@ private:
 			{
 				if (!bInEditingMode)
 				{
-					FText DirtySuffixText = LevelInstance->IsDirty() ? FText(LOCTEXT("IsDirtySuffix", "*")) : FText::GetEmpty();
 					FText IsCurrentSuffixText = LevelInstance->GetLoadedLevel() && LevelInstance->GetLoadedLevel()->IsCurrentLevel() ? FText(LOCTEXT("IsCurrentSuffix", " (Current)")) : FText::GetEmpty();
-					return FText::Format(LOCTEXT("LevelInstanceDisplay", "{0}{1}{2}"), FText::FromString(TreeItem->GetDisplayString()), DirtySuffixText, IsCurrentSuffixText);
+					return FText::Format(LOCTEXT("LevelInstanceDisplay", "{0}{1}"), FText::FromString(TreeItem->GetDisplayString()), IsCurrentSuffixText);
 				}
 			}
 			return FText::FromString(TreeItem->GetDisplayString());
