Outliner: Fix warning due to calling DoesPackageExist on an empty package name

diff --git a/Engine/Source/Editor/SceneOutliner/Private/SceneOutlinerUnsavedColumn.cpp b/Engine/Source/Editor/SceneOutliner/Private/SceneOutlinerUnsavedColumn.cpp
--- a/Engine/Source/Editor/SceneOutliner/Private/SceneOutlinerUnsavedColumn.cpp
+++ b/Engine/Source/Editor/SceneOutliner/Private/SceneOutlinerUnsavedColumn.cpp
@@ -118,24 +118,18 @@ SHeaderRow::FColumn::FArguments FSceneOutlinerActorUnsavedColumn::ConstructHeade
 
 const TSharedRef<SWidget> FSceneOutlinerActorUnsavedColumn::ConstructRowWidget(FSceneOutlinerTreeItemRef TreeItem, const STableRow<FSceneOutlinerTreeItemPtr>& Row)
 {
-	FString ItemPath = SceneOutliner::FSceneOutlinerHelpers::GetExternalPackageName(TreeItem.Get());
+	const FString ExternalPackageName = SceneOutliner::FSceneOutlinerHelpers::GetExternalPackageName(TreeItem.Get());
+	const FString ExternalPackageFileName = !ExternalPackageName.IsEmpty() ? USourceControlHelpers::PackageFilename(ExternalPackageName) : FString();
 
-	if (ItemPath.IsEmpty()) 
+	if (ExternalPackageFileName.IsEmpty())
 	{
 		return SNullWidget::NullWidget;
 	}
 
-	FString AssetPath =  USourceControlHelpers::PackageFilename(ItemPath);
-
-	if (AssetPath.IsEmpty())
-	{
-		return SNullWidget::NullWidget;
-	}
-
-	TSharedRef<SUnsavedActorWidget> ActualWidget = SNew(SUnsavedActorWidget, AssetPath);
+	const TSharedRef<SUnsavedActorWidget> ActualWidget = SNew(SUnsavedActorWidget, ExternalPackageFileName);
 
 	// Store a reference to the widget for this asset so we can grab the unsaved status for sorting
-	UnsavedActorWidgets.Add(AssetPath, ActualWidget);
+	UnsavedActorWidgets.Add(ExternalPackageFileName, ActualWidget);
 	
 	return SNew(SHorizontalBox)
 			+ SHorizontalBox::Slot()
