TEDS-Outliner Integration: Add initial support for populating the Outliner using TEDS queries.

diff --git a/Engine/Plugins/Experimental/TypedElementsDataStorage/Source/TypedElementsDataStorageUI/Private/TypedElementSceneOutliner.cpp b/Engine/Plugins/Experimental/TypedElementsDataStorage/Source/TypedElementsDataStorageUI/Private/TypedElementSceneOutliner.cpp
deleted file mode 100644
diff --git a/Engine/Plugins/Experimental/TypedElementsDataStorage/Source/TypedElementsDataStorageUI/Public/TypedElementSceneOutliner.h b/Engine/Plugins/Experimental/TypedElementsDataStorage/Source/TypedElementsDataStorageUI/Public/TypedElementSceneOutliner.h
deleted file mode 100644
diff --git a/Engine/Source/Editor/SceneOutliner/Private/ActorHierarchy.cpp b/Engine/Source/Editor/SceneOutliner/Private/ActorHierarchy.cpp
--- a/Engine/Source/Editor/SceneOutliner/Private/ActorHierarchy.cpp
+++ b/Engine/Source/Editor/SceneOutliner/Private/ActorHierarchy.cpp
@@ -33,6 +33,13 @@ TUniquePtr<FActorHierarchy> FActorHierarchy::Create(ISceneOutlinerMode* Mode, co
 {
 	FActorHierarchy* Hierarchy = new FActorHierarchy(Mode, World);
 
+	Create_Internal(Hierarchy, World);
+
+	return TUniquePtr<FActorHierarchy>(Hierarchy);
+}
+
+void FActorHierarchy::Create_Internal(FActorHierarchy* Hierarchy, const TWeakObjectPtr<UWorld>& World)
+{
 	GEngine->OnLevelActorAdded().AddRaw(Hierarchy, &FActorHierarchy::OnLevelActorAdded);
 	GEngine->OnLevelActorDeleted().AddRaw(Hierarchy, &FActorHierarchy::OnLevelActorDeleted);
 	GEngine->OnLevelActorDetached().AddRaw(Hierarchy, &FActorHierarchy::OnLevelActorDetached);
@@ -70,8 +77,6 @@ TUniquePtr<FActorHierarchy> FActorHierarchy::Create(ISceneOutlinerMode* Mode, co
 	Folders.OnFolderCreated.AddRaw(Hierarchy, &FActorHierarchy::OnBroadcastFolderCreate);
 	Folders.OnFolderMoved.AddRaw(Hierarchy, &FActorHierarchy::OnBroadcastFolderMove);
 	Folders.OnFolderDeleted.AddRaw(Hierarchy, &FActorHierarchy::OnBroadcastFolderDelete);
-
-	return TUniquePtr<FActorHierarchy>(Hierarchy);
 }
 
 FActorHierarchy::FActorHierarchy(ISceneOutlinerMode* Mode, const TWeakObjectPtr<UWorld>& World)
@@ -144,7 +149,7 @@ FSceneOutlinerTreeItemPtr FActorHierarchy::FindOrCreateParentItem(const ISceneOu
 				// If Parent can be listed in SceneOutliner return nullptr so it gets created
 				else if (ParentActor->IsListedInSceneOutliner())
 				{
-					return bCreate ? Mode->CreateItemFor<FActorTreeItem>(ParentActor, true) : nullptr;
+					return bCreate ? CreateItemForActor(ParentActor, true) : nullptr;
 				}
 			}
 
@@ -177,7 +182,7 @@ FSceneOutlinerTreeItemPtr FActorHierarchy::FindOrCreateParentItem(const ISceneOu
 					}
 					else
 					{
-						return bCreate ? Mode->CreateItemFor<FActorTreeItem>(OwningActor, true) : nullptr;
+						return bCreate ? CreateItemForActor(OwningActor, true) : nullptr;
 					}
 				}
 			}
@@ -229,7 +234,7 @@ FSceneOutlinerTreeItemPtr FActorHierarchy::FindOrCreateParentItem(const ISceneOu
 				}
 				else
 				{
-					return bCreate ? Mode->CreateItemFor<FActorTreeItem>(OwningActor, true) : nullptr;
+					return bCreate ? CreateItemForActor(OwningActor, true) : nullptr;
 				}
 			}
 		}
@@ -259,7 +264,7 @@ FSceneOutlinerTreeItemPtr FActorHierarchy::FindOrCreateParentItem(const ISceneOu
 			}
 			else
 			{
-				return bCreate ? Mode->CreateItemFor<FActorTreeItem>(Owner, true) : nullptr;
+				return bCreate ? CreateItemForActor(Owner, true) : nullptr;
 			}
 		}
 		// do not default to world on Component items
@@ -303,7 +308,7 @@ FSceneOutlinerTreeItemPtr FActorHierarchy::FindOrCreateParentItem(const ISceneOu
 							}
 							else
 							{
-								return bCreate ? Mode->CreateItemFor<FActorTreeItem>(ParentActor, true) : nullptr;
+								return bCreate ? CreateItemForActor(ParentActor, true) : nullptr;
 							}
 						}
 
@@ -332,7 +337,7 @@ FSceneOutlinerTreeItemPtr FActorHierarchy::FindOrCreateParentItem(const ISceneOu
 					}
 					else
 					{
-						return bCreate ? Mode->CreateItemFor<FActorTreeItem>(ParentActor, true) : nullptr;
+						return bCreate ? CreateItemForActor(ParentActor, true) : nullptr;
 					}
 				}
 			}
@@ -379,71 +384,77 @@ bool FActorHierarchy::IsShowingUnloadedActors() const
 	return bShowingUnloadedActors && RepresentingWorld.IsValid() && !RepresentingWorld->IsPlayInEditor();
 }
 
-void FActorHierarchy::CreateWorldChildren(UWorld* World, TArray<FSceneOutlinerTreeItemPtr>& OutItems) const
+void FActorHierarchy::CreateFolderItems(UWorld* InWorld, TArray<FSceneOutlinerTreeItemPtr>& OutItems) const
 {
-	TRACE_CPUPROFILER_EVENT_SCOPE(FActorHierarchy::CreateWorldChildren);
-
-	check(World);
-
 	if (Mode->ShouldShowFolders() && bShowingEmptyFolders)
 	{
 		// Add any folders which might match the current search terms
-		FActorFolders::Get().ForEachFolder(*World, [this, &World, &OutItems](const FFolder& Folder)
+		FActorFolders::Get().ForEachFolder(*InWorld, [this, &InWorld, &OutItems](const FFolder& Folder)
 		{
-			if (FSceneOutlinerTreeItemPtr FolderItem = Mode->CreateItemFor<FActorFolderTreeItem>(FActorFolderTreeItem(Folder, World)))
+			if (FSceneOutlinerTreeItemPtr FolderItem = Mode->CreateItemFor<FActorFolderTreeItem>(FActorFolderTreeItem(Folder, InWorld)))
 			{
 				OutItems.Add(FolderItem);
 			}
 			return true;
 		});
 	}
-	
+}
+
+FSceneOutlinerTreeItemPtr FActorHierarchy::CreateItemForActor(AActor* InActor, bool bForce) const
+{
+	return Mode->CreateItemFor<FActorTreeItem>(InActor, bForce);
+}
+
+bool FActorHierarchy::CheckLevelInstanceEditing(UWorld* World, AActor* Actor) const
+{
 	const ULevelInstanceSubsystem* LevelInstanceSubsystem = World->GetSubsystem<ULevelInstanceSubsystem>();
-	// Create all actor items
-	for (FActorIterator ActorIt(World); ActorIt; ++ActorIt)
+	// If we are not showing LevelInstances, LevelInstance sub actor items should not be created unless they belong to a LevelInstance which is being edited
+	if (LevelInstanceSubsystem)
 	{
-		AActor* Actor = *ActorIt;
-		// If we are not showing LevelInstances, LevelInstance sub actor items should not be created unless they belong to a LevelInstance which is being edited
-		if (LevelInstanceSubsystem)
+		if (const ILevelInstanceInterface* ParentLevelInstance = LevelInstanceSubsystem->GetParentLevelInstance(Actor))
 		{
-			if (const ILevelInstanceInterface* ParentLevelInstance = LevelInstanceSubsystem->GetParentLevelInstance(Actor))
+			if (!bShowingLevelInstances && !ParentLevelInstance->IsEditing())
 			{
-				if (!bShowingLevelInstances && !ParentLevelInstance->IsEditing())
-				{
-					continue;
-				}
+				return false;
 			}
 		}
-		
-		if (FSceneOutlinerTreeItemPtr ActorItem = Mode->CreateItemFor<FActorTreeItem>(Actor))
-		{
-			if (bShowingOnlyActorWithValidComponents)
-			{
-				int32 InsertLocation = OutItems.Num();
+	}
 
-				// Create all component items
-				CreateComponentItems(Actor, OutItems);
+	return true;
+}
 
-				if (OutItems.Num() != InsertLocation)
-				{
-					// Add the actor before the components
-					OutItems.Insert(ActorItem, InsertLocation);
-				}
-			}
-			else
-			{
-				OutItems.Add(ActorItem);
+void FActorHierarchy::InsertActorItemAndCreateComponents(AActor* InActor, FSceneOutlinerTreeItemPtr ActorItem, TArray<FSceneOutlinerTreeItemPtr>& OutItems) const
+{
+	if (bShowingOnlyActorWithValidComponents)
+	{
+		int32 InsertLocation = OutItems.Num();
 
-				// Create all component items
-				CreateComponentItems(Actor, OutItems);
-			}
+		// Create all component items
+		CreateComponentItems(InActor, OutItems);
+
+		if (OutItems.Num() != InsertLocation)
+		{
+			// Add the actor before the components
+			OutItems.Insert(ActorItem, InsertLocation);
 		}
 	}
+	else
+	{
+		OutItems.Add(ActorItem);
+
+		// Create all component items
+		CreateComponentItems(InActor, OutItems);
+	}
+}
 
+void FActorHierarchy::CreateUnloadedItems(UWorld* World, TArray<FSceneOutlinerTreeItemPtr>& OutItems) const
+{
 	if (IsShowingUnloadedActors())
 	{
 		if (UWorldPartitionSubsystem* WorldPartitionSubsystem = UWorld::GetSubsystem<UWorldPartitionSubsystem>(World))
 		{
+			const ULevelInstanceSubsystem* LevelInstanceSubsystem = World->GetSubsystem<ULevelInstanceSubsystem>();
+
 			WorldPartitionSubsystem->ForEachWorldPartition([this, LevelInstanceSubsystem, &OutItems](UWorldPartition* WorldPartition)
 			{
 				// Skip unloaded actors if they are part of a non editing level instance and the outliner hides the content of level instances
@@ -474,6 +485,33 @@ void FActorHierarchy::CreateWorldChildren(UWorld* World, TArray<FSceneOutlinerTr
 	}
 }
 
+void FActorHierarchy::CreateWorldChildren(UWorld* World, TArray<FSceneOutlinerTreeItemPtr>& OutItems) const
+{
+	TRACE_CPUPROFILER_EVENT_SCOPE(FActorHierarchy::CreateWorldChildren);
+
+	check(World);
+
+	CreateFolderItems(World, OutItems);
+	
+	// Create all actor items
+	for (FActorIterator ActorIt(World); ActorIt; ++ActorIt)
+	{
+		AActor* Actor = *ActorIt;
+
+		if(!CheckLevelInstanceEditing(World, Actor))
+		{
+			continue;
+		}
+		
+		if (FSceneOutlinerTreeItemPtr ActorItem = CreateItemForActor(Actor))
+		{
+			InsertActorItemAndCreateComponents(Actor, ActorItem, OutItems);
+		}
+	}
+
+	CreateUnloadedItems(World, OutItems);
+}
+
 void FActorHierarchy::CreateItems(TArray<FSceneOutlinerTreeItemPtr>& OutItems) const
 {
 	TRACE_CPUPROFILER_EVENT_SCOPE(FActorHierarchy::CreateItems);
@@ -560,7 +598,7 @@ void FActorHierarchy::CreateChildren(const FSceneOutlinerTreeItemPtr& Item, TArr
 
 		for (auto ChildActor : ChildActors)
 		{
-			if (FSceneOutlinerTreeItemPtr ChildActorItem = Mode->CreateItemFor<FActorTreeItem>(ChildActor))
+			if (FSceneOutlinerTreeItemPtr ChildActorItem = CreateItemForActor(ChildActor))
 			{
 				OutChildren.Add(ChildActorItem);
 
@@ -601,7 +639,7 @@ void FActorHierarchy::OnLevelActorAdded(AActor* InActor)
 	{
 		FSceneOutlinerHierarchyChangedData EventData;
 		EventData.Type = FSceneOutlinerHierarchyChangedData::Added;
-		EventData.Items.Add(Mode->CreateItemFor<FActorTreeItem>(InActor));
+		EventData.Items.Add(CreateItemForActor(InActor));
 		HierarchyChangedEvent.Broadcast(EventData);
 	}
 }
@@ -639,10 +677,8 @@ void FActorHierarchy::OnLevelActorDetached(AActor* InActor, const AActor* InPare
 	}
 }
 
-void FActorHierarchy::OnLoadedActorAdded(AActor& InActor)
+void FActorHierarchy::RemoveActorDesc(AActor& InActor)
 {
-	OnLevelActorAdded(&InActor);
-
 	// Loaded actors can be in sub-world partitions so we use the outer world
 	if (UWorldPartition* WorldPartition = FWorldPartitionHelpers::GetWorldPartition(&InActor))
 	{
@@ -657,10 +693,8 @@ void FActorHierarchy::OnLoadedActorAdded(AActor& InActor)
 	}
 }
 
-void FActorHierarchy::OnLoadedActorRemoved(AActor& InActor)
+void FActorHierarchy::AddActorDesc(AActor& InActor)
 {
-	OnLevelActorDeleted(&InActor);
-
 	if (IsShowingUnloadedActors())
 	{
 		if (UWorldPartition* WorldPartition = FWorldPartitionHelpers::GetWorldPartition(&InActor))
@@ -677,6 +711,20 @@ void FActorHierarchy::OnLoadedActorRemoved(AActor& InActor)
 	}
 }
 
+void FActorHierarchy::OnLoadedActorAdded(AActor& InActor)
+{
+	OnLevelActorAdded(&InActor);
+
+	RemoveActorDesc(InActor);
+}
+
+void FActorHierarchy::OnLoadedActorRemoved(AActor& InActor)
+{
+	OnLevelActorDeleted(&InActor);
+
+	AddActorDesc(InActor);
+}
+
 void FActorHierarchy::OnActorDescAdded(FWorldPartitionActorDesc* ActorDesc)
 {
 	if (IsShowingUnloadedActors() && ActorDesc && !ActorDesc->IsLoaded(true) && FActorDescTreeItem::ShouldDisplayInOutliner(ActorDesc))
@@ -749,13 +797,25 @@ void FActorHierarchy::OnLevelAdded(ULevel* InLevel, UWorld* InWorld)
 		{
 			if (Actor != nullptr)
 			{
-				EventData.Items.Add(Mode->CreateItemFor<FActorTreeItem>(Actor));
+				EventData.Items.Add(CreateItemForActor(Actor));
 			}
 		}
 		HierarchyChangedEvent.Broadcast(EventData);
 	}
 }
 
+void FActorHierarchy::RemoveLevelActorFolders(ULevel* InLevel, UWorld* InWorld)
+{
+	// If either this level or the owning world are using actor folders, remove level's actor folders
+	if (InLevel->IsUsingActorFolders() || InWorld->PersistentLevel->IsUsingActorFolders())
+	{
+		FSceneOutlinerHierarchyChangedData EventData;
+		EventData.Type = FSceneOutlinerHierarchyChangedData::Removed;
+		EventData.ItemIDs.Add(InLevel);
+		HierarchyChangedEvent.Broadcast(EventData);
+	}
+}
+
 void FActorHierarchy::OnLevelRemoved(ULevel* InLevel, UWorld* InWorld)
 {
 	if (InLevel != nullptr && RepresentingWorld.Get() == InWorld)
@@ -774,14 +834,7 @@ void FActorHierarchy::OnLevelRemoved(ULevel* InLevel, UWorld* InWorld)
 			HierarchyChangedEvent.Broadcast(EventData);
 		}
 
-		// If either this level or the owning world are using actor folders, remove level's actor folders
-		if (InLevel->IsUsingActorFolders() || InWorld->PersistentLevel->IsUsingActorFolders())
-		{
-			FSceneOutlinerHierarchyChangedData EventData;
-			EventData.Type = FSceneOutlinerHierarchyChangedData::Removed;
-			EventData.ItemIDs.Add(InLevel);
-			HierarchyChangedEvent.Broadcast(EventData);
-		}
+		RemoveLevelActorFolders(InLevel, InWorld);
 	}
 }
 
diff --git a/Engine/Source/Editor/SceneOutliner/Public/ActorBrowsingMode.h b/Engine/Source/Editor/SceneOutliner/Public/ActorBrowsingMode.h
--- a/Engine/Source/Editor/SceneOutliner/Public/ActorBrowsingMode.h
+++ b/Engine/Source/Editor/SceneOutliner/Public/ActorBrowsingMode.h
@@ -10,55 +10,55 @@ class IWorldPartitionEditorModule;
 class FActorBrowsingMode : public FActorModeInteractive
 {
 public:
-	FActorBrowsingMode(SSceneOutliner* InSceneOutliner, TWeakObjectPtr<UWorld> InSpecifiedWorldToDisplay = nullptr);
-	virtual ~FActorBrowsingMode();
+	SCENEOUTLINER_API FActorBrowsingMode(SSceneOutliner* InSceneOutliner, TWeakObjectPtr<UWorld> InSpecifiedWorldToDisplay = nullptr);
+	SCENEOUTLINER_API virtual ~FActorBrowsingMode();
 
 	/* Begin ISceneOutlinerMode Interface */
-	virtual void Rebuild() override;
-	virtual FCreateSceneOutlinerMode CreateFolderPickerMode(const FFolder::FRootObject& InRootObject = FFolder::GetInvalidRootObject()) const override;
-	virtual void CreateViewContent(FMenuBuilder& MenuBuilder) override;
-	virtual TSharedPtr<FDragDropOperation> CreateDragDropOperation(const FPointerEvent& MouseEvent, const TArray<FSceneOutlinerTreeItemPtr>& InTreeItems) const override;
-	virtual bool ParseDragDrop(FSceneOutlinerDragDropPayload& OutPayload, const FDragDropOperation& Operation) const override;
-	virtual FSceneOutlinerDragValidationInfo ValidateDrop(const ISceneOutlinerTreeItem& DropTarget, const FSceneOutlinerDragDropPayload& Payload) const override;
-	virtual TSharedPtr<SWidget> CreateContextMenu() override;
-	virtual void OnItemAdded(FSceneOutlinerTreeItemPtr Item) override;
-	virtual void OnItemRemoved(FSceneOutlinerTreeItemPtr Item) override;
-	virtual void OnItemSelectionChanged(FSceneOutlinerTreeItemPtr Item, ESelectInfo::Type SelectionType, const FSceneOutlinerItemSelection& Selection) override;
-	virtual void OnItemDoubleClick(FSceneOutlinerTreeItemPtr Item) override;
-	virtual void OnFilterTextCommited(FSceneOutlinerItemSelection& Selection, ETextCommit::Type CommitType) override;
-	virtual void OnItemPassesFilters(const ISceneOutlinerTreeItem& Item) override;
-	virtual FReply OnKeyDown(const FKeyEvent& InKeyEvent) override;
-	virtual void OnDuplicateSelected() override;
-	virtual void OnDrop(ISceneOutlinerTreeItem& DropTarget, const FSceneOutlinerDragDropPayload& Payload, const FSceneOutlinerDragValidationInfo& ValidationInfo) const override;
-	virtual bool CanRenameItem(const ISceneOutlinerTreeItem& Item) const override;
-	virtual FText GetStatusText() const override;
-	virtual FSlateColor GetStatusTextColor() const override;
-	virtual ESelectionMode::Type GetSelectionMode() const override { return ESelectionMode::Multi; }
-	virtual bool SupportsKeyboardFocus() const override { return true; }
-	virtual bool ShouldShowFolders() const override { return true; }
-	virtual bool SupportsCreateNewFolder() const override { return true; }
-	virtual bool ShowStatusBar() const override { return true; }
-	virtual bool ShowViewButton() const override { return true; }
-	virtual bool ShowFilterOptions() const override { return true; }
-	virtual bool CanDelete() const override;
-	virtual bool CanRename() const override;
-	virtual bool CanCut() const override;
-	virtual bool CanCopy() const override;
-	virtual bool CanPaste() const override;
-	virtual bool CanSupportDragAndDrop() const { return true; }
-	virtual bool HasErrors() const;
-	virtual FText GetErrorsText() const;
-	virtual void RepairErrors() const;	
-	virtual FFolder CreateNewFolder() override;
-	virtual FFolder GetFolder(const FFolder& ParentPath, const FName& LeafName) override;
-	virtual bool CreateFolder(const FFolder& NewFolder) override;
-	virtual bool ReparentItemToFolder(const FFolder& FolderPath, const FSceneOutlinerTreeItemPtr& Item) override;
-	virtual void SelectFoldersDescendants(const TArray<FFolderTreeItem*>& FolderItems, bool bSelectImmediateChildrenOnly) override;
-	virtual void PinItems(const TArray<FSceneOutlinerTreeItemPtr>& InItems) override;
-	virtual void UnpinItems(const TArray<FSceneOutlinerTreeItemPtr>& InItems) override;
-	virtual bool CanPinItems(const TArray<FSceneOutlinerTreeItemPtr>& InItems) const override;
-	virtual bool CanUnpinItems(const TArray<FSceneOutlinerTreeItemPtr>& InItems) const override;
-	virtual void SynchronizeSelection() override;
+	SCENEOUTLINER_API virtual void Rebuild() override;
+	SCENEOUTLINER_API virtual FCreateSceneOutlinerMode CreateFolderPickerMode(const FFolder::FRootObject& InRootObject = FFolder::GetInvalidRootObject()) const override;
+	SCENEOUTLINER_API virtual void CreateViewContent(FMenuBuilder& MenuBuilder) override;
+	SCENEOUTLINER_API virtual TSharedPtr<FDragDropOperation> CreateDragDropOperation(const FPointerEvent& MouseEvent, const TArray<FSceneOutlinerTreeItemPtr>& InTreeItems) const override;
+	SCENEOUTLINER_API virtual bool ParseDragDrop(FSceneOutlinerDragDropPayload& OutPayload, const FDragDropOperation& Operation) const override;
+	SCENEOUTLINER_API virtual FSceneOutlinerDragValidationInfo ValidateDrop(const ISceneOutlinerTreeItem& DropTarget, const FSceneOutlinerDragDropPayload& Payload) const override;
+	SCENEOUTLINER_API virtual TSharedPtr<SWidget> CreateContextMenu() override;
+	SCENEOUTLINER_API virtual void OnItemAdded(FSceneOutlinerTreeItemPtr Item) override;
+	SCENEOUTLINER_API virtual void OnItemRemoved(FSceneOutlinerTreeItemPtr Item) override;
+	SCENEOUTLINER_API virtual void OnItemSelectionChanged(FSceneOutlinerTreeItemPtr Item, ESelectInfo::Type SelectionType, const FSceneOutlinerItemSelection& Selection) override;
+	SCENEOUTLINER_API virtual void OnItemDoubleClick(FSceneOutlinerTreeItemPtr Item) override;
+	SCENEOUTLINER_API virtual void OnFilterTextCommited(FSceneOutlinerItemSelection& Selection, ETextCommit::Type CommitType) override;
+	SCENEOUTLINER_API virtual void OnItemPassesFilters(const ISceneOutlinerTreeItem& Item) override;
+	SCENEOUTLINER_API virtual FReply OnKeyDown(const FKeyEvent& InKeyEvent) override;
+	SCENEOUTLINER_API virtual void OnDuplicateSelected() override;
+	SCENEOUTLINER_API virtual void OnDrop(ISceneOutlinerTreeItem& DropTarget, const FSceneOutlinerDragDropPayload& Payload, const FSceneOutlinerDragValidationInfo& ValidationInfo) const override;
+	SCENEOUTLINER_API virtual bool CanRenameItem(const ISceneOutlinerTreeItem& Item) const override;
+	SCENEOUTLINER_API virtual FText GetStatusText() const override;
+	SCENEOUTLINER_API virtual FSlateColor GetStatusTextColor() const override;
+	SCENEOUTLINER_API virtual ESelectionMode::Type GetSelectionMode() const override { return ESelectionMode::Multi; }
+	SCENEOUTLINER_API virtual bool SupportsKeyboardFocus() const override { return true; }
+	SCENEOUTLINER_API virtual bool ShouldShowFolders() const override { return true; }
+	SCENEOUTLINER_API virtual bool SupportsCreateNewFolder() const override { return true; }
+	SCENEOUTLINER_API virtual bool ShowStatusBar() const override { return true; }
+	SCENEOUTLINER_API virtual bool ShowViewButton() const override { return true; }
+	SCENEOUTLINER_API virtual bool ShowFilterOptions() const override { return true; }
+	SCENEOUTLINER_API virtual bool CanDelete() const override;
+	SCENEOUTLINER_API virtual bool CanRename() const override;
+	SCENEOUTLINER_API virtual bool CanCut() const override;
+	SCENEOUTLINER_API virtual bool CanCopy() const override;
+	SCENEOUTLINER_API virtual bool CanPaste() const override;
+	SCENEOUTLINER_API virtual bool CanSupportDragAndDrop() const { return true; }
+	SCENEOUTLINER_API virtual bool HasErrors() const;
+	SCENEOUTLINER_API virtual FText GetErrorsText() const;
+	SCENEOUTLINER_API virtual void RepairErrors() const;	
+	SCENEOUTLINER_API virtual FFolder CreateNewFolder() override;
+	SCENEOUTLINER_API virtual FFolder GetFolder(const FFolder& ParentPath, const FName& LeafName) override;
+	SCENEOUTLINER_API virtual bool CreateFolder(const FFolder& NewFolder) override;
+	SCENEOUTLINER_API virtual bool ReparentItemToFolder(const FFolder& FolderPath, const FSceneOutlinerTreeItemPtr& Item) override;
+	SCENEOUTLINER_API virtual void SelectFoldersDescendants(const TArray<FFolderTreeItem*>& FolderItems, bool bSelectImmediateChildrenOnly) override;
+	SCENEOUTLINER_API virtual void PinItems(const TArray<FSceneOutlinerTreeItemPtr>& InItems) override;
+	SCENEOUTLINER_API virtual void UnpinItems(const TArray<FSceneOutlinerTreeItemPtr>& InItems) override;
+	SCENEOUTLINER_API virtual bool CanPinItems(const TArray<FSceneOutlinerTreeItemPtr>& InItems) const override;
+	SCENEOUTLINER_API virtual bool CanUnpinItems(const TArray<FSceneOutlinerTreeItemPtr>& InItems) const override;
+	SCENEOUTLINER_API virtual void SynchronizeSelection() override;
 
 	/* End ISceneOutlinerMode Interface */
 public:
@@ -106,7 +106,7 @@ public:
 	void OnDeleteActorsEnd();
 	
 	/** Function called by the Outliner Filter Bar to compare an item with Type Filters*/
-	virtual bool CompareItemWithClassName(SceneOutliner::FilterBarType InItem, const TSet<FTopLevelAssetPath>&) const override;
+	SCENEOUTLINER_API virtual bool CompareItemWithClassName(SceneOutliner::FilterBarType InItem, const TSet<FTopLevelAssetPath>&) const override;
 private:
 	/** Build and up the context menu */
 	TSharedPtr<SWidget> BuildContextMenu();
diff --git a/Engine/Source/Editor/SceneOutliner/Public/ActorHierarchy.h b/Engine/Source/Editor/SceneOutliner/Public/ActorHierarchy.h
--- a/Engine/Source/Editor/SceneOutliner/Public/ActorHierarchy.h
+++ b/Engine/Source/Editor/SceneOutliner/Public/ActorHierarchy.h
@@ -13,16 +13,16 @@ class UWorldPartition;
 class FActorHierarchy : public ISceneOutlinerHierarchy
 {
 public:
-	virtual ~FActorHierarchy();
+	SCENEOUTLINER_API virtual ~FActorHierarchy();
 
 	static TUniquePtr<FActorHierarchy> Create(ISceneOutlinerMode* Mode, const TWeakObjectPtr<UWorld>& World);
 
 	/** Create a linearization of all applicable items in the hierarchy */
-	virtual void CreateItems(TArray<FSceneOutlinerTreeItemPtr>& OutItems) const override;
+	SCENEOUTLINER_API virtual void CreateItems(TArray<FSceneOutlinerTreeItemPtr>& OutItems) const override;
 	/** Create a linearization of all direct and indirect children of a given item in the hierarchy */
-	virtual void CreateChildren(const FSceneOutlinerTreeItemPtr& Item, TArray<FSceneOutlinerTreeItemPtr>& OutChildren) const override;
+	SCENEOUTLINER_API virtual void CreateChildren(const FSceneOutlinerTreeItemPtr& Item, TArray<FSceneOutlinerTreeItemPtr>& OutChildren) const override;
 	/** Forcibly create a parent item for a given tree item */
-	virtual FSceneOutlinerTreeItemPtr FindOrCreateParentItem(const ISceneOutlinerTreeItem& Item, const TMap<FSceneOutlinerTreeItemID, FSceneOutlinerTreeItemPtr>& Items, bool bCreate = false) override;
+	SCENEOUTLINER_API virtual FSceneOutlinerTreeItemPtr FindOrCreateParentItem(const ISceneOutlinerTreeItem& Item, const TMap<FSceneOutlinerTreeItemID, FSceneOutlinerTreeItemPtr>& Items, bool bCreate = false) override;
 
 	void SetShowingComponents(bool bInShowingComponents) { bShowingComponents = bInShowingComponents; }
 	void SetShowingOnlyActorWithValidComponents(bool bInShowingOnlyActorWithValidComponents) { bShowingOnlyActorWithValidComponents = bInShowingOnlyActorWithValidComponents; }
@@ -37,23 +37,23 @@ private:
 	/** Create all component items for an actor if we are showing components and place them in OutItems */
 	void CreateComponentItems(const AActor* Actor, TArray<FSceneOutlinerTreeItemPtr>& OutItems) const;
 
-private:
+protected:
 	// Update the hierarchy when actor or world changing events occur
 	void OnWorldPartitionCreated(UWorld* InWorld);
 
-	void OnLevelActorAdded(AActor* InActor);
-	void OnLevelActorDeleted(AActor* InActor);
+	SCENEOUTLINER_API virtual void OnLevelActorAdded(AActor* InActor);
+	SCENEOUTLINER_API virtual void OnLevelActorDeleted(AActor* InActor);
 		
-	void OnLevelActorAttached(AActor* InActor, const AActor* InParent);
-	void OnLevelActorDetached(AActor* InActor, const AActor* InParent);
+	SCENEOUTLINER_API virtual void OnLevelActorAttached(AActor* InActor, const AActor* InParent);
+	SCENEOUTLINER_API virtual void OnLevelActorDetached(AActor* InActor, const AActor* InParent);
 
-	void OnLevelActorFolderChanged(const AActor* InActor, FName OldPath);
+	SCENEOUTLINER_API virtual void OnLevelActorFolderChanged(const AActor* InActor, FName OldPath);
 
-	void OnLoadedActorAdded(AActor& InActor);
-	void OnLoadedActorRemoved(AActor& InActor);
+	SCENEOUTLINER_API virtual void OnLoadedActorAdded(AActor& InActor);
+	SCENEOUTLINER_API virtual void OnLoadedActorRemoved(AActor& InActor);
 
-	void OnActorDescAdded(FWorldPartitionActorDesc* ActorDesc);
-	void OnActorDescRemoved(FWorldPartitionActorDesc* ActorDesc);
+	SCENEOUTLINER_API virtual void OnActorDescAdded(FWorldPartitionActorDesc* ActorDesc);
+	SCENEOUTLINER_API virtual void OnActorDescRemoved(FWorldPartitionActorDesc* ActorDesc);
 	
 	void OnComponentsUpdated();
 
@@ -62,8 +62,8 @@ private:
 	void OnActorFolderAdded(UActorFolder* InActorFolder);
 	void OnActorFoldersUpdatedEvent(ULevel* InLevel);
 		
-	void OnLevelAdded(ULevel* InLevel, UWorld* InWorld);
-	void OnLevelRemoved(ULevel* InLevel, UWorld* InWorld);
+	SCENEOUTLINER_API virtual void OnLevelAdded(ULevel* InLevel, UWorld* InWorld);
+	SCENEOUTLINER_API virtual void OnLevelRemoved(ULevel* InLevel, UWorld* InWorld);
 
 	void OnWorldPartitionInitialized(UWorldPartition* InWorldPartition);
 	void OnWorldPartitionUninitialized(UWorldPartition* InWorldPartition);
@@ -77,7 +77,30 @@ private:
 	/** Called when a folder is to be deleted */
 	void OnBroadcastFolderDelete(UWorld& InWorld, const FFolder& InFolder);
 
-private:
+	// Remove the unloaded actor item for the given actor
+	SCENEOUTLINER_API void RemoveActorDesc(AActor& InActor);
+	
+	// Add the unloaded actor item for the given actor
+	SCENEOUTLINER_API void AddActorDesc(AActor& InActor);
+
+	// Remove all Actor Folders belonging to this level
+	SCENEOUTLINER_API void RemoveLevelActorFolders(ULevel* InLevel, UWorld* InWorld);
+
+	// Create items for Actor Folders
+	SCENEOUTLINER_API void CreateFolderItems(UWorld* InWorld, TArray<FSceneOutlinerTreeItemPtr>& OutItems) const;
+
+	// Create items for unloaded actors
+	SCENEOUTLINER_API void CreateUnloadedItems(UWorld* World, TArray<FSceneOutlinerTreeItemPtr>& OutItems) const;
+
+	// Insert the given item at the right position in the tree, and optionally create and insert components if needed
+	SCENEOUTLINER_API void InsertActorItemAndCreateComponents(AActor* InActor, FSceneOutlinerTreeItemPtr ActorItem, TArray<FSceneOutlinerTreeItemPtr>& OutItems) const;
+
+	// Check if the actor we are creating an item for is in the current level instance, if we are currently editing one
+	SCENEOUTLINER_API bool CheckLevelInstanceEditing(UWorld* World, AActor* Actor) const;
+
+	// Create a SceneOutlinerTreeItem for the given actor
+	virtual FSceneOutlinerTreeItemPtr CreateItemForActor(AActor* InActor, bool bForce = false) const;
+protected:
 	/** Send a an event indicating a full refresh of the hierarchy is required */
 	void FullRefreshEvent();
 
@@ -89,8 +112,11 @@ private:
 	bool bShowingUnloadedActors = false;
 	bool bShowingEmptyFolders = false;
 
-private:
-	FActorHierarchy(ISceneOutlinerMode* Mode, const TWeakObjectPtr<UWorld>& Worlds);
+protected:
+
+	SCENEOUTLINER_API static void Create_Internal(FActorHierarchy* Hierarchy, const TWeakObjectPtr<UWorld>& World);
+	
+	SCENEOUTLINER_API FActorHierarchy(ISceneOutlinerMode* Mode, const TWeakObjectPtr<UWorld>& Worlds);
 
 	FActorHierarchy(const FActorHierarchy&) = delete;
 	FActorHierarchy& operator=(const FActorHierarchy&) = delete;
diff --git a/Engine/Source/Editor/SceneOutliner/Public/ActorModeInteractive.h b/Engine/Source/Editor/SceneOutliner/Public/ActorModeInteractive.h
--- a/Engine/Source/Editor/SceneOutliner/Public/ActorModeInteractive.h
+++ b/Engine/Source/Editor/SceneOutliner/Public/ActorModeInteractive.h
@@ -7,18 +7,19 @@
 class FActorModeInteractive : public FActorMode
 {
 public:
-	FActorModeInteractive(const FActorModeParams& Params);
-	virtual ~FActorModeInteractive();
+	SCENEOUTLINER_API FActorModeInteractive(const FActorModeParams& Params);
+	SCENEOUTLINER_API virtual ~FActorModeInteractive();
 
-	virtual bool IsInteractive() const override { return true; }
-private:
+	SCENEOUTLINER_API virtual bool IsInteractive() const override { return true; }
+protected:
 	/* Events */
 
 	void OnMapChange(uint32 MapFlags);
 	void OnNewCurrentLevel();
 
+	SCENEOUTLINER_API virtual void OnActorLabelChanged(AActor* ChangedActor);
+	
 	void OnLevelSelectionChanged(UObject* Obj);
-	void OnActorLabelChanged(AActor* ChangedActor);
 	void OnObjectsReplaced(const TMap<UObject*, UObject*>& ReplacementMap);
 	void OnLevelActorRequestsRename(const AActor* Actor);
 	void OnPostLoadMapWithWorld(UWorld* World);
diff --git a/Engine/Source/Editor/UnrealEd/Private/LevelEditorViewport.cpp b/Engine/Source/Editor/UnrealEd/Private/LevelEditorViewport.cpp
--- a/Engine/Source/Editor/UnrealEd/Private/LevelEditorViewport.cpp
+++ b/Engine/Source/Editor/UnrealEd/Private/LevelEditorViewport.cpp
@@ -1412,6 +1412,13 @@ void FLevelEditorViewportClient::DestroyDropPreviewActors()
 			AActor* PreviewActor = (*ActorIt).Get();
 			if (PreviewActor && PreviewActor != GetWorld()->GetDefaultBrush())
 			{
+				// Temporary hack to remove the preview actor from the TypedElementRegistry (if it exists)
+				UTypedElementRegistry* TypedElementRegistry = UTypedElementRegistry::GetInstance();
+				if(ITypedElementDataStorageCompatibilityInterface* DataStorageCompatibilityInterface = TypedElementRegistry->GetMutableDataStorageCompatibility())
+				{
+					DataStorageCompatibilityInterface->RemoveCompatibleObject(PreviewActor);
+				}
+
 				GetWorld()->DestroyActor(PreviewActor);
 			}
 		}
