Outliner: Fix a bug where if an item was requested to be removed and then re-added before the pending removal was executed, the re-add would fail

diff --git a/Engine/Source/Editor/SceneOutliner/Private/SSceneOutliner.cpp b/Engine/Source/Editor/SceneOutliner/Private/SSceneOutliner.cpp
--- a/Engine/Source/Editor/SceneOutliner/Private/SSceneOutliner.cpp
+++ b/Engine/Source/Editor/SceneOutliner/Private/SSceneOutliner.cpp
@@ -941,6 +941,8 @@ void SSceneOutliner::RemoveItemFromTree(FSceneOutlinerTreeItemRef ReferenceItem)
 			RootTreeItems.Remove(Item);
 		}
 
+		PendingTreeItemMap_Removal.Remove(Item->GetID());
+
 		TreeItemMap.Remove(Item->GetID());
 
 		Mode->OnItemRemoved(Item);
@@ -1969,7 +1971,13 @@ void SSceneOutliner::OnHierarchyChangedEvent(FSceneOutlinerHierarchyChangedData
 	{
 		for (const auto& TreeItemPtr : Event.Items)
 		{
-			if (TreeItemPtr.IsValid() && !TreeItemMap.Find(TreeItemPtr->GetID()))
+			if(!TreeItemPtr.IsValid())
+			{
+				continue;
+			}
+
+			// If the item doesn't exist in the tree, or is being removed and re-added in the same frame - it is not a duplicate and can be added
+			if(!TreeItemMap.Find(TreeItemPtr->GetID()) || PendingTreeItemMap_Removal.Find(TreeItemPtr->GetID()))
 			{
 				AddPendingItemAndChildren(TreeItemPtr);
 				if (Event.ItemActions)
@@ -1992,6 +2000,7 @@ void SSceneOutliner::OnHierarchyChangedEvent(FSceneOutlinerHierarchyChangedData
 			if (Item)
 			{
 				PendingOperations.Emplace(SceneOutliner::FPendingTreeOperation::Removed, Item->ToSharedRef());
+				PendingTreeItemMap_Removal.Add(TreeItemID, Item->ToSharedRef());
 			}
 		}
 		Refresh();
@@ -2077,6 +2086,7 @@ void SSceneOutliner::OnItemLabelChanged(FSceneOutlinerTreeItemPtr ChangedItem)
 		{
 			// No longer matches the filters, remove it
 			PendingOperations.Emplace(SceneOutliner::FPendingTreeOperation::Removed, ExistingItem->ToSharedRef());
+			PendingTreeItemMap_Removal.Add(ChangedItem->GetID(), ExistingItem->ToSharedRef());
 			Refresh();
 		}
 	}
diff --git a/Engine/Source/Editor/SceneOutliner/Public/SSceneOutliner.h b/Engine/Source/Editor/SceneOutliner/Public/SSceneOutliner.h
--- a/Engine/Source/Editor/SceneOutliner/Public/SSceneOutliner.h
+++ b/Engine/Source/Editor/SceneOutliner/Public/SSceneOutliner.h
@@ -847,6 +847,9 @@ private:
 	/** Pending tree items that are yet to be added the tree */
 	FSceneOutlinerTreeItemMap PendingTreeItemMap;
 
+	/** Pending tree items that are yet to be removed from the tree */
+	FSceneOutlinerTreeItemMap PendingTreeItemMap_Removal;
+
 	/** Folders pending selection */
 	TArray<FFolder> PendingFoldersSelect;
 
