SceneOutliner: QoL Add "Browse to Asset" on World tree root allows selecting the map in the ContentBrowser

#rnx

diff --git a/Engine/Source/Editor/SceneOutliner/Private/WorldTreeItem.cpp b/Engine/Source/Editor/SceneOutliner/Private/WorldTreeItem.cpp
--- a/Engine/Source/Editor/SceneOutliner/Private/WorldTreeItem.cpp
+++ b/Engine/Source/Editor/SceneOutliner/Private/WorldTreeItem.cpp
@@ -17,6 +17,7 @@
 #include "ISceneOutlinerMode.h"
 #include "IDocumentation.h"
 #include "FolderTreeItem.h"
+#include "Editor.h"
 
 #define LOCTEXT_NAMESPACE "SceneOutliner_WorldTreeItem"
 
@@ -224,6 +225,9 @@ void FWorldTreeItem::GenerateContextMenu(UToolMenu* Menu, SSceneOutliner& Outlin
 	// Create Folder menu entry
 	FSceneOutlinerMenuHelper::AddMenuEntryCreateFolder(Section, Outliner);
 
+	const FSlateIcon FindInContentBrowserIcon(FAppStyle::GetAppStyleSetName(), "SystemWideCommands.FindInContentBrowser");
+	Section.AddMenuEntry("BrowseToWorld", LOCTEXT("BrowseToWorld", "Browse to Asset"), FText(), FindInContentBrowserIcon, FUIAction(FExecuteAction::CreateSP(this, &FWorldTreeItem::BrowseToAsset), FCanExecuteAction::CreateSP(this, &FWorldTreeItem::CanBrowseToAsset)));
+	
 	// Open world settings menu entry
 	const FSlateIcon WorldSettingsIcon(FAppStyle::GetAppStyleSetName(), "LevelEditor.WorldProperties.Tab");
 	Section.AddMenuEntry("OpenWorldSettings", LOCTEXT("OpenWorldSettings", "World Settings"), FText(), WorldSettingsIcon, FExecuteAction::CreateSP(this, &FWorldTreeItem::OpenWorldSettings));
@@ -245,4 +249,22 @@ void FWorldTreeItem::OpenWorldSettings() const
 	LevelEditorModule.GetLevelEditorTabManager()->TryInvokeTab(FName("WorldSettingsTab"));	
 }
 
+void FWorldTreeItem::BrowseToAsset() const
+{
+	if (CanBrowseToAsset())
+	{
+		GEditor->SyncBrowserToObject(World.Get());
+	}
+}
+
+bool FWorldTreeItem::CanBrowseToAsset() const
+{
+	if (UWorld* WorldPtr = World.Get())
+	{
+		return !WorldPtr->GetPackage()->HasAnyPackageFlags(PKG_PlayInEditor | PKG_NewlyCreated);
+	}
+
+	return false;
+}
+
 #undef LOCTEXT_NAMESPACE
diff --git a/Engine/Source/Editor/SceneOutliner/Public/WorldTreeItem.h b/Engine/Source/Editor/SceneOutliner/Public/WorldTreeItem.h
--- a/Engine/Source/Editor/SceneOutliner/Public/WorldTreeItem.h
+++ b/Engine/Source/Editor/SceneOutliner/Public/WorldTreeItem.h
@@ -56,6 +56,10 @@ public:
 
 	/** Open the world settings for the contained world */
 	void OpenWorldSettings() const;
+	/** Browse to world asset in content browser */
+	void BrowseToAsset() const;
+	bool CanBrowseToAsset() const;
+
 	/** Get just the name of the world, for tooltip use */
 	FString GetWorldName() const;
 };
