SceneOutliner
- Remove pending removal operation when readding an item.

diff --git a/Engine/Source/Editor/SceneOutliner/Private/SSceneOutliner.cpp b/Engine/Source/Editor/SceneOutliner/Private/SSceneOutliner.cpp
--- a/Engine/Source/Editor/SceneOutliner/Private/SSceneOutliner.cpp
+++ b/Engine/Source/Editor/SceneOutliner/Private/SSceneOutliner.cpp
@@ -822,6 +822,13 @@ void SSceneOutliner::AddPendingItemAndChildren(FSceneOutlinerTreeItemPtr Item)
 	// Verify that there isn't already a pending operation for this item:
 	if (PendingTreeItemMap.Contains(Item->GetID()))
 	{
+		PendingOperations.RemoveAll(
+			[&Item](const SceneOutliner::FPendingTreeOperation& PendingOp)
+			{
+				return PendingOp.Item->GetID() == Item->GetID() && PendingOp.Type == SceneOutliner::FPendingTreeOperation::Removed;
+			});
+
+		PendingTreeItemMap_Removal.Remove(Item->GetID());
 		return;
 	}
 
@@ -829,7 +836,7 @@ void SSceneOutliner::AddPendingItemAndChildren(FSceneOutlinerTreeItemPtr Item)
 
 	TArray<FSceneOutlinerTreeItemPtr> Children;
 	Mode->GetHierarchy()->CreateChildren(Item, Children);
-	for (auto& Child : Children)
+	for (const FSceneOutlinerTreeItemPtr& Child : Children)
 	{
 		AddPendingItem(Child);
 	}
@@ -1969,7 +1976,7 @@ void SSceneOutliner::OnHierarchyChangedEvent(FSceneOutlinerHierarchyChangedData
 {
 	if (Event.Type == FSceneOutlinerHierarchyChangedData::Added)
 	{
-		for (const auto& TreeItemPtr : Event.Items)
+		for (const FSceneOutlinerTreeItemPtr& TreeItemPtr : Event.Items)
 		{
 			if(!TreeItemPtr.IsValid())
 			{
@@ -1989,7 +1996,7 @@ void SSceneOutliner::OnHierarchyChangedEvent(FSceneOutlinerHierarchyChangedData
 	}
 	else if (Event.Type == FSceneOutlinerHierarchyChangedData::Removed)
 	{
-		for (const auto& TreeItemID : Event.ItemIDs)
+		for (const FSceneOutlinerTreeItemID& TreeItemID : Event.ItemIDs)
 		{
 			FSceneOutlinerTreeItemPtr* Item = TreeItemMap.Find(TreeItemID);
 			if (!Item)
@@ -2007,7 +2014,7 @@ void SSceneOutliner::OnHierarchyChangedEvent(FSceneOutlinerHierarchyChangedData
 	}
 	else if (Event.Type == FSceneOutlinerHierarchyChangedData::Moved)
 	{
-		for (const auto& TreeItemID : Event.ItemIDs)
+		for (const FSceneOutlinerTreeItemID& TreeItemID : Event.ItemIDs)
 		{
 			FSceneOutlinerTreeItemPtr* Item = TreeItemMap.Find(TreeItemID);
 			if (!Item)
@@ -2021,7 +2028,7 @@ void SSceneOutliner::OnHierarchyChangedEvent(FSceneOutlinerHierarchyChangedData
 			}
 		}
 
-		for (const auto& TreeItemPtr : Event.Items)
+		for (const FSceneOutlinerTreeItemPtr& TreeItemPtr : Event.Items)
 		{
 			if (TreeItemPtr.IsValid())
 			{
