SceneOutliner: dissociate bShowTransient from showing Level Instance hierarchy
- Remove some redundant filters now that this is done
- ISceneOutlinerMode::InitializeViewMenuExtender allows extending the menu (use Extension hooks) instead of only adding through ISceneOutlinerMode::CreateViewContent

#rnx

diff --git a/Engine/Source/Editor/DataLayerEditor/Private/DataLayer/DataLayerHierarchy.h b/Engine/Source/Editor/DataLayerEditor/Private/DataLayer/DataLayerHierarchy.h
--- a/Engine/Source/Editor/DataLayerEditor/Private/DataLayer/DataLayerHierarchy.h
+++ b/Engine/Source/Editor/DataLayerEditor/Private/DataLayer/DataLayerHierarchy.h
@@ -36,7 +36,7 @@ public:
 	void SetShowOnlySelectedActors(bool bInbShowOnlySelectedActors) { bShowOnlySelectedActors = bInbShowOnlySelectedActors; }
 	void SetHighlightSelectedDataLayers(bool bInHighlightSelectedDataLayers) { bHighlightSelectedDataLayers = bInHighlightSelectedDataLayers; }
 	void SetShowLevelInstanceContent(bool bInShowLevelInstanceContent) { bShowLevelInstanceContent = bInShowLevelInstanceContent; }
-
+	bool GetShowLevelInstanceContent() const { return bShowLevelInstanceContent; }
 private:
 	UWorld* GetOwningWorld() const;
 	FDataLayerHierarchy(FDataLayerMode* Mode, const TWeakObjectPtr<UWorld>& Worlds);
diff --git a/Engine/Source/Editor/DataLayerEditor/Private/DataLayer/DataLayerMode.cpp b/Engine/Source/Editor/DataLayerEditor/Private/DataLayer/DataLayerMode.cpp
--- a/Engine/Source/Editor/DataLayerEditor/Private/DataLayer/DataLayerMode.cpp
+++ b/Engine/Source/Editor/DataLayerEditor/Private/DataLayer/DataLayerMode.cpp
@@ -210,41 +210,29 @@ FDataLayerMode::FDataLayerMode(const FDataLayerModeParams& Params)
 	});
 	FilterInfoMap.Add(TEXT("HideUnloadedActorsFilter"), HideUnloadedActorsInfo);
 
-	SceneOutliner->AddFilter(MakeShared<TSceneOutlinerPredicateFilter<FDataLayerActorTreeItem>>(FDataLayerActorTreeItem::FFilterPredicate::CreateLambda([this](const AActor* Actor, const UDataLayerInstance* DataLayer)
-	{
-		return FActorMode::IsActorDisplayable(SceneOutliner, Actor);
-	}), FSceneOutlinerFilter::EDefaultBehaviour::Pass));
-
-	FSceneOutlinerFilterInfo HideLevelInstancesInfo(LOCTEXT("ToggleHideLevelInstanceContent", "Hide Level Instance Content"), LOCTEXT("ToggleHideLevelInstanceContentToolTip", "When enabled, hides all level instance content."), bHideLevelInstanceContent, FCreateSceneOutlinerFilter::CreateStatic(&FDataLayerMode::CreateHideLevelInstancesFilter));
-	HideLevelInstancesInfo.OnToggle().AddLambda([this](bool bIsActive)
-	{
-		UWorldPartitionEditorPerProjectUserSettings* Settings = GetMutableDefault<UWorldPartitionEditorPerProjectUserSettings>();
-		Settings->bHideLevelInstanceContent = bHideLevelInstanceContent = bIsActive;
-		Settings->PostEditChange();
-
-		if (auto DataLayerHierarchy = StaticCast<FDataLayerHierarchy*>(Hierarchy.Get()))
+		
+	// Add a an actor filter and interactive filter which sets the interactive mode of LevelInstance items and their children
+	SceneOutliner->AddFilter(MakeShared<FDataLayerActorFilter>(
+		FDataLayerActorTreeItem::FFilterPredicate::CreateLambda([this](const AActor* Actor, const UDataLayerInstance* DataLayer) 
 		{
-			DataLayerHierarchy->SetShowLevelInstanceContent(!bIsActive);
-		}
-	});
-	FilterInfoMap.Add(TEXT("HideLevelInstancesFilter"), HideLevelInstancesInfo);
-
-	// Add a filter which sets the interactive mode of LevelInstance items and their children
-	SceneOutliner->AddFilter(MakeShared<FDataLayerActorFilter>(FDataLayerActorTreeItem::FFilterPredicate::CreateStatic([](const AActor* Actor, const UDataLayerInstance* DataLayer) {return true; }), FSceneOutlinerFilter::EDefaultBehaviour::Pass, FDataLayerActorTreeItem::FFilterPredicate::CreateLambda([this](const AActor* Actor, const UDataLayerInstance* DataLayer)
-	{
-		if (!bHideLevelInstanceContent)
+			return FActorMode::IsActorDisplayable(SceneOutliner, Actor, !bHideLevelInstanceContent);
+		}), 
+		FSceneOutlinerFilter::EDefaultBehaviour::Pass, 
+		FDataLayerActorTreeItem::FInteractivePredicate::CreateLambda([this](const AActor* Actor, const UDataLayerInstance* DataLayer)
 		{
-			if (const ULevelInstanceSubsystem* LevelInstanceSubsystem = UWorld::GetSubsystem<ULevelInstanceSubsystem>(RepresentingWorld.Get()))
+			if (!bHideLevelInstanceContent)
 			{
-				const ILevelInstanceInterface* ParentLevelInstance = LevelInstanceSubsystem->GetParentLevelInstance(Actor);
-				if (ParentLevelInstance && !LevelInstanceSubsystem->IsEditingLevelInstance(ParentLevelInstance))
+				if (const ULevelInstanceSubsystem* LevelInstanceSubsystem = UWorld::GetSubsystem<ULevelInstanceSubsystem>(RepresentingWorld.Get()))
 				{
-					return false;
+					const ILevelInstanceInterface* ParentLevelInstance = LevelInstanceSubsystem->GetParentLevelInstance(Actor);
+					if (ParentLevelInstance && !LevelInstanceSubsystem->IsEditingLevelInstance(ParentLevelInstance))
+					{
+						return false;
+					}
 				}
 			}
-		}
-		return true;
-	})));
+			return true;
+		})));
 
 	DataLayerEditorSubsystem = UDataLayerEditorSubsystem::Get();
 	Rebuild();
@@ -277,26 +265,6 @@ TSharedRef<FSceneOutlinerFilter> FDataLayerMode::CreateHideUnloadedActorsFilter(
 	return MakeShareable(new FActorDescFilter(FActorDescTreeItem::FFilterPredicate::CreateStatic([](const FWorldPartitionActorDesc* ActorDesc) { return true; }), FSceneOutlinerFilter::EDefaultBehaviour::Pass));
 }
 
-TSharedRef<FSceneOutlinerFilter> FDataLayerMode::CreateHideLevelInstancesFilter()
-{
-	return MakeShareable(new FDataLayerActorFilter(FDataLayerActorTreeItem::FFilterPredicate::CreateStatic([](const AActor* Actor, const UDataLayerInstance* DataLayerInstance)
-	{
-		// Check if actor belongs to a LevelInstance
-		if (const ULevelInstanceSubsystem* LevelInstanceSubsystem = Actor->GetWorld()->GetSubsystem<ULevelInstanceSubsystem>())
-		{
-			if (const ILevelInstanceInterface* ParentLevelInstance = LevelInstanceSubsystem->GetParentLevelInstance(Actor))
-			{
-				if (!LevelInstanceSubsystem->IsEditingLevelInstance(ParentLevelInstance))
-				{
-					return false;
-				}
-			}
-		}
-		// Or if the actor itself is a LevelInstance editor instance
-		return !Actor->IsA<ALevelInstanceEditorInstanceActor>();
-	}), FSceneOutlinerFilter::EDefaultBehaviour::Pass));
-}
-
 int32 FDataLayerMode::GetTypeSortPriority(const ISceneOutlinerTreeItem& Item) const
 {
 	if (Item.IsA<FWorldDataLayersTreeItem>())
@@ -1718,14 +1686,42 @@ TSharedPtr<SWidget> FDataLayerMode::CreateContextMenu()
 	return nullptr;
 }
 
-void FDataLayerMode::CreateViewContent(FMenuBuilder& MenuBuilder)
+void FDataLayerMode::InitializeViewMenuExtender(TSharedPtr<FExtender> Extender)
 {
-	MenuBuilder.AddMenuEntry(
-		LOCTEXT("ToggleHighlightSelectedDataLayers", "Highlight Selected"),
-		LOCTEXT("ToggleHighlightSelectedDataLayersToolTip", "When enabled, highlights Data Layers containing actors that are currently selected."),
-		FSlateIcon(),
-		FUIAction(
-			FExecuteAction::CreateLambda([this]()
+	Extender->AddMenuExtension(SceneOutliner::ExtensionHooks::Show, EExtensionHook::First, nullptr, FMenuExtensionDelegate::CreateLambda([this](FMenuBuilder& MenuBuilder)
+	{
+		MenuBuilder.AddMenuEntry(
+			LOCTEXT("ToggleHideLevelInstanceContent", "Hide Level Instance Content"),
+			LOCTEXT("ToggleHideLevelInstanceContentToolTip", "When enabled, hides all level instance content."),
+			FSlateIcon(),
+			FUIAction(FExecuteAction::CreateLambda([this]()
+			{
+				UWorldPartitionEditorPerProjectUserSettings* Settings = GetMutableDefault<UWorldPartitionEditorPerProjectUserSettings>();
+				bHideLevelInstanceContent = !bHideLevelInstanceContent;
+				Settings->bHideLevelInstanceContent = bHideLevelInstanceContent;
+				Settings->PostEditChange();
+
+				if (auto DataLayerHierarchy = StaticCast<FDataLayerHierarchy*>(Hierarchy.Get()))
+				{
+					DataLayerHierarchy->SetShowLevelInstanceContent(!bHideLevelInstanceContent);
+				}
+
+				SceneOutliner->FullRefresh();
+			}),
+			FCanExecuteAction(),
+			FIsActionChecked::CreateLambda([this]() { return bHideLevelInstanceContent; })),
+			NAME_None,
+			EUserInterfaceActionType::ToggleButton
+		);
+	}));
+
+	Extender->AddMenuExtension(SceneOutliner::ExtensionHooks::Show, EExtensionHook::After, nullptr, FMenuExtensionDelegate::CreateLambda([this](FMenuBuilder& MenuBuilder)
+	{
+		MenuBuilder.AddMenuEntry(
+			LOCTEXT("ToggleHighlightSelectedDataLayers", "Highlight Selected"),
+			LOCTEXT("ToggleHighlightSelectedDataLayersToolTip", "When enabled, highlights Data Layers containing actors that are currently selected."),
+			FSlateIcon(),
+			FUIAction(FExecuteAction::CreateLambda([this]()
 			{
 				UWorldPartitionEditorPerProjectUserSettings* Settings = GetMutableDefault<UWorldPartitionEditorPerProjectUserSettings>();
 				bHighlightSelectedDataLayers = !bHighlightSelectedDataLayers;
@@ -1738,21 +1734,19 @@ void FDataLayerMode::CreateViewContent(FMenuBuilder& MenuBuilder)
 				}
 
 				SceneOutliner->FullRefresh();
-			}), 
+			}),
 			FCanExecuteAction(),
-			FIsActionChecked::CreateLambda([this]() { return bHighlightSelectedDataLayers; })
-		),
-		NAME_None,
-		EUserInterfaceActionType::ToggleButton
-	);
-
-	MenuBuilder.BeginSection("AssetThumbnails", LOCTEXT("ShowAdvancedHeading", "Advanced"));
-	MenuBuilder.AddMenuEntry(
-		LOCTEXT("ToggleAllowRuntimeDataLayerEditing", "Allow Runtime Data Layer Editing"), 
-		LOCTEXT("ToggleAllowRuntimeDataLayerEditingToolTip", "When enabled, allows editing of Runtime Data Layers."),
-		FSlateIcon(),
-		FUIAction(
-			FExecuteAction::CreateLambda([this]()
+			FIsActionChecked::CreateLambda([this]() { return bHighlightSelectedDataLayers; })),
+			NAME_None,
+			EUserInterfaceActionType::ToggleButton
+		);
+
+		MenuBuilder.BeginSection("Advanced", LOCTEXT("ShowAdvancedHeading", "Advanced"));
+		MenuBuilder.AddMenuEntry(
+			LOCTEXT("ToggleAllowRuntimeDataLayerEditing", "Allow Runtime Data Layer Editing"),
+			LOCTEXT("ToggleAllowRuntimeDataLayerEditingToolTip", "When enabled, allows editing of Runtime Data Layers."),
+			FSlateIcon(),
+			FUIAction(FExecuteAction::CreateLambda([this]()
 			{
 				if (AWorldDataLayers* WorldDataLayers = RepresentingWorld.IsValid() ? RepresentingWorld->GetWorldDataLayers() : nullptr)
 				{
@@ -1766,44 +1760,43 @@ void FDataLayerMode::CreateViewContent(FMenuBuilder& MenuBuilder)
 			{
 				const AWorldDataLayers* WorldDataLayers = RepresentingWorld.IsValid() ? RepresentingWorld->GetWorldDataLayers() : nullptr;
 				return WorldDataLayers ? WorldDataLayers->GetAllowRuntimeDataLayerEditing() : true;
-			})
-		),
-		NAME_None,
-		EUserInterfaceActionType::ToggleButton
-	);
+			})),
+			NAME_None,
+			EUserInterfaceActionType::ToggleButton
+		);
 
-	TArray<UDataLayerInstance*> AllDataLayers;
-	if (const UDataLayerManager* DataLayerManager = UDataLayerManager::GetDataLayerManager(GetOwningWorld()))
-	{
-		DataLayerManager->ForEachDataLayerInstance([&AllDataLayers](UDataLayerInstance* DataLayer)
+		TArray<UDataLayerInstance*> AllDataLayers;
+		if (const UDataLayerManager* DataLayerManager = UDataLayerManager::GetDataLayerManager(GetOwningWorld()))
 		{
-			AllDataLayers.Add(DataLayer);
-			return true;
-		});
-	}
+			DataLayerManager->ForEachDataLayerInstance([&AllDataLayers](UDataLayerInstance* DataLayer)
+			{
+				AllDataLayers.Add(DataLayer);
+				return true;
+			});
+		}
 
-	MenuBuilder.AddMenuEntry(
-		LOCTEXT("ResetDataLayerUserSettings", "Reset User Settings"),
-		LOCTEXT("ResetDataLayerUserSettingsToolTip", "Resets Data Layers User Settings to their initial values."),
-		FSlateIcon(),
-		FUIAction(
-			FExecuteAction::CreateLambda([this]()
+		MenuBuilder.AddMenuEntry(
+			LOCTEXT("ResetDataLayerUserSettings", "Reset User Settings"),
+			LOCTEXT("ResetDataLayerUserSettingsToolTip", "Resets Data Layers User Settings to their initial values."),
+			FSlateIcon(),
+			FUIAction(FExecuteAction::CreateLambda([this]()
 			{
 				const FScopedTransaction Transaction(LOCTEXT("ResetDataLayerUserSettings", "Reset User Settings"));
 				DataLayerEditorSubsystem->ResetUserSettings();
-			})
-		)
-	);
+			}))
+		);
 
-	MenuBuilder.EndSection();
+		MenuBuilder.EndSection();
 
-	MenuBuilder.BeginSection("AssetThumbnails", LOCTEXT("ShowWorldHeading", "World"));
-	MenuBuilder.AddSubMenu(
-		LOCTEXT("ChooseWorldSubMenu", "Choose World"),
-		LOCTEXT("ChooseWorldSubMenuToolTip", "Choose the world to display in the outliner."),
-		FNewMenuDelegate::CreateRaw(this, &FDataLayerMode::BuildWorldPickerMenu)
-	);
-	MenuBuilder.EndSection();
+		MenuBuilder.BeginSection("World", LOCTEXT("ShowWorldHeading", "World"));
+
+		MenuBuilder.AddSubMenu(
+			LOCTEXT("ChooseWorldSubMenu", "Choose World"),
+			LOCTEXT("ChooseWorldSubMenuToolTip", "Choose the world to display in the outliner."),
+			FNewMenuDelegate::CreateRaw(this, &FDataLayerMode::BuildWorldPickerMenu)
+		);
+		MenuBuilder.EndSection();
+	}));
 }
 
 void FDataLayerMode::BuildWorldPickerMenu(FMenuBuilder& MenuBuilder)
diff --git a/Engine/Source/Editor/DataLayerEditor/Private/DataLayer/DataLayerMode.h b/Engine/Source/Editor/DataLayerEditor/Private/DataLayer/DataLayerMode.h
--- a/Engine/Source/Editor/DataLayerEditor/Private/DataLayer/DataLayerMode.h
+++ b/Engine/Source/Editor/DataLayerEditor/Private/DataLayer/DataLayerMode.h
@@ -73,7 +73,7 @@ public:
 
 	virtual void Rebuild() override;
 	virtual TSharedPtr<SWidget> CreateContextMenu() override;
-	virtual void CreateViewContent(FMenuBuilder& MenuBuilder) override;
+	virtual void InitializeViewMenuExtender(TSharedPtr<FExtender> Extender) override;
 	virtual int32 GetTypeSortPriority(const ISceneOutlinerTreeItem& Item) const override;
 	virtual ESelectionMode::Type GetSelectionMode() const override { return ESelectionMode::Multi; }
 	virtual bool SupportsKeyboardFocus() const override { return true; }
diff --git a/Engine/Source/Editor/SceneOutliner/Private/ActorBrowsingMode.cpp b/Engine/Source/Editor/SceneOutliner/Private/ActorBrowsingMode.cpp
--- a/Engine/Source/Editor/SceneOutliner/Private/ActorBrowsingMode.cpp
+++ b/Engine/Source/Editor/SceneOutliner/Private/ActorBrowsingMode.cpp
@@ -99,6 +99,9 @@ FActorBrowsingMode::FActorBrowsingMode(SSceneOutliner* InSceneOutliner, TWeakObj
 		LocalSettings = *SavedSettings;
 	}
 
+	bHideLevelInstanceHierarchy = LocalSettings.bHideLevelInstanceHierarchy;
+	InSceneOutliner->SetShowTransient(!LocalSettings.bHideTemporaryActors);
+
 	// Get the OutlinerModule to register FilterInfos with the FilterInfoMap
 	FSceneOutlinerFilterInfo ShowOnlySelectedActorsInfo(LOCTEXT("ToggleShowOnlySelected", "Only Selected"), LOCTEXT("ToggleShowOnlySelectedToolTip", "When enabled, only displays actors that are currently selected."), LocalSettings.bShowOnlySelectedActors, FCreateSceneOutlinerFilter::CreateStatic(&FActorBrowsingMode::CreateShowOnlySelectedActorsFilter));
 	ShowOnlySelectedActorsInfo.OnToggle().AddLambda([this](bool bIsActive)
@@ -111,19 +114,7 @@ FActorBrowsingMode::FActorBrowsingMode(SSceneOutliner* InSceneOutliner, TWeakObj
 			}
 		});
 	FilterInfoMap.Add(TEXT("ShowOnlySelectedActors"), ShowOnlySelectedActorsInfo);
-
-	FSceneOutlinerFilterInfo HideTemporaryActorsInfo(LOCTEXT("ToggleHideTemporaryActors", "Hide Temporary Actors"), LOCTEXT("ToggleHideTemporaryActorsToolTip", "When enabled, hides temporary/run-time Actors."), LocalSettings.bHideTemporaryActors, FCreateSceneOutlinerFilter::CreateStatic(&FActorBrowsingMode::CreateHideTemporaryActorsFilter));
-	HideTemporaryActorsInfo.OnToggle().AddLambda([this](bool bIsActive)
-		{
-			FActorBrowsingModeConfig* Settings = GetMutableConfig();
-			if(Settings)
-			{
-				Settings->bHideTemporaryActors = bIsActive;
-				SaveConfig();
-			}
-		});
-	FilterInfoMap.Add(TEXT("HideTemporaryActors"), HideTemporaryActorsInfo);
-
+		
 	FSceneOutlinerFilterInfo OnlyCurrentLevelInfo(LOCTEXT("ToggleShowOnlyCurrentLevel", "Only in Current Level"), LOCTEXT("ToggleShowOnlyCurrentLevelToolTip", "When enabled, only shows Actors that are in the Current Level."), LocalSettings.bShowOnlyActorsInCurrentLevel, FCreateSceneOutlinerFilter::CreateStatic(&FActorBrowsingMode::CreateIsInCurrentLevelFilter));
 	OnlyCurrentLevelInfo.OnToggle().AddLambda([this](bool bIsActive)
 		{
@@ -218,24 +209,6 @@ FActorBrowsingMode::FActorBrowsingMode(SSceneOutliner* InSceneOutliner, TWeakObj
 
 	FilterInfoMap.Add(TEXT("HideComponentsFilter"), HideComponentsInfo);
 
-	bHideLevelInstanceHierarchy = LocalSettings.bHideLevelInstanceHierarchy;
-	FSceneOutlinerFilterInfo HideLevelInstancesInfo(LOCTEXT("ToggleHideLevelInstanceContent", "Hide Level Instance Content"), LOCTEXT("ToggleHideLevelInstancesToolTip", "When enabled, hides all level instance content."), LocalSettings.bHideLevelInstanceHierarchy, FCreateSceneOutlinerFilter::CreateStatic(&FActorBrowsingMode::CreateHideLevelInstancesFilter));
-	HideLevelInstancesInfo.OnToggle().AddLambda([this](bool bIsActive)
-		{
-			FActorBrowsingModeConfig* Settings = GetMutableConfig();
-			if(Settings)
-			{
-				Settings->bHideLevelInstanceHierarchy = bHideLevelInstanceHierarchy = bIsActive;
-				SaveConfig();
-			}
-
-			if (auto ActorHierarchy = StaticCast<FActorHierarchy*>(Hierarchy.Get()))
-			{
-				ActorHierarchy->SetShowingLevelInstances(!bIsActive);
-			}
-		});
-	FilterInfoMap.Add(TEXT("HideLevelInstancesFilter"), HideLevelInstancesInfo);
-
 	bHideUnloadedActors = LocalSettings.bHideUnloadedActors;
 	FSceneOutlinerFilterInfo HideUnloadedActorsInfo(LOCTEXT("ToggleHideUnloadedActors", "Hide Unloaded Actors"), LOCTEXT("ToggleHideUnloadedActorsToolTip", "When enabled, hides all unloaded world partition actors."), LocalSettings.bHideUnloadedActors, FCreateSceneOutlinerFilter::CreateStatic(&FActorBrowsingMode::CreateHideUnloadedActorsFilter));
 	HideUnloadedActorsInfo.OnToggle().AddLambda([this] (bool bIsActive)
@@ -273,7 +246,7 @@ FActorBrowsingMode::FActorBrowsingMode(SSceneOutliner* InSceneOutliner, TWeakObj
 	FilterInfoMap.Add(TEXT("HideEmptyFoldersFilter"), HideEmptyFoldersInfo);
 
 	// Add a filter which sets the interactive mode of LevelInstance items and their children
-	SceneOutliner->AddFilter(MakeShared<FActorFilter>(FActorTreeItem::FFilterPredicate::CreateStatic([](const AActor* Actor) {return true; }), FSceneOutlinerFilter::EDefaultBehaviour::Pass, FActorTreeItem::FFilterPredicate::CreateLambda([this](const AActor* Actor)
+	SceneOutliner->AddInteractiveFilter(MakeShared<FActorFilter>(FActorTreeItem::FFilterPredicate::CreateLambda([this](const AActor* Actor)
 		{
 			if (!bHideLevelInstanceHierarchy)
 			{
@@ -291,7 +264,7 @@ FActorBrowsingMode::FActorBrowsingMode(SSceneOutliner* InSceneOutliner, TWeakObj
 				}
 			}
 			return true;
-		})));
+		}), FSceneOutlinerFilter::EDefaultBehaviour::Pass));
 
 	bAlwaysFrameSelection = LocalSettings.bAlwaysFrameSelection;
 
@@ -434,40 +407,129 @@ void FActorBrowsingMode::OnToggleAlwaysFrameSelection()
 	}
 }
 
-bool FActorBrowsingMode::ShouldAlwaysFrameSelection()
+bool FActorBrowsingMode::ShouldAlwaysFrameSelection() const
 {
 	return bAlwaysFrameSelection;
 }
 
-void FActorBrowsingMode::CreateViewContent(FMenuBuilder& MenuBuilder)
+void FActorBrowsingMode::OnToggleHideTemporaryActors()
+{
+	FActorBrowsingModeConfig* Settings = GetMutableConfig();
+	if (Settings)
+	{
+		Settings->bHideTemporaryActors = !Settings->bHideTemporaryActors;
+		SaveConfig();
+
+		SceneOutliner->SetShowTransient(!Settings->bHideTemporaryActors);
+		
+		SceneOutliner->FullRefresh();
+	}
+}
+
+bool FActorBrowsingMode::ShouldHideTemporaryActors() const
 {
-	MenuBuilder.BeginSection("OutlinerSelectionOptions", LOCTEXT("OptionsHeading", "Options"));
+	const FActorBrowsingModeConfig* Settings = GetConstConfig();
+	if (Settings)
+	{
+		return Settings->bHideTemporaryActors;
+	}
+
+	return false;
+}
 
-	MenuBuilder.AddMenuEntry(
-		LOCTEXT("AlwaysFrameSelectionLabel", "Always Frame Selection"),
-		LOCTEXT("AlwaysFrameSelectionTooltip", "When enabled, selecting an Actor in the Viewport also scrolls to that Actor in the Outliner."),
-		FSlateIcon(),
-		FUIAction(
-			FExecuteAction::CreateRaw(this, &FActorBrowsingMode::OnToggleAlwaysFrameSelection),
-			FCanExecuteAction(),
-			FIsActionChecked::CreateRaw(this, &FActorBrowsingMode::ShouldAlwaysFrameSelection)
-		),
-		NAME_None,
-		EUserInterfaceActionType::ToggleButton
-	);
+void FActorBrowsingMode::OnToggleHideLevelInstanceHierarchy()
+{
+	bHideLevelInstanceHierarchy = !bHideLevelInstanceHierarchy;
 
-	MenuBuilder.EndSection();
+	FActorBrowsingModeConfig* Settings = GetMutableConfig();
+	if (Settings)
+	{
+		Settings->bHideLevelInstanceHierarchy = bHideLevelInstanceHierarchy;
 
+		SaveConfig();
+	}
 
-	MenuBuilder.BeginSection("AssetThumbnails", LOCTEXT("ShowWorldHeading", "World"));
+	if (auto ActorHierarchy = StaticCast<FActorHierarchy*>(Hierarchy.Get()))
 	{
+		ActorHierarchy->SetShowingLevelInstances(!bHideLevelInstanceHierarchy);
+	}
+
+	SceneOutliner->FullRefresh();
+}
+
+bool FActorBrowsingMode::ShouldHideLevelInstanceHierarchy() const
+{
+	const FActorBrowsingModeConfig* Settings = GetConstConfig();
+	if (Settings)
+	{
+		return Settings->bHideLevelInstanceHierarchy;
+	}
+
+	return false;
+}
+
+void FActorBrowsingMode::InitializeViewMenuExtender(TSharedPtr<FExtender> Extender)
+{
+	FActorModeInteractive::InitializeViewMenuExtender(Extender);
+	
+	Extender->AddMenuExtension(SceneOutliner::ExtensionHooks::Show, EExtensionHook::First, nullptr, FMenuExtensionDelegate::CreateLambda([this](FMenuBuilder& MenuBuilder)
+	{
+		MenuBuilder.AddMenuEntry(
+			LOCTEXT("ToggleHideTemporaryActors", "Hide Temporary Actors"),
+			LOCTEXT("ToggleHideTemporaryActorsToolTip", "When enabled, hides temporary/run-time Actors."),
+			FSlateIcon(),
+			FUIAction(
+				FExecuteAction::CreateRaw(this, &FActorBrowsingMode::OnToggleHideTemporaryActors),
+				FCanExecuteAction(),
+				FIsActionChecked::CreateRaw(this, &FActorBrowsingMode::ShouldHideTemporaryActors)
+			),
+			NAME_None,
+			EUserInterfaceActionType::ToggleButton
+		);
+
+		MenuBuilder.AddMenuEntry(
+			LOCTEXT("ToggleHideLevelInstanceContent", "Hide Level Instance Content"),
+			LOCTEXT("ToggleHideLevelInstancesToolTip", "When enabled, hides all level instance content."),
+			FSlateIcon(),
+			FUIAction(
+				FExecuteAction::CreateRaw(this, &FActorBrowsingMode::OnToggleHideLevelInstanceHierarchy),
+				FCanExecuteAction(),
+				FIsActionChecked::CreateRaw(this, &FActorBrowsingMode::ShouldHideLevelInstanceHierarchy)
+			),
+			NAME_None,
+			EUserInterfaceActionType::ToggleButton
+		);
+	}));
+
+	Extender->AddMenuExtension(SceneOutliner::ExtensionHooks::Show, EExtensionHook::After, nullptr, FMenuExtensionDelegate::CreateLambda([this](FMenuBuilder& MenuBuilder)
+	{
+		MenuBuilder.BeginSection("OutlinerSelectionOptions", LOCTEXT("OptionsHeading", "Options"));
+
+		MenuBuilder.AddMenuEntry(
+			LOCTEXT("AlwaysFrameSelectionLabel", "Always Frame Selection"),
+			LOCTEXT("AlwaysFrameSelectionTooltip", "When enabled, selecting an Actor in the Viewport also scrolls to that Actor in the Outliner."),
+			FSlateIcon(),
+			FUIAction(
+				FExecuteAction::CreateRaw(this, &FActorBrowsingMode::OnToggleAlwaysFrameSelection),
+				FCanExecuteAction(),
+				FIsActionChecked::CreateRaw(this, &FActorBrowsingMode::ShouldAlwaysFrameSelection)
+			),
+			NAME_None,
+			EUserInterfaceActionType::ToggleButton
+		);
+
+		MenuBuilder.EndSection();
+
+		MenuBuilder.BeginSection("World", LOCTEXT("ShowWorldHeading", "World"));
+		
 		MenuBuilder.AddSubMenu(
 			LOCTEXT("ChooseWorldSubMenu", "Choose World"),
 			LOCTEXT("ChooseWorldSubMenuToolTip", "Choose the world to display in the outliner."),
 			FNewMenuDelegate::CreateRaw(this, &FActorMode::BuildWorldPickerMenu)
 		);
-	}
-	MenuBuilder.EndSection();
+		
+		MenuBuilder.EndSection();
+	}));
 }
 
 TSharedRef<FSceneOutlinerFilter> FActorBrowsingMode::CreateShowOnlySelectedActorsFilter()
@@ -479,14 +541,6 @@ TSharedRef<FSceneOutlinerFilter> FActorBrowsingMode::CreateShowOnlySelectedActor
 	return MakeShareable(new FActorFilter(FActorTreeItem::FFilterPredicate::CreateStatic(IsActorSelected), FSceneOutlinerFilter::EDefaultBehaviour::Fail, FActorTreeItem::FFilterPredicate::CreateStatic(IsActorSelected)));
 }
 
-TSharedRef<FSceneOutlinerFilter> FActorBrowsingMode::CreateHideTemporaryActorsFilter()
-{
-	return MakeShareable(new FActorFilter(FActorTreeItem::FFilterPredicate::CreateStatic([](const AActor* InActor)
-		{
-			return ((InActor->GetWorld() && InActor->GetWorld()->WorldType != EWorldType::PIE) || GEditor->ObjectsThatExistInEditorWorld.Get(InActor)) && !InActor->HasAnyFlags(EObjectFlags::RF_Transient);
-		}), FSceneOutlinerFilter::EDefaultBehaviour::Pass));
-}
-
 TSharedRef<FSceneOutlinerFilter> FActorBrowsingMode::CreateIsInCurrentLevelFilter()
 {
 	return MakeShareable(new TSceneOutlinerPredicateFilter<ISceneOutlinerTreeItem>(ISceneOutlinerTreeItem::FFilterPredicate::CreateStatic([](const ISceneOutlinerTreeItem& InItem)
@@ -569,26 +623,6 @@ TSharedRef<FSceneOutlinerFilter> FActorBrowsingMode::CreateHideComponentsFilter(
 		FSceneOutlinerFilter::EDefaultBehaviour::Pass));
 }
 
-TSharedRef<FSceneOutlinerFilter> FActorBrowsingMode::CreateHideLevelInstancesFilter()
-{
-	return MakeShareable(new FActorFilter(FActorTreeItem::FFilterPredicate::CreateStatic([](const AActor* Actor)
-		{
-			// Check if actor belongs to a LevelInstance
-			if (const ULevelInstanceSubsystem* LevelInstanceSubsystem = Actor->GetWorld()->GetSubsystem<ULevelInstanceSubsystem>())
-			{
-				if (const ILevelInstanceInterface* ParentLevelInstance = LevelInstanceSubsystem->GetParentLevelInstance(Actor))
-				{
-					if (!LevelInstanceSubsystem->IsEditingLevelInstance(ParentLevelInstance))
-					{
-						return false;
-					}
-				}
-			}
-			// Or if the actor itself is a LevelInstance editor instance
-			return Cast<ALevelInstanceEditorInstanceActor>(Actor) == nullptr;
-		}), FSceneOutlinerFilter::EDefaultBehaviour::Pass));
-}
-
 TSharedRef<FSceneOutlinerFilter> FActorBrowsingMode::CreateHideUnloadedActorsFilter()
 {
 	return MakeShareable(new FActorDescFilter(FActorDescTreeItem::FFilterPredicate::CreateStatic(
diff --git a/Engine/Source/Editor/SceneOutliner/Private/ActorMode.cpp b/Engine/Source/Editor/SceneOutliner/Private/ActorMode.cpp
--- a/Engine/Source/Editor/SceneOutliner/Private/ActorMode.cpp
+++ b/Engine/Source/Editor/SceneOutliner/Private/ActorMode.cpp
@@ -331,7 +331,7 @@ void FActorMode::SynchronizeActorSelection()
 
 bool FActorMode::IsActorDisplayable(const AActor* Actor) const
 {
-	return FActorMode::IsActorDisplayable(SceneOutliner, Actor);
+	return FActorMode::IsActorDisplayable(SceneOutliner, Actor, !bHideLevelInstanceHierarchy);
 }
 
 FFolder::FRootObject FActorMode::GetRootObject() const
@@ -348,19 +348,36 @@ FFolder::FRootObject FActorMode::GetPasteTargetRootObject() const
 	return FFolder::GetInvalidRootObject();
 }
 
-bool FActorMode::IsActorDisplayable(const SSceneOutliner* SceneOutliner, const AActor* Actor)
+bool FActorMode::IsActorDisplayable(const SSceneOutliner* SceneOutliner, const AActor* Actor, bool bShowLevelInstanceContent)
 {
-	return Actor &&
-		!SceneOutliner->GetSharedData().bOnlyShowFolders && 									// Don't show actors if we're only showing folders
-		Actor->IsEditable() &&																	// Only show actors that are allowed to be selected and drawn in editor
-		Actor->IsListedInSceneOutliner() &&
-		(((Actor->GetWorld() && Actor->GetWorld()->IsPlayInEditor()) || !Actor->HasAnyFlags(RF_Transient)) ||
-		(SceneOutliner->GetSharedData().bShowTransient && Actor->HasAnyFlags(RF_Transient))) &&	// Don't show transient actors in non-play worlds
-		!Actor->IsTemplate() &&																	// Should never happen, but we never want CDOs displayed
-		!FActorEditorUtils::IsABuilderBrush(Actor) &&											// Don't show the builder brush
-		!Actor->IsA(AWorldSettings::StaticClass()) &&											// Don't show the WorldSettings actor, even though it is technically editable
-		IsValidChecked(Actor) &&																// We don't want to show actors that are about to go away
-		FLevelUtils::IsLevelVisible(Actor->GetLevel());											// Only show Actors whose level is visible
+	bool bIsActorDisplayable =	Actor &&
+								!SceneOutliner->GetSharedData().bOnlyShowFolders &&		// Don't show actors if we're only showing folders
+								Actor->IsEditable() &&									// Only show actors that are allowed to be selected and drawn in editor
+								Actor->IsListedInSceneOutliner();
+	
+	if(bIsActorDisplayable)
+	{
+		if (Actor->HasAnyFlags(RF_Transient))
+		{
+			// Level Instance transient actors are shown based on passed in bShowLevelInstanceContent flag
+			if (Actor->IsInLevelInstance())
+			{
+				bIsActorDisplayable = bShowLevelInstanceContent;
+			}
+			else
+			{
+				// Don't show transient actors in non-play worlds, except if bShowTransient is true
+				bIsActorDisplayable = SceneOutliner->GetSharedData().bShowTransient || (Actor->GetWorld() && Actor->GetWorld()->IsPlayInEditor());
+			}
+		}
+	}
+
+	return	bIsActorDisplayable &&																	// Previous results
+			!Actor->IsTemplate() &&																	// Should never happen, but we never want CDOs displayed
+			!FActorEditorUtils::IsABuilderBrush(Actor) &&											// Don't show the builder brush
+			!Actor->IsA(AWorldSettings::StaticClass()) &&											// Don't show the WorldSettings actor, even though it is technically editable
+			IsValidChecked(Actor) &&																// We don't want to show actors that are about to go away
+			FLevelUtils::IsLevelVisible(Actor->GetLevel());											// Only show Actors whose level is visible
 }
 
 bool FActorMode::CanInteract(const ISceneOutlinerTreeItem& Item) const
diff --git a/Engine/Source/Editor/SceneOutliner/Private/SSceneOutliner.cpp b/Engine/Source/Editor/SceneOutliner/Private/SSceneOutliner.cpp
--- a/Engine/Source/Editor/SceneOutliner/Private/SSceneOutliner.cpp
+++ b/Engine/Source/Editor/SceneOutliner/Private/SSceneOutliner.cpp
@@ -609,9 +609,11 @@ FSlateColor SSceneOutliner::GetViewButtonForegroundColor() const
 TSharedRef<SWidget> SSceneOutliner::GetViewButtonContent(bool bShowFilters)
 {
 	// Menu should stay open on selection if filters are not being shown
-	FMenuBuilder MenuBuilder(bShowFilters, NULL);
+	TSharedPtr<FExtender> MenuExtender = MakeShared<FExtender>();
+	Mode->InitializeViewMenuExtender(MenuExtender);
+	FMenuBuilder MenuBuilder(bShowFilters, nullptr, MenuExtender);
 
-	MenuBuilder.BeginSection("OutlinerSettings", LOCTEXT("HierarchyHeading", "Hierarchy"));
+	MenuBuilder.BeginSection(SceneOutliner::ExtensionHooks::Hierarchy, LOCTEXT("HierarchyHeading", "Hierarchy"));
 	{
 		MenuBuilder.AddMenuEntry(
 			LOCTEXT("ExpandAll", "Expand All"),
@@ -642,7 +644,7 @@ TSharedRef<SWidget> SSceneOutliner::GetViewButtonContent(bool bShowFilters)
 
 	if (bShowFilters)
 	{
-		MenuBuilder.BeginSection("AssetThumbnails", LOCTEXT("ShowHeading", "Show"));
+		MenuBuilder.BeginSection(SceneOutliner::ExtensionHooks::Show, LOCTEXT("ShowHeading", "Show"));
 		{
 			// Add mode filters
 			for (auto& ModeFilterInfo : Mode->GetFilterInfos())
diff --git a/Engine/Source/Editor/SceneOutliner/Public/ActorBrowsingMode.h b/Engine/Source/Editor/SceneOutliner/Public/ActorBrowsingMode.h
--- a/Engine/Source/Editor/SceneOutliner/Public/ActorBrowsingMode.h
+++ b/Engine/Source/Editor/SceneOutliner/Public/ActorBrowsingMode.h
@@ -16,7 +16,7 @@ public:
 	/* Begin ISceneOutlinerMode Interface */
 	SCENEOUTLINER_API virtual void Rebuild() override;
 	SCENEOUTLINER_API virtual FCreateSceneOutlinerMode CreateFolderPickerMode(const FFolder::FRootObject& InRootObject = FFolder::GetInvalidRootObject()) const override;
-	SCENEOUTLINER_API virtual void CreateViewContent(FMenuBuilder& MenuBuilder) override;
+	SCENEOUTLINER_API virtual void InitializeViewMenuExtender(TSharedPtr<FExtender> Extender) override;
 	SCENEOUTLINER_API virtual TSharedPtr<FDragDropOperation> CreateDragDropOperation(const FPointerEvent& MouseEvent, const TArray<FSceneOutlinerTreeItemPtr>& InTreeItems) const override;
 	SCENEOUTLINER_API virtual bool ParseDragDrop(FSceneOutlinerDragDropPayload& OutPayload, const FDragDropOperation& Operation) const override;
 	SCENEOUTLINER_API virtual FSceneOutlinerDragValidationInfo ValidateDrop(const ISceneOutlinerTreeItem& DropTarget, const FSceneOutlinerDragDropPayload& Payload) const override;
@@ -124,18 +124,22 @@ private:
 
 	/** Filter factories */
 	static TSharedRef<FSceneOutlinerFilter> CreateShowOnlySelectedActorsFilter();
-	static TSharedRef<FSceneOutlinerFilter> CreateHideTemporaryActorsFilter();
 	static TSharedRef<FSceneOutlinerFilter> CreateIsInCurrentLevelFilter();
 	static TSharedRef<FSceneOutlinerFilter> CreateIsInCurrentDataLayersFilter();
 	static TSharedRef<FSceneOutlinerFilter> CreateHideComponentsFilter();
-	static TSharedRef<FSceneOutlinerFilter> CreateHideLevelInstancesFilter();
 	static TSharedRef<FSceneOutlinerFilter> CreateHideUnloadedActorsFilter();
 	static TSharedRef<FSceneOutlinerFilter> CreateHideEmptyFoldersFilter();
 	TSharedRef<FSceneOutlinerFilter> CreateIsInCurrentContentBundleFilter();
 
 	/** Functions to expose selection framing to the UI */
 	void OnToggleAlwaysFrameSelection();
-	bool ShouldAlwaysFrameSelection();
+	bool ShouldAlwaysFrameSelection() const;
+
+	void OnToggleHideTemporaryActors();
+	bool ShouldHideTemporaryActors() const;
+
+	void OnToggleHideLevelInstanceHierarchy();
+	bool ShouldHideLevelInstanceHierarchy() const;
 
 	/**
 	 * Get a mutable version of the ActorBrowser config for setting values.
diff --git a/Engine/Source/Editor/SceneOutliner/Public/ActorMode.h b/Engine/Source/Editor/SceneOutliner/Public/ActorMode.h
--- a/Engine/Source/Editor/SceneOutliner/Public/ActorMode.h
+++ b/Engine/Source/Editor/SceneOutliner/Public/ActorMode.h
@@ -70,7 +70,7 @@ public:
 
 	virtual int32 GetTypeSortPriority(const ISceneOutlinerTreeItem& Item) const override;
 
-	static bool IsActorDisplayable(const SSceneOutliner* SceneOutliner, const AActor* Actor);
+	static bool IsActorDisplayable(const SSceneOutliner* SceneOutliner, const AActor* Actor, bool bShowLevelInstanceContent = false);
 	static bool IsActorLevelDisplayable(ULevel* InLevel);
 
 	virtual FFolder::FRootObject GetRootObject() const override;
diff --git a/Engine/Source/Editor/SceneOutliner/Public/ISceneOutlinerMode.h b/Engine/Source/Editor/SceneOutliner/Public/ISceneOutlinerMode.h
--- a/Engine/Source/Editor/SceneOutliner/Public/ISceneOutlinerMode.h
+++ b/Engine/Source/Editor/SceneOutliner/Public/ISceneOutlinerMode.h
@@ -37,6 +37,9 @@ public:
 	/** Fill the view button menu with content */
 	virtual void CreateViewContent(FMenuBuilder& MenuBuilder) {}
 
+	/** Allows extending SceneOutliner View Button MenuBuilder (allows using ExtensionHooks which isn't possible with CreateViewContent override) */
+	virtual void InitializeViewMenuExtender(TSharedPtr<FExtender> Extender) { }
+
 	/** Returns a factory to create a Folder Picker mode which matches this mode */
 	virtual FCreateSceneOutlinerMode CreateFolderPickerMode(const FFolder::FRootObject& InRootObject = FFolder::GetInvalidRootObject()) const { return FCreateSceneOutlinerMode(); }
 public:
diff --git a/Engine/Source/Editor/SceneOutliner/Public/SSceneOutliner.h b/Engine/Source/Editor/SceneOutliner/Public/SSceneOutliner.h
--- a/Engine/Source/Editor/SceneOutliner/Public/SSceneOutliner.h
+++ b/Engine/Source/Editor/SceneOutliner/Public/SSceneOutliner.h
@@ -82,6 +82,12 @@ namespace SceneOutliner
 			Rename = 1 << 2,
 		};
 	}
+
+	namespace ExtensionHooks
+	{
+		static FName Hierarchy(TEXT("Hierarchy"));
+		static FName Show(TEXT("Show"));
+	}
 }
 
 /**
@@ -510,6 +516,9 @@ public:
 	/** Get the unique identifier associated with this outliner */
 	FName GetOutlinerIdentifier() const { return OutlinerIdentifier; }
 
+	/** Toggle if SceneOutliner should show Transient objects */
+	void SetShowTransient(bool bInShowTransient) { SharedData.Get()->bShowTransient = bInShowTransient; }
+
 private:
 	/** Methods that implement structural modification logic for the tree */
 
