FSceneOutlinerModule::CreateActorPicker: New optional param bHideLevelInstanceHierarchy (defaults to true)

#rnx

diff --git a/Engine/Source/Editor/SceneOutliner/Private/SceneOutlinerModule.cpp b/Engine/Source/Editor/SceneOutliner/Private/SceneOutlinerModule.cpp
--- a/Engine/Source/Editor/SceneOutliner/Private/SceneOutlinerModule.cpp
+++ b/Engine/Source/Editor/SceneOutliner/Private/SceneOutlinerModule.cpp
@@ -77,7 +77,7 @@ TSharedRef<ISceneOutliner> FSceneOutlinerModule::CreateSceneOutliner(const FScen
 		.IsEnabled(FSlateApplication::Get().GetNormalExecutionAttribute());
 }
 
-TSharedRef<ISceneOutliner> FSceneOutlinerModule::CreateActorPicker(const FSceneOutlinerInitializationOptions& InInitOptions, const FOnActorPicked& OnActorPickedDelegate, TWeakObjectPtr<UWorld> SpecifiedWorld) const
+TSharedRef<ISceneOutliner> FSceneOutlinerModule::CreateActorPicker(const FSceneOutlinerInitializationOptions& InInitOptions, const FOnActorPicked& OnActorPickedDelegate, TWeakObjectPtr<UWorld> SpecifiedWorld, bool bInHideLevelInstanceHierarchy) const
 {
 	FSceneOutlinerInitializationOptions InitOptions(InInitOptions);
 	if (!InitOptions.ModeFactory.IsBound())
@@ -94,13 +94,13 @@ TSharedRef<ISceneOutliner> FSceneOutlinerModule::CreateActorPicker(const FSceneO
 				}
 			});
 
-		FCreateSceneOutlinerMode ModeFactory = FCreateSceneOutlinerMode::CreateLambda([OnItemPicked, SpecifiedWorld](SSceneOutliner* Outliner)
+		FCreateSceneOutlinerMode ModeFactory = FCreateSceneOutlinerMode::CreateLambda([OnItemPicked, SpecifiedWorld, bInHideLevelInstanceHierarchy](SSceneOutliner* Outliner)
 			{
 				FActorModeParams Params;
 				Params.SceneOutliner = Outliner;
 				Params.SpecifiedWorldToDisplay = SpecifiedWorld;
 				Params.bHideComponents = true;
-				Params.bHideLevelInstanceHierarchy = true;
+				Params.bHideLevelInstanceHierarchy = bInHideLevelInstanceHierarchy;
 				Params.bHideUnloadedActors = true;
 				Params.bHideEmptyFolders = true;
 				return new FActorPickingMode(Params, OnItemPicked);
diff --git a/Engine/Source/Editor/SceneOutliner/Public/SceneOutlinerModule.h b/Engine/Source/Editor/SceneOutliner/Public/SceneOutlinerModule.h
--- a/Engine/Source/Editor/SceneOutliner/Public/SceneOutlinerModule.h
+++ b/Engine/Source/Editor/SceneOutliner/Public/SceneOutlinerModule.h
@@ -46,7 +46,7 @@ public:
 	virtual TSharedRef<ISceneOutliner> CreateActorPicker(
 		const FSceneOutlinerInitializationOptions& InInitOptions,
 		const FOnActorPicked& OnActorPickedDelegate,
-		TWeakObjectPtr<UWorld> SpecifiedWorld = nullptr) const;
+		TWeakObjectPtr<UWorld> SpecifiedWorld = nullptr, bool bHideLevelInstanceHierarchy = true) const;
 
 	/** Creates a component picker widget. Calls the OnComponentPickedDelegate when an item is selected. */
 	virtual TSharedRef<ISceneOutliner> CreateComponentPicker(
