[Backout] - CL27523118
Original CL Desc
-----------------------------------------------------------------
SceneOutliner
- Remove pending removal operation when readding an item.

diff --git a/Engine/Source/Editor/SceneOutliner/Private/SSceneOutliner.cpp b/Engine/Source/Editor/SceneOutliner/Private/SSceneOutliner.cpp
--- a/Engine/Source/Editor/SceneOutliner/Private/SSceneOutliner.cpp
+++ b/Engine/Source/Editor/SceneOutliner/Private/SSceneOutliner.cpp
@@ -824,13 +824,6 @@ void SSceneOutliner::AddPendingItemAndChildren(FSceneOutlinerTreeItemPtr Item)
 	// Verify that there isn't already a pending operation for this item:
 	if (PendingTreeItemMap.Contains(Item->GetID()))
 	{
-		PendingOperations.RemoveAll(
-			[&Item](const SceneOutliner::FPendingTreeOperation& PendingOp)
-			{
-				return PendingOp.Item->GetID() == Item->GetID() && PendingOp.Type == SceneOutliner::FPendingTreeOperation::Removed;
-			});
-
-		PendingTreeItemMap_Removal.Remove(Item->GetID());
 		return;
 	}
 
@@ -838,7 +831,7 @@ void SSceneOutliner::AddPendingItemAndChildren(FSceneOutlinerTreeItemPtr Item)
 
 	TArray<FSceneOutlinerTreeItemPtr> Children;
 	Mode->GetHierarchy()->CreateChildren(Item, Children);
-	for (const FSceneOutlinerTreeItemPtr& Child : Children)
+	for (auto& Child : Children)
 	{
 		AddPendingItem(Child);
 	}
@@ -1978,7 +1971,7 @@ void SSceneOutliner::OnHierarchyChangedEvent(FSceneOutlinerHierarchyChangedData
 {
 	if (Event.Type == FSceneOutlinerHierarchyChangedData::Added)
 	{
-		for (const FSceneOutlinerTreeItemPtr& TreeItemPtr : Event.Items)
+		for (const auto& TreeItemPtr : Event.Items)
 		{
 			if(!TreeItemPtr.IsValid())
 			{
@@ -1998,7 +1991,7 @@ void SSceneOutliner::OnHierarchyChangedEvent(FSceneOutlinerHierarchyChangedData
 	}
 	else if (Event.Type == FSceneOutlinerHierarchyChangedData::Removed)
 	{
-		for (const FSceneOutlinerTreeItemID& TreeItemID : Event.ItemIDs)
+		for (const auto& TreeItemID : Event.ItemIDs)
 		{
 			FSceneOutlinerTreeItemPtr* Item = TreeItemMap.Find(TreeItemID);
 			if (!Item)
@@ -2016,7 +2009,7 @@ void SSceneOutliner::OnHierarchyChangedEvent(FSceneOutlinerHierarchyChangedData
 	}
 	else if (Event.Type == FSceneOutlinerHierarchyChangedData::Moved)
 	{
-		for (const FSceneOutlinerTreeItemID& TreeItemID : Event.ItemIDs)
+		for (const auto& TreeItemID : Event.ItemIDs)
 		{
 			FSceneOutlinerTreeItemPtr* Item = TreeItemMap.Find(TreeItemID);
 			if (!Item)
@@ -2030,7 +2023,7 @@ void SSceneOutliner::OnHierarchyChangedEvent(FSceneOutlinerHierarchyChangedData
 			}
 		}
 
-		for (const FSceneOutlinerTreeItemPtr& TreeItemPtr : Event.Items)
+		for (const auto& TreeItemPtr : Event.Items)
 		{
 			if (TreeItemPtr.IsValid())
 			{
