LevelInstance: Allow changing editor visibility on child actors of Level Instances (actors are transient so this isn't "undoable")

#rnx

diff --git a/Engine/Source/Editor/SceneOutliner/Private/ActorTreeItem.cpp b/Engine/Source/Editor/SceneOutliner/Private/ActorTreeItem.cpp
--- a/Engine/Source/Editor/SceneOutliner/Private/ActorTreeItem.cpp
+++ b/Engine/Source/Editor/SceneOutliner/Private/ActorTreeItem.cpp
@@ -352,13 +352,7 @@ bool FActorTreeItem::ShouldShowPinnedState() const
 
 bool FActorTreeItem::ShouldShowVisibilityState() const
 {
-	if (const AActor* ActorPtr = Actor.Get())
-	{
-		const ULevel* Level = ActorPtr->GetLevel();
-		return Level->IsPersistentLevel() || !Level->IsInstancedLevel();
-	}
-
-	return false;
+	return true;
 }
 
 void FActorTreeItem::OnVisibilityChanged(const bool bNewVisibility)
