Outliner: Move the drag and drop code from FActorBrowsingMode to FActorMode
TEDS: Add a selectionset member to the selection column

diff --git a/Engine/Source/Editor/SceneOutliner/Private/ActorBrowsingMode.cpp b/Engine/Source/Editor/SceneOutliner/Private/ActorBrowsingMode.cpp
--- a/Engine/Source/Editor/SceneOutliner/Private/ActorBrowsingMode.cpp
+++ b/Engine/Source/Editor/SceneOutliner/Private/ActorBrowsingMode.cpp
@@ -1390,83 +1390,6 @@ bool FActorBrowsingMode::CanPasteFoldersOnlyFromClipboard() const
 	return PasteString.StartsWith("BEGIN FOLDERLIST");
 }
 
-bool FActorBrowsingMode::GetFolderNamesFromPayload(const FSceneOutlinerDragDropPayload& InPayload, TArray<FName>& OutFolders, FFolder::FRootObject& OutCommonRootObject) const
-{
-	return FFolder::GetFolderPathsAndCommonRootObject(InPayload.GetData<FFolder>(SceneOutliner::FFolderPathSelector()), OutFolders, OutCommonRootObject);
-}
-
-TSharedPtr<FDragDropOperation> FActorBrowsingMode::CreateDragDropOperation(const FPointerEvent& MouseEvent, const TArray<FSceneOutlinerTreeItemPtr>& InTreeItems) const
-{
-	FSceneOutlinerDragDropPayload DraggedObjects(InTreeItems);
-
-	// If the drag contains only actors, we shortcut and create a simple FActorDragDropGraphEdOp rather than an FSceneOutlinerDragDrop composite op.
-	if (DraggedObjects.Has<FActorTreeItem>() && !DraggedObjects.Has<FFolderTreeItem>())
-	{
-		return FActorDragDropGraphEdOp::New(DraggedObjects.GetData<TWeakObjectPtr<AActor>>(SceneOutliner::FWeakActorSelector()));
-	}
-
-	TSharedPtr<FSceneOutlinerDragDropOp> OutlinerOp = MakeShareable(new FSceneOutlinerDragDropOp());
-
-	if (DraggedObjects.Has<FActorTreeItem>())
-	{
-		TSharedPtr<FActorDragDropOp> ActorOperation = MakeShareable(new FActorDragDropGraphEdOp);
-		ActorOperation->Init(DraggedObjects.GetData<TWeakObjectPtr<AActor>>(SceneOutliner::FWeakActorSelector()));
-		OutlinerOp->AddSubOp(ActorOperation);
-	}
-
-	if (DraggedObjects.Has<FFolderTreeItem>())
-	{
-		FFolder::FRootObject CommonRootObject;
-		TArray<FName> DraggedFolders;
-		if (GetFolderNamesFromPayload(DraggedObjects, DraggedFolders, CommonRootObject))
-		{
-			TSharedPtr<FFolderDragDropOp> FolderOperation = MakeShareable(new FFolderDragDropOp);
-			FolderOperation->Init(DraggedFolders, RepresentingWorld.Get(), CommonRootObject);
-			OutlinerOp->AddSubOp(FolderOperation);
-		}
-	}
-	OutlinerOp->Construct();
-	return OutlinerOp;
-}
-
-bool FActorBrowsingMode::ParseDragDrop(FSceneOutlinerDragDropPayload& OutPayload, const FDragDropOperation& Operation) const
-{
-	if (Operation.IsOfType<FSceneOutlinerDragDropOp>())
-	{
-		const auto& OutlinerOp = static_cast<const FSceneOutlinerDragDropOp&>(Operation);
-		if (const auto& FolderOp = OutlinerOp.GetSubOp<FFolderDragDropOp>())
-		{
-			for (const auto& Folder : FolderOp->Folders)
-			{
-				OutPayload.DraggedItems.Add(SceneOutliner->GetTreeItem(FFolder(FolderOp->RootObject, Folder)));
-			}
-		}
-		if (const auto& ActorOp = OutlinerOp.GetSubOp<FActorDragDropOp>())
-		{
-			for (const auto& Actor : ActorOp->Actors)
-			{
-				OutPayload.DraggedItems.Add(SceneOutliner->GetTreeItem(Actor.Get()));
-			}
-		}
-		return true;
-	}
-	else if (Operation.IsOfType<FActorDragDropOp>())
-	{
-		for (const TWeakObjectPtr<AActor>& Actor : static_cast<const FActorDragDropOp&>(Operation).Actors)
-		{
-			OutPayload.DraggedItems.Add(SceneOutliner->GetTreeItem(Actor.Get()));
-		}
-		return true;
-	}
-
-	return false;
-}
-
-FFolder FActorBrowsingMode::GetWorldDefaultRootFolder() const
-{
-	return FFolder::GetWorldRootFolder(RepresentingWorld.Get());
-}
-
 void FActorBrowsingMode::SynchronizeSelectedActorDescs()
 {
 	if (UWorldPartitionSubsystem* WorldPartitionSubsystem = UWorld::GetSubsystem<UWorldPartitionSubsystem>(RepresentingWorld.Get()))
@@ -1482,569 +1405,6 @@ void FActorBrowsingMode::SynchronizeSelectedActorDescs()
 	}
 }
 
-FSceneOutlinerDragValidationInfo FActorBrowsingMode::ValidateDrop(const ISceneOutlinerTreeItem& DropTarget, const FSceneOutlinerDragDropPayload& Payload) const
-{
-	if (Payload.Has<FFolderTreeItem>())
-	{
-		FFolder::FRootObject TargetRootObject = DropTarget.GetRootObject();
-		FFolder::FRootObject CommonPayloadFoldersRootObject;
-		TArray<FName> PayloadFolders;
-		const bool bHasCommonRootObject = GetFolderNamesFromPayload(Payload, PayloadFolders, CommonPayloadFoldersRootObject);
-		if (!bHasCommonRootObject)
-		{
-			return FSceneOutlinerDragValidationInfo(ESceneOutlinerDropCompatibility::IncompatibleGeneric, LOCTEXT("CantMoveFoldersWithMultipleRoots", "Cannot move folders with multiple roots"));
-		}
-		else if (CommonPayloadFoldersRootObject != TargetRootObject)
-		{
-			return FSceneOutlinerDragValidationInfo(ESceneOutlinerDropCompatibility::IncompatibleGeneric, LOCTEXT("CantChangeFoldersRoot", "Cannot change folders root"));
-		}
-	}
-
-	if (const FActorTreeItem* ActorItem = DropTarget.CastTo<FActorTreeItem>())
-	{
-		const AActor* ActorTarget = ActorItem->Actor.Get();
-		if (!ActorTarget)
-		{
-			return FSceneOutlinerDragValidationInfo(ESceneOutlinerDropCompatibility::IncompatibleGeneric, FText());
-		}
-
-		const ILevelInstanceInterface* LevelInstanceTarget = Cast<ILevelInstanceInterface>(ActorTarget);
-		const ULevelInstanceSubsystem* LevelInstanceSubsystem = RepresentingWorld->GetSubsystem<ULevelInstanceSubsystem>();
-		if (LevelInstanceTarget)
-		{
-			check(LevelInstanceSubsystem);
-			if (!LevelInstanceTarget->IsEditing())
-			{
-				return FSceneOutlinerDragValidationInfo(ESceneOutlinerDropCompatibility::IncompatibleGeneric, LOCTEXT("Error_AttachToClosedLevelInstance", "Cannot attach to LevelInstance which is not being edited"));
-			}
-		}
-		else
-		{
-			if (Payload.Has<FFolderTreeItem>())
-			{
-				return FSceneOutlinerDragValidationInfo(ESceneOutlinerDropCompatibility::IncompatibleGeneric, LOCTEXT("FoldersOnActorError", "Cannot attach folders to actors"));
-			}
-
-			if (!Payload.Has<FActorTreeItem>())
-			{
-				return FSceneOutlinerDragValidationInfo(ESceneOutlinerDropCompatibility::IncompatibleGeneric, FText());
-			}
-		}
-
-		FText AttachErrorMsg;
-		bool bCanAttach = true;
-		bool bDraggedOntoAttachmentParent = true;
-		const auto& DragActors = Payload.GetData<TWeakObjectPtr<AActor>>(SceneOutliner::FWeakActorSelector());
-		for (const auto& DragActorPtr : DragActors)
-		{
-			AActor* DragActor = DragActorPtr.Get();
-			if (DragActor)
-			{
-				if (bCanAttach)
-				{
-					if (LevelInstanceSubsystem)
-					{
-						// Either all actors must be in a LevelInstance or none of them
-						if (const ILevelInstanceInterface* ParentLevelInstance = LevelInstanceSubsystem->GetParentLevelInstance(DragActor))
-						{
-							if (!ParentLevelInstance->IsEditing())
-							{
-								return FSceneOutlinerDragValidationInfo(ESceneOutlinerDropCompatibility::IncompatibleGeneric, LOCTEXT("Error_RemoveEditingLevelInstance", "Cannot detach from a LevelInstance which is not being edited"));
-							}
-						}
-
-						if (!LevelInstanceSubsystem->CanMoveActorToLevel(DragActor, &AttachErrorMsg))
-						{
-							bCanAttach = bDraggedOntoAttachmentParent = false;
-							break;
-						}
-					}
-
-					if (DragActor->IsChildActor())
-					{
-						AttachErrorMsg = FText::Format(LOCTEXT("Error_AttachChildActor", "Cannot move {0} as it is a child actor."), FText::FromString(DragActor->GetActorLabel()));
-						bCanAttach = bDraggedOntoAttachmentParent = false;
-						break;
-					}
-					if (!LevelInstanceTarget && !GEditor->CanParentActors(ActorTarget, DragActor, &AttachErrorMsg))
-					{
-						bCanAttach = false;
-					}
-				}
-
-				if (DragActor->GetSceneOutlinerParent() != ActorTarget)
-				{
-					bDraggedOntoAttachmentParent = false;
-				}
-			}
-		}
-
-		const FText ActorLabel = FText::FromString(ActorTarget->GetActorLabel());
-		if (bDraggedOntoAttachmentParent)
-		{
-			if (DragActors.Num() == 1)
-			{
-				return FSceneOutlinerDragValidationInfo(ESceneOutlinerDropCompatibility::CompatibleDetach, ActorLabel);
-			}
-			else
-			{
-				return FSceneOutlinerDragValidationInfo(ESceneOutlinerDropCompatibility::CompatibleMultipleDetach, ActorLabel);
-			}
-		}
-		else if (bCanAttach)
-		{
-			if (DragActors.Num() == 1)
-			{
-				return FSceneOutlinerDragValidationInfo(ESceneOutlinerDropCompatibility::CompatibleAttach, ActorLabel);
-			}
-			else
-			{
-				return FSceneOutlinerDragValidationInfo(ESceneOutlinerDropCompatibility::CompatibleMultipleAttach, ActorLabel);
-			}
-		}
-		else
-		{
-			if (DragActors.Num() == 1)
-			{
-				return FSceneOutlinerDragValidationInfo(ESceneOutlinerDropCompatibility::IncompatibleGeneric, AttachErrorMsg);
-			}
-			else
-			{
-				const FText ReasonText = FText::Format(LOCTEXT("DropOntoText", "{0}. {1}"), ActorLabel, AttachErrorMsg);
-				return FSceneOutlinerDragValidationInfo(ESceneOutlinerDropCompatibility::IncompatibleMultipleAttach, ReasonText);
-			}
-		}
-	}
-	else if (DropTarget.IsA<FFolderTreeItem>() || DropTarget.IsA<FWorldTreeItem>() || DropTarget.IsA<FLevelTreeItem>())
-	{
-		const FFolderTreeItem* FolderItem = DropTarget.CastTo<FFolderTreeItem>();
-		const FWorldTreeItem* WorldItem = DropTarget.CastTo<FWorldTreeItem>();
-		const FLevelTreeItem* LevelItem = DropTarget.CastTo<FLevelTreeItem>();
-		// WorldTreeItem and LevelTreeItem are treated as root folders (path = none), with the difference that LevelTreeItem has a RootObject.
-		const FFolder DestinationPath = FolderItem ? FolderItem->GetFolder() : (LevelItem ? FFolder(FFolder::GetOptionalFolderRootObject(LevelItem->Level.Get()).Get(FFolder::GetInvalidRootObject())) : GetWorldDefaultRootFolder());
-		const FFolder::FRootObject& DestinationRootObject = DestinationPath.GetRootObject();
-		ILevelInstanceInterface* LevelInstanceTarget = Cast<ILevelInstanceInterface>(DestinationPath.GetRootObjectPtr());
-		if (LevelInstanceTarget && !LevelInstanceTarget->IsEditing())
-		{
-			return FSceneOutlinerDragValidationInfo(ESceneOutlinerDropCompatibility::IncompatibleGeneric, LOCTEXT("Error_DragInNonEditingLevelInstance", "Cannot drag into a LevelInstance which is not being edited"));
-		}
-
-		if (Payload.Has<FFolderTreeItem>())
-		{
-			FFolder::FRootObject CommonFolderRootObject;
-			TArray<FName> DraggedFolders;
-			if (GetFolderNamesFromPayload(Payload, DraggedFolders, CommonFolderRootObject))
-			{
-				// Iterate over all the folders that have been dragged
-				for (const FName& DraggedFolder : DraggedFolders)
-				{
-					const FName Leaf = FEditorFolderUtils::GetLeafName(DraggedFolder);
-					const FName Parent = FEditorFolderUtils::GetParentPath(DraggedFolder);
-
-					if ((CommonFolderRootObject != DestinationRootObject) && FFolder::IsRootObjectValid(CommonFolderRootObject) && FFolder::IsRootObjectValid(DestinationRootObject))
-					{
-						FFormatNamedArguments Args;
-						Args.Add(TEXT("SourceName"), FText::FromName(Leaf));
-						FText Text = FText::Format(LOCTEXT("CantChangeFolderRoot", "Cannot change {SourceName} folder root"), Args);
-						return FSceneOutlinerDragValidationInfo(ESceneOutlinerDropCompatibility::IncompatibleGeneric, Text);
-					}
-
-					if (Parent == DestinationPath.GetPath())
-					{
-						FFormatNamedArguments Args;
-						Args.Add(TEXT("SourceName"), FText::FromName(Leaf));
-
-						FText Text;
-						if (DestinationPath.IsNone())
-						{
-							Text = FText::Format(LOCTEXT("FolderAlreadyAssignedRoot", "{SourceName} is already assigned to root"), Args);
-						}
-						else
-						{
-							Args.Add(TEXT("DestPath"), FText::FromName(DestinationPath.GetPath()));
-							Text = FText::Format(LOCTEXT("FolderAlreadyAssigned", "{SourceName} is already assigned to {DestPath}"), Args);
-						}
-
-						return FSceneOutlinerDragValidationInfo(ESceneOutlinerDropCompatibility::IncompatibleGeneric, Text);
-					}
-
-					const FString DragFolderPath = DraggedFolder.ToString();
-					const FString LeafName = Leaf.ToString();
-					const FString DstFolderPath = DestinationPath.IsNone() ? FString() : DestinationPath.ToString();
-					const FString NewPath = DstFolderPath / LeafName;
-
-					if (FActorFolders::Get().ContainsFolder(*RepresentingWorld, FFolder(DestinationRootObject, FName(*NewPath))))
-					{
-						// The folder already exists
-						FFormatNamedArguments Args;
-						Args.Add(TEXT("DragName"), FText::FromString(LeafName));
-						return FSceneOutlinerDragValidationInfo(ESceneOutlinerDropCompatibility::IncompatibleGeneric,
-							FText::Format(LOCTEXT("FolderAlreadyExistsRoot", "A folder called \"{DragName}\" already exists at this level"), Args));
-					}
-					else if (DragFolderPath == DstFolderPath || DstFolderPath.StartsWith(DragFolderPath + "/"))
-					{
-						// Cannot drag as a child of itself
-						FFormatNamedArguments Args;
-						Args.Add(TEXT("FolderPath"), FText::FromName(DraggedFolder));
-						return FSceneOutlinerDragValidationInfo(
-							ESceneOutlinerDropCompatibility::IncompatibleGeneric,
-							FText::Format(LOCTEXT("ChildOfItself", "Cannot move \"{FolderPath}\" to be a child of itself"), Args));
-					}
-				}
-			}
-		}
-
-		if (Payload.Has<FActorTreeItem>())
-		{
-			const ULevelInstanceSubsystem* LevelInstanceSubsystem = RepresentingWorld->GetSubsystem<ULevelInstanceSubsystem>();
-			// Iterate over all the actors that have been dragged
-			for (const TWeakObjectPtr<AActor>& WeakActor : Payload.GetData<TWeakObjectPtr<AActor>>(SceneOutliner::FWeakActorSelector()))
-			{
-				const AActor* Actor = WeakActor.Get();
-
-				bool bActorContainedInLevelInstance = false;
-				if (LevelInstanceSubsystem)
-				{
-					if (const ILevelInstanceInterface* ParentLevelInstance = LevelInstanceSubsystem->GetParentLevelInstance(Actor))
-					{
-						if (!ParentLevelInstance->IsEditing())
-						{
-							return FSceneOutlinerDragValidationInfo(ESceneOutlinerDropCompatibility::IncompatibleGeneric, LOCTEXT("Error_RemoveEditingLevelInstance", "Cannot detach from a LevelInstance which is not being edited"));
-						}
-						bActorContainedInLevelInstance = true;
-					}
-
-					if (const ILevelInstanceInterface* LevelInstance = Cast<ILevelInstanceInterface>(Actor))
-					{
-						FText Reason;
-						if (!LevelInstanceSubsystem->CanMoveActorToLevel(Actor, &Reason))
-						{
-							return FSceneOutlinerDragValidationInfo(ESceneOutlinerDropCompatibility::IncompatibleGeneric, Reason);
-						}
-					}
-				}
-
-				if (Actor->IsChildActor())
-				{
-					return FSceneOutlinerDragValidationInfo(ESceneOutlinerDropCompatibility::IncompatibleGeneric, FText::Format(LOCTEXT("Error_AttachChildActor", "Cannot move {0} as it is a child actor."), FText::FromString(Actor->GetActorLabel())));
-				}
-				else if ((Actor->GetFolderRootObject() != DestinationRootObject) && FFolder::IsRootObjectValid(Actor->GetFolderRootObject()) && FFolder::IsRootObjectValid(DestinationRootObject))
-				{
-					FFormatNamedArguments Args;
-					Args.Add(TEXT("SourceName"), FText::FromString(Actor->GetActorLabel()));
-					FText Text = FText::Format(LOCTEXT("CantChangeActorRoot", "Cannot change {SourceName} folder root"), Args);
-					return FSceneOutlinerDragValidationInfo(ESceneOutlinerDropCompatibility::IncompatibleGeneric, Text);
-				}
-				else if (Actor->GetFolder() == DestinationPath && !Actor->GetSceneOutlinerParent() && !bActorContainedInLevelInstance)
-				{
-					FFormatNamedArguments Args;
-					Args.Add(TEXT("SourceName"), FText::FromString(Actor->GetActorLabel()));
-
-					FText Text;
-					if (DestinationPath.IsNone())
-					{
-						Text = FText::Format(LOCTEXT("FolderAlreadyAssignedRoot", "{SourceName} is already assigned to root"), Args);
-					}
-					else
-					{
-						Args.Add(TEXT("DestPath"), FText::FromName(DestinationPath.GetPath()));
-						Text = FText::Format(LOCTEXT("FolderAlreadyAssigned", "{SourceName} is already assigned to {DestPath}"), Args);
-					}
-
-					return FSceneOutlinerDragValidationInfo(ESceneOutlinerDropCompatibility::IncompatibleGeneric, Text);
-				}
-			}
-		}
-
-		// Everything else is a valid operation
-		if (DestinationPath.IsNone())
-		{
-			return FSceneOutlinerDragValidationInfo(ESceneOutlinerDropCompatibility::CompatibleGeneric, LOCTEXT("MoveToRoot", "Move to root"));
-		}
-		else
-		{
-			FFormatNamedArguments Args;
-			Args.Add(TEXT("DestPath"), FText::FromName(DestinationPath.GetPath()));
-			return FSceneOutlinerDragValidationInfo(ESceneOutlinerDropCompatibility::CompatibleGeneric, FText::Format(LOCTEXT("MoveInto", "Move into \"{DestPath}\""), Args));
-		}
-	}
-	else if (DropTarget.IsA<FComponentTreeItem>())
-	{
-		// we don't allow drag and drop on components for now
-		return FSceneOutlinerDragValidationInfo(ESceneOutlinerDropCompatibility::IncompatibleGeneric, FText());
-	}
-	return FSceneOutlinerDragValidationInfo::Invalid();
-}
-
-void FActorBrowsingMode::OnDrop(ISceneOutlinerTreeItem& DropTarget, const FSceneOutlinerDragDropPayload& Payload, const FSceneOutlinerDragValidationInfo& ValidationInfo) const
-{
-	if (const FActorTreeItem* ActorItem = DropTarget.CastTo<FActorTreeItem>())
-	{
-		AActor* DropActor = ActorItem->Actor.Get();
-		if (!DropActor)
-		{
-			return;
-		}
-
-		FMessageLog EditorErrors("EditorErrors");
-		EditorErrors.NewPage(LOCTEXT("ActorAttachmentsPageLabel", "Actor attachment"));
-
-		if (ValidationInfo.CompatibilityType == ESceneOutlinerDropCompatibility::CompatibleMultipleDetach || ValidationInfo.CompatibilityType == ESceneOutlinerDropCompatibility::CompatibleDetach)
-		{
-			const FScopedTransaction Transaction(LOCTEXT("UndoAction_DetachActors", "Detach actors"));
-
-			TArray<TWeakObjectPtr<AActor>> DraggedActors = Payload.GetData<TWeakObjectPtr<AActor>>(SceneOutliner::FWeakActorSelector());
-			for (const auto& WeakActor : DraggedActors)
-			{
-				if (auto* DragActor = WeakActor.Get())
-				{
-					// Detach from parent
-					USceneComponent* RootComp = DragActor->GetRootComponent();
-					if (RootComp && RootComp->GetAttachParent())
-					{
-						AActor* OldParent = RootComp->GetAttachParent()->GetOwner();
-						// Attachment is persisted on the child so modify both actors for Undo/Redo but do not mark the Parent package dirty
-						OldParent->Modify(/*bAlwaysMarkDirty=*/false);
-						RootComp->DetachFromComponent(FDetachmentTransformRules::KeepWorldTransform);
-
-						DragActor->SetFolderPath_Recursively(OldParent->GetFolderPath());
-					}
-				}
-			}
-		}
-		else if (ValidationInfo.CompatibilityType == ESceneOutlinerDropCompatibility::CompatibleMultipleAttach || ValidationInfo.CompatibilityType == ESceneOutlinerDropCompatibility::CompatibleAttach)
-		{
-			// Show socket chooser if we have sockets to select
-
-			if (ILevelInstanceInterface* TargetLevelInstance = Cast<ILevelInstanceInterface>(DropActor))
-			{
-				check(TargetLevelInstance->IsEditing());
-				const FScopedTransaction Transaction(LOCTEXT("UndoAction_MoveActorsToLevelInstance", "Move actors to LevelInstance"));
-
-				const FFolder DestinationPath = FFolder(FFolder::FRootObject(DropActor));
-				auto MoveToDestination = [&DestinationPath](FFolderTreeItem& Item)
-				{
-					Item.MoveTo(DestinationPath);
-				};
-				Payload.ForEachItem<FFolderTreeItem>(MoveToDestination);
-
-				// Since target root is directly the Level Instance, clear folder path
-				TArray<AActor*> DraggedActors = Payload.GetData<AActor*>(SceneOutliner::FActorSelector());
-				for (auto& Actor : DraggedActors)
-				{
-					Actor->SetFolderPath_Recursively(FName());
-				}
-
-				ULevelInstanceSubsystem* LevelInstanceSubsystem = RepresentingWorld->GetSubsystem<ULevelInstanceSubsystem>();
-				check(LevelInstanceSubsystem);
-				LevelInstanceSubsystem->MoveActorsTo(TargetLevelInstance, DraggedActors);
-			}
-			else
-			{
-				auto PerformAttachment = [](FName SocketName, TWeakObjectPtr<AActor> Parent, const TArray<TWeakObjectPtr<AActor>> NewAttachments)
-				{
-					AActor* ParentActor = Parent.Get();
-					if (ParentActor)
-					{
-						// modify parent and child
-						const FScopedTransaction Transaction(LOCTEXT("UndoAction_PerformAttachment", "Attach actors"));
-
-						// Attach each child
-						bool bAttached = false;
-						for (auto& Child : NewAttachments)
-						{
-							AActor* ChildActor = Child.Get();
-							if (GEditor->CanParentActors(ParentActor, ChildActor))
-							{
-								GEditor->ParentActors(ParentActor, ChildActor, SocketName);
-
-								ChildActor->SetFolderPath_Recursively(ParentActor->GetFolderPath());
-							}
-						}
-					}
-				};
-
-				TArray<TWeakObjectPtr<AActor>> DraggedActors = Payload.GetData<TWeakObjectPtr<AActor>>(SceneOutliner::FWeakActorSelector());
-				//@TODO: Should create a menu for each component that contains sockets, or have some form of disambiguation within the menu (like a fully qualified path)
-				// Instead, we currently only display the sockets on the root component
-				USceneComponent* Component = DropActor->GetRootComponent();
-				if ((Component != NULL) && (Component->HasAnySockets()))
-				{
-					// Create the popup
-					FSlateApplication::Get().PushMenu(
-						SceneOutliner->AsShared(),
-						FWidgetPath(),
-						SNew(SSocketChooserPopup)
-						.SceneComponent(Component)
-						.OnSocketChosen_Lambda(PerformAttachment, DropActor, MoveTemp(DraggedActors)),
-						FSlateApplication::Get().GetCursorPos(),
-						FPopupTransitionEffect(FPopupTransitionEffect::TypeInPopup)
-					);
-				}
-				else
-				{
-					PerformAttachment(NAME_None, DropActor, MoveTemp(DraggedActors));
-				}
-			}
-
-		}
-		// Report errors
-		EditorErrors.Notify(NSLOCTEXT("ActorAttachmentError", "AttachmentsFailed", "Attachments Failed!"));
-	}
-	else if (DropTarget.IsA<FFolderTreeItem>() || DropTarget.IsA<FWorldTreeItem>() || DropTarget.IsA<FLevelTreeItem>())
-	{
-		const FFolderTreeItem* FolderItem = DropTarget.CastTo<FFolderTreeItem>();
-		const FWorldTreeItem* WorldItem = DropTarget.CastTo<FWorldTreeItem>();
-		const FLevelTreeItem* LevelItem = DropTarget.CastTo<FLevelTreeItem>();
-		// WorldTreeItem and LevelTreeItem are treated as root folders (path = none), with the difference that LevelTreeItem has a RootObject.
-		const FFolder DestinationPath = FolderItem ? FolderItem->GetFolder() : (LevelItem ? FFolder(FFolder::GetOptionalFolderRootObject(LevelItem->Level.Get()).Get(FFolder::GetInvalidRootObject())) : GetWorldDefaultRootFolder());
-
-		const FScopedTransaction Transaction(LOCTEXT("MoveOutlinerItems", "Move World Outliner Items"));
-
-		auto MoveToDestination = [&DestinationPath](FFolderTreeItem& Item)
-		{
-			Item.MoveTo(DestinationPath);
-		};
-		Payload.ForEachItem<FFolderTreeItem>(MoveToDestination);
-
-		// Set the folder path on all the dragged actors, and detach any that need to be moved
-		if (Payload.Has<FActorTreeItem>())
-		{
-			TSet<const AActor*> ParentActors;
-			TSet<const AActor*> ChildActors;
-
-			TArray<AActor*> MovingActorsToValidRootObject;
-			Payload.ForEachItem<FActorTreeItem>([&DestinationPath, &ParentActors, &ChildActors, &MovingActorsToValidRootObject](const FActorTreeItem& ActorItem)
-			{
-				AActor* Actor = ActorItem.Actor.Get();
-				if (Actor)
-				{
-					// First mark this object as a parent, then set its children's path
-					ParentActors.Add(Actor);
-
-					const FFolder SrcFolder = Actor->GetFolder();
-
-					// If the folder root object changes, 1st pass will put actors at root. 2nd pass will set the destination path.
-					FName NewPath = (SrcFolder.GetRootObject() == DestinationPath.GetRootObject()) ? DestinationPath.GetPath() : NAME_None;
-
-					Actor->SetFolderPath(NewPath);
-					FActorEditorUtils::TraverseActorTree_ParentFirst(Actor, [&](AActor* InActor) {
-						ChildActors.Add(InActor);
-						InActor->SetFolderPath(NewPath);
-						return true;
-						}, false);
-
-					if ((Actor->GetFolderRootObject() != DestinationPath.GetRootObject()) && SrcFolder.IsRootObjectPersistentLevel() && (DestinationPath.IsRootObjectValid() && !DestinationPath.IsRootObjectPersistentLevel()))
-					{
-						MovingActorsToValidRootObject.Add(Actor);
-					}
-				}
-			});
-
-			// Detach parent actors
-			for (const AActor* Parent : ParentActors)
-			{
-				auto* RootComp = Parent->GetRootComponent();
-
-				// We don't detach if it's a child of another that's been dragged
-				if (RootComp && RootComp->GetAttachParent() && !ChildActors.Contains(Parent))
-				{
-					if (AActor* OldParentActor = RootComp->GetAttachParent()->GetOwner())
-					{
-						// Attachment is persisted on the child so modify both actors for Undo/Redo but do not mark the Parent package dirty
-						OldParentActor->Modify(/*bAlwaysMarkDirty=*/false);
-					}
-					RootComp->DetachFromComponent(FDetachmentTransformRules::KeepWorldTransform);
-				}
-			}
-
-			auto MoveActorsToLevel = [](const TArray<AActor*>& InActorsToMove, ULevel* InDestLevel, const FName& InDestinationPath)
-			{
-				// We are moving actors to another level
-				const bool bWarnAboutReferences = true;
-				const bool bWarnAboutRenaming = true;
-				const bool bMoveAllOrFail = true;
-				TArray<AActor*> MovedActors;
-				if (!EditorLevelUtils::MoveActorsToLevel(InActorsToMove, InDestLevel, bWarnAboutReferences, bWarnAboutRenaming, bMoveAllOrFail, &MovedActors))
-				{
-					UE_LOG(LogActorBrowser, Warning, TEXT("Failed to move actors because not all actors could be moved"));
-				}
-				// Once moved, update actors folder path
-				for (AActor* Actor : MovedActors)
-				{
-					Actor->SetFolderPath_Recursively(InDestinationPath);
-				}
-			};
-
-			if (DestinationPath.IsRootObjectPersistentLevel())
-			{
-				const ULevelInstanceSubsystem* LevelInstanceSubsystem = RepresentingWorld->GetSubsystem<ULevelInstanceSubsystem>();
-				check(LevelInstanceSubsystem);
-				ULevel* DestinationLevel = RepresentingWorld->PersistentLevel;
-				check(DestinationLevel);
-
-				TArray<AActor*> LevelInstanceActorsToMove;
-				TArray<AActor*> ActorsToMoveToPersistentLevel;
-				Payload.ForEachItem<FActorTreeItem>([LevelInstanceSubsystem, &LevelInstanceActorsToMove, &ActorsToMoveToPersistentLevel](const FActorTreeItem& ActorItem)
-				{
-					AActor* Actor = ActorItem.Actor.Get();
-					if (const ILevelInstanceInterface* ParentLevelInstance = LevelInstanceSubsystem->GetParentLevelInstance(Actor))
-					{
-						check(ParentLevelInstance->IsEditing());
-						LevelInstanceActorsToMove.Add(Actor);
-					}
-					else
-					{
-						const FFolder ActorSrcFolder = Actor->GetFolder();
-						if (ActorSrcFolder.IsRootObjectValid() && !ActorSrcFolder.IsRootObjectPersistentLevel())
-						{
-							ActorsToMoveToPersistentLevel.Add(Actor);
-						}
-					}
-				});
-
-				// We are moving actors outside of an editing level instance to a folder (or root) into the persistent level.
-				if (LevelInstanceActorsToMove.Num() > 0)
-				{
-					TArray<AActor*> MovedActors;
-					LevelInstanceSubsystem->MoveActorsToLevel(LevelInstanceActorsToMove, DestinationLevel, &MovedActors);
-					// Once moved, update actors folder path
-					for (AActor* Actor : MovedActors)
-					{
-						Actor->SetFolderPath_Recursively(DestinationPath.GetPath());
-					}
-				}
-				if (ActorsToMoveToPersistentLevel.Num() > 0)
-				{
-					MoveActorsToLevel(ActorsToMoveToPersistentLevel, DestinationLevel, DestinationPath.GetPath());
-				}
-			}
-			else if (MovingActorsToValidRootObject.Num())
-			{
-				if (ILevelInstanceInterface* TargetLevelInstance = Cast<ILevelInstanceInterface>(DestinationPath.GetRootObjectPtr()))
-				{
-					// We are moving actors inside an editing level instance
-					check(TargetLevelInstance->IsEditing());
-
-					ULevelInstanceSubsystem* LevelInstanceSubsystem = RepresentingWorld->GetSubsystem<ULevelInstanceSubsystem>();
-					check(LevelInstanceSubsystem);
-					TArray<AActor*> MovedActors;
-					LevelInstanceSubsystem->MoveActorsTo(TargetLevelInstance, MovingActorsToValidRootObject, &MovedActors);
-					// Once moved, update actors folder path
-					for (AActor* Actor : MovedActors)
-					{
-						Actor->SetFolderPath_Recursively(DestinationPath.GetPath());
-					}
-				}
-				else if (ULevel* DestinationLevel = Cast<ULevel>(DestinationPath.GetRootObjectPtr()))
-				{
-					MoveActorsToLevel(MovingActorsToValidRootObject, DestinationLevel, DestinationPath.GetPath());
-				}
-			}
-		}
-	}
-}
-
 FFolder FActorBrowsingMode::CreateNewFolder()
 {
 	const FScopedTransaction Transaction(LOCTEXT("UndoAction_CreateFolder", "Create Folder"));
diff --git a/Engine/Source/Editor/SceneOutliner/Private/ActorMode.cpp b/Engine/Source/Editor/SceneOutliner/Private/ActorMode.cpp
--- a/Engine/Source/Editor/SceneOutliner/Private/ActorMode.cpp
+++ b/Engine/Source/Editor/SceneOutliner/Private/ActorMode.cpp
@@ -15,13 +15,23 @@
 #include "FolderTreeItem.h"
 #include "ComponentTreeItem.h"
 #include "ActorDescTreeItem.h"
+#include "EditorActorFolders.h"
+#include "EditorFolderUtils.h"
+#include "EditorLevelUtils.h"
 #include "EditorModeManager.h"
 #include "WorldTreeItem.h"
 #include "LevelInstance/LevelInstanceInterface.h"
 #include "LevelInstance/LevelInstanceSubsystem.h"
 #include "LevelInstance/LevelInstanceEditorInstanceActor.h"
 #include "LevelEditor.h"
+#include "Logging/MessageLog.h"
+#include "DragAndDrop/ActorDragDropGraphEdOp.h"
+#include "DragAndDrop/FolderDragDropOp.h"
 #include "Modules/ModuleManager.h"
+#include "ScopedTransaction.h"
+#include "SSocketChooser.h"
+
+DEFINE_LOG_CATEGORY_STATIC(LogActorMode, Log, All);
 
 static int32 GSceneOutlinerAutoRepresentingWorldNetMode = NM_Client;
 static FAutoConsoleVariableRef CVarAutoRepresentingWorldNetMode(
@@ -481,4 +491,645 @@ int32 FActorMode::GetTypeSortPriority(const ISceneOutlinerTreeItem& Item) const
 	return -1;
 }
 
+FSceneOutlinerDragValidationInfo FActorMode::ValidateDrop(const ISceneOutlinerTreeItem& DropTarget, const FSceneOutlinerDragDropPayload& Payload) const
+{
+	if (Payload.Has<FFolderTreeItem>())
+	{
+		FFolder::FRootObject TargetRootObject = DropTarget.GetRootObject();
+		FFolder::FRootObject CommonPayloadFoldersRootObject;
+		TArray<FName> PayloadFolders;
+		const bool bHasCommonRootObject = GetFolderNamesFromPayload(Payload, PayloadFolders, CommonPayloadFoldersRootObject);
+		if (!bHasCommonRootObject)
+		{
+			return FSceneOutlinerDragValidationInfo(ESceneOutlinerDropCompatibility::IncompatibleGeneric, LOCTEXT("CantMoveFoldersWithMultipleRoots", "Cannot move folders with multiple roots"));
+		}
+		else if (CommonPayloadFoldersRootObject != TargetRootObject)
+		{
+			return FSceneOutlinerDragValidationInfo(ESceneOutlinerDropCompatibility::IncompatibleGeneric, LOCTEXT("CantChangeFoldersRoot", "Cannot change folders root"));
+		}
+	}
+
+	if (const FActorTreeItem* ActorItem = DropTarget.CastTo<FActorTreeItem>())
+	{
+		const AActor* ActorTarget = ActorItem->Actor.Get();
+		if (!ActorTarget)
+		{
+			return FSceneOutlinerDragValidationInfo(ESceneOutlinerDropCompatibility::IncompatibleGeneric, FText());
+		}
+
+		const ILevelInstanceInterface* LevelInstanceTarget = Cast<ILevelInstanceInterface>(ActorTarget);
+		const ULevelInstanceSubsystem* LevelInstanceSubsystem = RepresentingWorld->GetSubsystem<ULevelInstanceSubsystem>();
+		if (LevelInstanceTarget)
+		{
+			check(LevelInstanceSubsystem);
+			if (!LevelInstanceTarget->IsEditing())
+			{
+				return FSceneOutlinerDragValidationInfo(ESceneOutlinerDropCompatibility::IncompatibleGeneric, LOCTEXT("Error_AttachToClosedLevelInstance", "Cannot attach to LevelInstance which is not being edited"));
+			}
+		}
+		else
+		{
+			if (Payload.Has<FFolderTreeItem>())
+			{
+				return FSceneOutlinerDragValidationInfo(ESceneOutlinerDropCompatibility::IncompatibleGeneric, LOCTEXT("FoldersOnActorError", "Cannot attach folders to actors"));
+			}
+
+			if (!Payload.Has<FActorTreeItem>())
+			{
+				return FSceneOutlinerDragValidationInfo(ESceneOutlinerDropCompatibility::IncompatibleGeneric, FText());
+			}
+		}
+
+		FText AttachErrorMsg;
+		bool bCanAttach = true;
+		bool bDraggedOntoAttachmentParent = true;
+		const auto& DragActors = Payload.GetData<TWeakObjectPtr<AActor>>(SceneOutliner::FWeakActorSelector());
+		for (const auto& DragActorPtr : DragActors)
+		{
+			AActor* DragActor = DragActorPtr.Get();
+			if (DragActor)
+			{
+				if (bCanAttach)
+				{
+					if (LevelInstanceSubsystem)
+					{
+						// Either all actors must be in a LevelInstance or none of them
+						if (const ILevelInstanceInterface* ParentLevelInstance = LevelInstanceSubsystem->GetParentLevelInstance(DragActor))
+						{
+							if (!ParentLevelInstance->IsEditing())
+							{
+								return FSceneOutlinerDragValidationInfo(ESceneOutlinerDropCompatibility::IncompatibleGeneric, LOCTEXT("Error_RemoveEditingLevelInstance", "Cannot detach from a LevelInstance which is not being edited"));
+							}
+						}
+
+						if (!LevelInstanceSubsystem->CanMoveActorToLevel(DragActor, &AttachErrorMsg))
+						{
+							bCanAttach = bDraggedOntoAttachmentParent = false;
+							break;
+						}
+					}
+
+					if (DragActor->IsChildActor())
+					{
+						AttachErrorMsg = FText::Format(LOCTEXT("Error_AttachChildActor", "Cannot move {0} as it is a child actor."), FText::FromString(DragActor->GetActorLabel()));
+						bCanAttach = bDraggedOntoAttachmentParent = false;
+						break;
+					}
+					if (!LevelInstanceTarget && !GEditor->CanParentActors(ActorTarget, DragActor, &AttachErrorMsg))
+					{
+						bCanAttach = false;
+					}
+				}
+
+				if (DragActor->GetSceneOutlinerParent() != ActorTarget)
+				{
+					bDraggedOntoAttachmentParent = false;
+				}
+			}
+		}
+
+		const FText ActorLabel = FText::FromString(ActorTarget->GetActorLabel());
+		if (bDraggedOntoAttachmentParent)
+		{
+			if (DragActors.Num() == 1)
+			{
+				return FSceneOutlinerDragValidationInfo(ESceneOutlinerDropCompatibility::CompatibleDetach, ActorLabel);
+			}
+			else
+			{
+				return FSceneOutlinerDragValidationInfo(ESceneOutlinerDropCompatibility::CompatibleMultipleDetach, ActorLabel);
+			}
+		}
+		else if (bCanAttach)
+		{
+			if (DragActors.Num() == 1)
+			{
+				return FSceneOutlinerDragValidationInfo(ESceneOutlinerDropCompatibility::CompatibleAttach, ActorLabel);
+			}
+			else
+			{
+				return FSceneOutlinerDragValidationInfo(ESceneOutlinerDropCompatibility::CompatibleMultipleAttach, ActorLabel);
+			}
+		}
+		else
+		{
+			if (DragActors.Num() == 1)
+			{
+				return FSceneOutlinerDragValidationInfo(ESceneOutlinerDropCompatibility::IncompatibleGeneric, AttachErrorMsg);
+			}
+			else
+			{
+				const FText ReasonText = FText::Format(LOCTEXT("DropOntoText", "{0}. {1}"), ActorLabel, AttachErrorMsg);
+				return FSceneOutlinerDragValidationInfo(ESceneOutlinerDropCompatibility::IncompatibleMultipleAttach, ReasonText);
+			}
+		}
+	}
+	else if (DropTarget.IsA<FFolderTreeItem>() || DropTarget.IsA<FWorldTreeItem>() || DropTarget.IsA<FLevelTreeItem>())
+	{
+		const FFolderTreeItem* FolderItem = DropTarget.CastTo<FFolderTreeItem>();
+		const FWorldTreeItem* WorldItem = DropTarget.CastTo<FWorldTreeItem>();
+		const FLevelTreeItem* LevelItem = DropTarget.CastTo<FLevelTreeItem>();
+		// WorldTreeItem and LevelTreeItem are treated as root folders (path = none), with the difference that LevelTreeItem has a RootObject.
+		const FFolder DestinationPath = FolderItem ? FolderItem->GetFolder() : (LevelItem ? FFolder(FFolder::GetOptionalFolderRootObject(LevelItem->Level.Get()).Get(FFolder::GetInvalidRootObject())) : GetWorldDefaultRootFolder());
+		const FFolder::FRootObject& DestinationRootObject = DestinationPath.GetRootObject();
+		ILevelInstanceInterface* LevelInstanceTarget = Cast<ILevelInstanceInterface>(DestinationPath.GetRootObjectPtr());
+		if (LevelInstanceTarget && !LevelInstanceTarget->IsEditing())
+		{
+			return FSceneOutlinerDragValidationInfo(ESceneOutlinerDropCompatibility::IncompatibleGeneric, LOCTEXT("Error_DragInNonEditingLevelInstance", "Cannot drag into a LevelInstance which is not being edited"));
+		}
+
+		if (Payload.Has<FFolderTreeItem>())
+		{
+			FFolder::FRootObject CommonFolderRootObject;
+			TArray<FName> DraggedFolders;
+			if (GetFolderNamesFromPayload(Payload, DraggedFolders, CommonFolderRootObject))
+			{
+				// Iterate over all the folders that have been dragged
+				for (const FName& DraggedFolder : DraggedFolders)
+				{
+					const FName Leaf = FEditorFolderUtils::GetLeafName(DraggedFolder);
+					const FName Parent = FEditorFolderUtils::GetParentPath(DraggedFolder);
+
+					if ((CommonFolderRootObject != DestinationRootObject) && FFolder::IsRootObjectValid(CommonFolderRootObject) && FFolder::IsRootObjectValid(DestinationRootObject))
+					{
+						FFormatNamedArguments Args;
+						Args.Add(TEXT("SourceName"), FText::FromName(Leaf));
+						FText Text = FText::Format(LOCTEXT("CantChangeFolderRoot", "Cannot change {SourceName} folder root"), Args);
+						return FSceneOutlinerDragValidationInfo(ESceneOutlinerDropCompatibility::IncompatibleGeneric, Text);
+					}
+
+					if (Parent == DestinationPath.GetPath())
+					{
+						FFormatNamedArguments Args;
+						Args.Add(TEXT("SourceName"), FText::FromName(Leaf));
+
+						FText Text;
+						if (DestinationPath.IsNone())
+						{
+							Text = FText::Format(LOCTEXT("FolderAlreadyAssignedRoot", "{SourceName} is already assigned to root"), Args);
+						}
+						else
+						{
+							Args.Add(TEXT("DestPath"), FText::FromName(DestinationPath.GetPath()));
+							Text = FText::Format(LOCTEXT("FolderAlreadyAssigned", "{SourceName} is already assigned to {DestPath}"), Args);
+						}
+
+						return FSceneOutlinerDragValidationInfo(ESceneOutlinerDropCompatibility::IncompatibleGeneric, Text);
+					}
+
+					const FString DragFolderPath = DraggedFolder.ToString();
+					const FString LeafName = Leaf.ToString();
+					const FString DstFolderPath = DestinationPath.IsNone() ? FString() : DestinationPath.ToString();
+					const FString NewPath = DstFolderPath / LeafName;
+
+					if (FActorFolders::Get().ContainsFolder(*RepresentingWorld, FFolder(DestinationRootObject, FName(*NewPath))))
+					{
+						// The folder already exists
+						FFormatNamedArguments Args;
+						Args.Add(TEXT("DragName"), FText::FromString(LeafName));
+						return FSceneOutlinerDragValidationInfo(ESceneOutlinerDropCompatibility::IncompatibleGeneric,
+							FText::Format(LOCTEXT("FolderAlreadyExistsRoot", "A folder called \"{DragName}\" already exists at this level"), Args));
+					}
+					else if (DragFolderPath == DstFolderPath || DstFolderPath.StartsWith(DragFolderPath + "/"))
+					{
+						// Cannot drag as a child of itself
+						FFormatNamedArguments Args;
+						Args.Add(TEXT("FolderPath"), FText::FromName(DraggedFolder));
+						return FSceneOutlinerDragValidationInfo(
+							ESceneOutlinerDropCompatibility::IncompatibleGeneric,
+							FText::Format(LOCTEXT("ChildOfItself", "Cannot move \"{FolderPath}\" to be a child of itself"), Args));
+					}
+				}
+			}
+		}
+
+		if (Payload.Has<FActorTreeItem>())
+		{
+			const ULevelInstanceSubsystem* LevelInstanceSubsystem = RepresentingWorld->GetSubsystem<ULevelInstanceSubsystem>();
+			// Iterate over all the actors that have been dragged
+			for (const TWeakObjectPtr<AActor>& WeakActor : Payload.GetData<TWeakObjectPtr<AActor>>(SceneOutliner::FWeakActorSelector()))
+			{
+				const AActor* Actor = WeakActor.Get();
+
+				bool bActorContainedInLevelInstance = false;
+				if (LevelInstanceSubsystem)
+				{
+					if (const ILevelInstanceInterface* ParentLevelInstance = LevelInstanceSubsystem->GetParentLevelInstance(Actor))
+					{
+						if (!ParentLevelInstance->IsEditing())
+						{
+							return FSceneOutlinerDragValidationInfo(ESceneOutlinerDropCompatibility::IncompatibleGeneric, LOCTEXT("Error_RemoveEditingLevelInstance", "Cannot detach from a LevelInstance which is not being edited"));
+						}
+						bActorContainedInLevelInstance = true;
+					}
+
+					if (const ILevelInstanceInterface* LevelInstance = Cast<ILevelInstanceInterface>(Actor))
+					{
+						FText Reason;
+						if (!LevelInstanceSubsystem->CanMoveActorToLevel(Actor, &Reason))
+						{
+							return FSceneOutlinerDragValidationInfo(ESceneOutlinerDropCompatibility::IncompatibleGeneric, Reason);
+						}
+					}
+				}
+
+				if (Actor->IsChildActor())
+				{
+					return FSceneOutlinerDragValidationInfo(ESceneOutlinerDropCompatibility::IncompatibleGeneric, FText::Format(LOCTEXT("Error_AttachChildActor", "Cannot move {0} as it is a child actor."), FText::FromString(Actor->GetActorLabel())));
+				}
+				else if ((Actor->GetFolderRootObject() != DestinationRootObject) && FFolder::IsRootObjectValid(Actor->GetFolderRootObject()) && FFolder::IsRootObjectValid(DestinationRootObject))
+				{
+					FFormatNamedArguments Args;
+					Args.Add(TEXT("SourceName"), FText::FromString(Actor->GetActorLabel()));
+					FText Text = FText::Format(LOCTEXT("CantChangeActorRoot", "Cannot change {SourceName} folder root"), Args);
+					return FSceneOutlinerDragValidationInfo(ESceneOutlinerDropCompatibility::IncompatibleGeneric, Text);
+				}
+				else if (Actor->GetFolder() == DestinationPath && !Actor->GetSceneOutlinerParent() && !bActorContainedInLevelInstance)
+				{
+					FFormatNamedArguments Args;
+					Args.Add(TEXT("SourceName"), FText::FromString(Actor->GetActorLabel()));
+
+					FText Text;
+					if (DestinationPath.IsNone())
+					{
+						Text = FText::Format(LOCTEXT("FolderAlreadyAssignedRoot", "{SourceName} is already assigned to root"), Args);
+					}
+					else
+					{
+						Args.Add(TEXT("DestPath"), FText::FromName(DestinationPath.GetPath()));
+						Text = FText::Format(LOCTEXT("FolderAlreadyAssigned", "{SourceName} is already assigned to {DestPath}"), Args);
+					}
+
+					return FSceneOutlinerDragValidationInfo(ESceneOutlinerDropCompatibility::IncompatibleGeneric, Text);
+				}
+			}
+		}
+
+		// Everything else is a valid operation
+		if (DestinationPath.IsNone())
+		{
+			return FSceneOutlinerDragValidationInfo(ESceneOutlinerDropCompatibility::CompatibleGeneric, LOCTEXT("MoveToRoot", "Move to root"));
+		}
+		else
+		{
+			FFormatNamedArguments Args;
+			Args.Add(TEXT("DestPath"), FText::FromName(DestinationPath.GetPath()));
+			return FSceneOutlinerDragValidationInfo(ESceneOutlinerDropCompatibility::CompatibleGeneric, FText::Format(LOCTEXT("MoveInto", "Move into \"{DestPath}\""), Args));
+		}
+	}
+	else if (DropTarget.IsA<FComponentTreeItem>())
+	{
+		// we don't allow drag and drop on components for now
+		return FSceneOutlinerDragValidationInfo(ESceneOutlinerDropCompatibility::IncompatibleGeneric, FText());
+	}
+	return FSceneOutlinerDragValidationInfo::Invalid();
+}
+
+TSharedPtr<FDragDropOperation> FActorMode::CreateDragDropOperation(const FPointerEvent& MouseEvent, const TArray<FSceneOutlinerTreeItemPtr>& InTreeItems) const
+{
+	FSceneOutlinerDragDropPayload DraggedObjects(InTreeItems);
+
+	// If the drag contains only actors, we shortcut and create a simple FActorDragDropGraphEdOp rather than an FSceneOutlinerDragDrop composite op.
+	if (DraggedObjects.Has<FActorTreeItem>() && !DraggedObjects.Has<FFolderTreeItem>())
+	{
+		return FActorDragDropGraphEdOp::New(DraggedObjects.GetData<TWeakObjectPtr<AActor>>(SceneOutliner::FWeakActorSelector()));
+	}
+
+	TSharedPtr<FSceneOutlinerDragDropOp> OutlinerOp = MakeShareable(new FSceneOutlinerDragDropOp());
+
+	if (DraggedObjects.Has<FActorTreeItem>())
+	{
+		TSharedPtr<FActorDragDropOp> ActorOperation = MakeShareable(new FActorDragDropGraphEdOp);
+		ActorOperation->Init(DraggedObjects.GetData<TWeakObjectPtr<AActor>>(SceneOutliner::FWeakActorSelector()));
+		OutlinerOp->AddSubOp(ActorOperation);
+	}
+
+	if (DraggedObjects.Has<FFolderTreeItem>())
+	{
+		FFolder::FRootObject CommonRootObject;
+		TArray<FName> DraggedFolders;
+		if (GetFolderNamesFromPayload(DraggedObjects, DraggedFolders, CommonRootObject))
+		{
+			TSharedPtr<FFolderDragDropOp> FolderOperation = MakeShareable(new FFolderDragDropOp);
+			FolderOperation->Init(DraggedFolders, RepresentingWorld.Get(), CommonRootObject);
+			OutlinerOp->AddSubOp(FolderOperation);
+		}
+	}
+	OutlinerOp->Construct();
+	return OutlinerOp;
+}
+
+bool FActorMode::ParseDragDrop(FSceneOutlinerDragDropPayload& OutPayload, const FDragDropOperation& Operation) const
+{
+	if (Operation.IsOfType<FSceneOutlinerDragDropOp>())
+	{
+		const auto& OutlinerOp = static_cast<const FSceneOutlinerDragDropOp&>(Operation);
+		if (const auto& FolderOp = OutlinerOp.GetSubOp<FFolderDragDropOp>())
+		{
+			for (const auto& Folder : FolderOp->Folders)
+			{
+				OutPayload.DraggedItems.Add(SceneOutliner->GetTreeItem(FFolder(FolderOp->RootObject, Folder)));
+			}
+		}
+		if (const auto& ActorOp = OutlinerOp.GetSubOp<FActorDragDropOp>())
+		{
+			for (const auto& Actor : ActorOp->Actors)
+			{
+				OutPayload.DraggedItems.Add(SceneOutliner->GetTreeItem(Actor.Get()));
+			}
+		}
+		return true;
+	}
+	else if (Operation.IsOfType<FActorDragDropOp>())
+	{
+		for (const TWeakObjectPtr<AActor>& Actor : static_cast<const FActorDragDropOp&>(Operation).Actors)
+		{
+			OutPayload.DraggedItems.Add(SceneOutliner->GetTreeItem(Actor.Get()));
+		}
+		return true;
+	}
+
+	return false;
+}
+
+void FActorMode::OnDrop(ISceneOutlinerTreeItem& DropTarget, const FSceneOutlinerDragDropPayload& Payload, const FSceneOutlinerDragValidationInfo& ValidationInfo) const
+{
+	if (const FActorTreeItem* ActorItem = DropTarget.CastTo<FActorTreeItem>())
+	{
+		AActor* DropActor = ActorItem->Actor.Get();
+		if (!DropActor)
+		{
+			return;
+		}
+
+		FMessageLog EditorErrors("EditorErrors");
+		EditorErrors.NewPage(LOCTEXT("ActorAttachmentsPageLabel", "Actor attachment"));
+
+		if (ValidationInfo.CompatibilityType == ESceneOutlinerDropCompatibility::CompatibleMultipleDetach || ValidationInfo.CompatibilityType == ESceneOutlinerDropCompatibility::CompatibleDetach)
+		{
+			const FScopedTransaction Transaction(LOCTEXT("UndoAction_DetachActors", "Detach actors"));
+
+			TArray<TWeakObjectPtr<AActor>> DraggedActors = Payload.GetData<TWeakObjectPtr<AActor>>(SceneOutliner::FWeakActorSelector());
+			for (const auto& WeakActor : DraggedActors)
+			{
+				if (auto* DragActor = WeakActor.Get())
+				{
+					// Detach from parent
+					USceneComponent* RootComp = DragActor->GetRootComponent();
+					if (RootComp && RootComp->GetAttachParent())
+					{
+						AActor* OldParent = RootComp->GetAttachParent()->GetOwner();
+						// Attachment is persisted on the child so modify both actors for Undo/Redo but do not mark the Parent package dirty
+						OldParent->Modify(/*bAlwaysMarkDirty=*/false);
+						RootComp->DetachFromComponent(FDetachmentTransformRules::KeepWorldTransform);
+
+						DragActor->SetFolderPath_Recursively(OldParent->GetFolderPath());
+					}
+				}
+			}
+		}
+		else if (ValidationInfo.CompatibilityType == ESceneOutlinerDropCompatibility::CompatibleMultipleAttach || ValidationInfo.CompatibilityType == ESceneOutlinerDropCompatibility::CompatibleAttach)
+		{
+			// Show socket chooser if we have sockets to select
+
+			if (ILevelInstanceInterface* TargetLevelInstance = Cast<ILevelInstanceInterface>(DropActor))
+			{
+				check(TargetLevelInstance->IsEditing());
+				const FScopedTransaction Transaction(LOCTEXT("UndoAction_MoveActorsToLevelInstance", "Move actors to LevelInstance"));
+
+				const FFolder DestinationPath = FFolder(FFolder::FRootObject(DropActor));
+				auto MoveToDestination = [&DestinationPath](FFolderTreeItem& Item)
+				{
+					Item.MoveTo(DestinationPath);
+				};
+				Payload.ForEachItem<FFolderTreeItem>(MoveToDestination);
+
+				// Since target root is directly the Level Instance, clear folder path
+				TArray<AActor*> DraggedActors = Payload.GetData<AActor*>(SceneOutliner::FActorSelector());
+				for (auto& Actor : DraggedActors)
+				{
+					Actor->SetFolderPath_Recursively(FName());
+				}
+
+				ULevelInstanceSubsystem* LevelInstanceSubsystem = RepresentingWorld->GetSubsystem<ULevelInstanceSubsystem>();
+				check(LevelInstanceSubsystem);
+				LevelInstanceSubsystem->MoveActorsTo(TargetLevelInstance, DraggedActors);
+			}
+			else
+			{
+				auto PerformAttachment = [](FName SocketName, TWeakObjectPtr<AActor> Parent, const TArray<TWeakObjectPtr<AActor>> NewAttachments)
+				{
+					AActor* ParentActor = Parent.Get();
+					if (ParentActor)
+					{
+						// modify parent and child
+						const FScopedTransaction Transaction(LOCTEXT("UndoAction_PerformAttachment", "Attach actors"));
+
+						// Attach each child
+						bool bAttached = false;
+						for (auto& Child : NewAttachments)
+						{
+							AActor* ChildActor = Child.Get();
+							if (GEditor->CanParentActors(ParentActor, ChildActor))
+							{
+								GEditor->ParentActors(ParentActor, ChildActor, SocketName);
+
+								ChildActor->SetFolderPath_Recursively(ParentActor->GetFolderPath());
+							}
+						}
+					}
+				};
+
+				TArray<TWeakObjectPtr<AActor>> DraggedActors = Payload.GetData<TWeakObjectPtr<AActor>>(SceneOutliner::FWeakActorSelector());
+				//@TODO: Should create a menu for each component that contains sockets, or have some form of disambiguation within the menu (like a fully qualified path)
+				// Instead, we currently only display the sockets on the root component
+				USceneComponent* Component = DropActor->GetRootComponent();
+				if ((Component != NULL) && (Component->HasAnySockets()))
+				{
+					// Create the popup
+					FSlateApplication::Get().PushMenu(
+						SceneOutliner->AsShared(),
+						FWidgetPath(),
+						SNew(SSocketChooserPopup)
+						.SceneComponent(Component)
+						.OnSocketChosen_Lambda(PerformAttachment, DropActor, MoveTemp(DraggedActors)),
+						FSlateApplication::Get().GetCursorPos(),
+						FPopupTransitionEffect(FPopupTransitionEffect::TypeInPopup)
+					);
+				}
+				else
+				{
+					PerformAttachment(NAME_None, DropActor, MoveTemp(DraggedActors));
+				}
+			}
+
+		}
+		// Report errors
+		EditorErrors.Notify(NSLOCTEXT("ActorAttachmentError", "AttachmentsFailed", "Attachments Failed!"));
+	}
+	else if (DropTarget.IsA<FFolderTreeItem>() || DropTarget.IsA<FWorldTreeItem>() || DropTarget.IsA<FLevelTreeItem>())
+	{
+		const FFolderTreeItem* FolderItem = DropTarget.CastTo<FFolderTreeItem>();
+		const FWorldTreeItem* WorldItem = DropTarget.CastTo<FWorldTreeItem>();
+		const FLevelTreeItem* LevelItem = DropTarget.CastTo<FLevelTreeItem>();
+		// WorldTreeItem and LevelTreeItem are treated as root folders (path = none), with the difference that LevelTreeItem has a RootObject.
+		const FFolder DestinationPath = FolderItem ? FolderItem->GetFolder() : (LevelItem ? FFolder(FFolder::GetOptionalFolderRootObject(LevelItem->Level.Get()).Get(FFolder::GetInvalidRootObject())) : GetWorldDefaultRootFolder());
+
+		const FScopedTransaction Transaction(LOCTEXT("MoveOutlinerItems", "Move World Outliner Items"));
+
+		auto MoveToDestination = [&DestinationPath](FFolderTreeItem& Item)
+		{
+			Item.MoveTo(DestinationPath);
+		};
+		Payload.ForEachItem<FFolderTreeItem>(MoveToDestination);
+
+		// Set the folder path on all the dragged actors, and detach any that need to be moved
+		if (Payload.Has<FActorTreeItem>())
+		{
+			TSet<const AActor*> ParentActors;
+			TSet<const AActor*> ChildActors;
+
+			TArray<AActor*> MovingActorsToValidRootObject;
+			Payload.ForEachItem<FActorTreeItem>([&DestinationPath, &ParentActors, &ChildActors, &MovingActorsToValidRootObject](const FActorTreeItem& ActorItem)
+			{
+				AActor* Actor = ActorItem.Actor.Get();
+				if (Actor)
+				{
+					// First mark this object as a parent, then set its children's path
+					ParentActors.Add(Actor);
+
+					const FFolder SrcFolder = Actor->GetFolder();
+
+					// If the folder root object changes, 1st pass will put actors at root. 2nd pass will set the destination path.
+					FName NewPath = (SrcFolder.GetRootObject() == DestinationPath.GetRootObject()) ? DestinationPath.GetPath() : NAME_None;
+
+					Actor->SetFolderPath(NewPath);
+					FActorEditorUtils::TraverseActorTree_ParentFirst(Actor, [&](AActor* InActor) {
+						ChildActors.Add(InActor);
+						InActor->SetFolderPath(NewPath);
+						return true;
+						}, false);
+
+					if ((Actor->GetFolderRootObject() != DestinationPath.GetRootObject()) && SrcFolder.IsRootObjectPersistentLevel() && (DestinationPath.IsRootObjectValid() && !DestinationPath.IsRootObjectPersistentLevel()))
+					{
+						MovingActorsToValidRootObject.Add(Actor);
+					}
+				}
+			});
+
+			// Detach parent actors
+			for (const AActor* Parent : ParentActors)
+			{
+				auto* RootComp = Parent->GetRootComponent();
+
+				// We don't detach if it's a child of another that's been dragged
+				if (RootComp && RootComp->GetAttachParent() && !ChildActors.Contains(Parent))
+				{
+					if (AActor* OldParentActor = RootComp->GetAttachParent()->GetOwner())
+					{
+						// Attachment is persisted on the child so modify both actors for Undo/Redo but do not mark the Parent package dirty
+						OldParentActor->Modify(/*bAlwaysMarkDirty=*/false);
+					}
+					RootComp->DetachFromComponent(FDetachmentTransformRules::KeepWorldTransform);
+				}
+			}
+
+			auto MoveActorsToLevel = [](const TArray<AActor*>& InActorsToMove, ULevel* InDestLevel, const FName& InDestinationPath)
+			{
+				// We are moving actors to another level
+				const bool bWarnAboutReferences = true;
+				const bool bWarnAboutRenaming = true;
+				const bool bMoveAllOrFail = true;
+				TArray<AActor*> MovedActors;
+				if (!EditorLevelUtils::MoveActorsToLevel(InActorsToMove, InDestLevel, bWarnAboutReferences, bWarnAboutRenaming, bMoveAllOrFail, &MovedActors))
+				{
+					UE_LOG(LogActorMode, Warning, TEXT("Failed to move actors because not all actors could be moved"));
+				}
+				// Once moved, update actors folder path
+				for (AActor* Actor : MovedActors)
+				{
+					Actor->SetFolderPath_Recursively(InDestinationPath);
+				}
+			};
+
+			if (DestinationPath.IsRootObjectPersistentLevel())
+			{
+				const ULevelInstanceSubsystem* LevelInstanceSubsystem = RepresentingWorld->GetSubsystem<ULevelInstanceSubsystem>();
+				check(LevelInstanceSubsystem);
+				ULevel* DestinationLevel = RepresentingWorld->PersistentLevel;
+				check(DestinationLevel);
+
+				TArray<AActor*> LevelInstanceActorsToMove;
+				TArray<AActor*> ActorsToMoveToPersistentLevel;
+				Payload.ForEachItem<FActorTreeItem>([LevelInstanceSubsystem, &LevelInstanceActorsToMove, &ActorsToMoveToPersistentLevel](const FActorTreeItem& ActorItem)
+				{
+					AActor* Actor = ActorItem.Actor.Get();
+					if (const ILevelInstanceInterface* ParentLevelInstance = LevelInstanceSubsystem->GetParentLevelInstance(Actor))
+					{
+						check(ParentLevelInstance->IsEditing());
+						LevelInstanceActorsToMove.Add(Actor);
+					}
+					else
+					{
+						const FFolder ActorSrcFolder = Actor->GetFolder();
+						if (ActorSrcFolder.IsRootObjectValid() && !ActorSrcFolder.IsRootObjectPersistentLevel())
+						{
+							ActorsToMoveToPersistentLevel.Add(Actor);
+						}
+					}
+				});
+
+				// We are moving actors outside of an editing level instance to a folder (or root) into the persistent level.
+				if (LevelInstanceActorsToMove.Num() > 0)
+				{
+					TArray<AActor*> MovedActors;
+					LevelInstanceSubsystem->MoveActorsToLevel(LevelInstanceActorsToMove, DestinationLevel, &MovedActors);
+					// Once moved, update actors folder path
+					for (AActor* Actor : MovedActors)
+					{
+						Actor->SetFolderPath_Recursively(DestinationPath.GetPath());
+					}
+				}
+				if (ActorsToMoveToPersistentLevel.Num() > 0)
+				{
+					MoveActorsToLevel(ActorsToMoveToPersistentLevel, DestinationLevel, DestinationPath.GetPath());
+				}
+			}
+			else if (MovingActorsToValidRootObject.Num())
+			{
+				if (ILevelInstanceInterface* TargetLevelInstance = Cast<ILevelInstanceInterface>(DestinationPath.GetRootObjectPtr()))
+				{
+					// We are moving actors inside an editing level instance
+					check(TargetLevelInstance->IsEditing());
+
+					ULevelInstanceSubsystem* LevelInstanceSubsystem = RepresentingWorld->GetSubsystem<ULevelInstanceSubsystem>();
+					check(LevelInstanceSubsystem);
+					TArray<AActor*> MovedActors;
+					LevelInstanceSubsystem->MoveActorsTo(TargetLevelInstance, MovingActorsToValidRootObject, &MovedActors);
+					// Once moved, update actors folder path
+					for (AActor* Actor : MovedActors)
+					{
+						Actor->SetFolderPath_Recursively(DestinationPath.GetPath());
+					}
+				}
+				else if (ULevel* DestinationLevel = Cast<ULevel>(DestinationPath.GetRootObjectPtr()))
+				{
+					MoveActorsToLevel(MovingActorsToValidRootObject, DestinationLevel, DestinationPath.GetPath());
+				}
+			}
+		}
+	}
+}
+
+bool FActorMode::GetFolderNamesFromPayload(const FSceneOutlinerDragDropPayload& InPayload, TArray<FName>& OutFolders, FFolder::FRootObject& OutCommonRootObject) const
+{
+	return FFolder::GetFolderPathsAndCommonRootObject(InPayload.GetData<FFolder>(SceneOutliner::FFolderPathSelector()), OutFolders, OutCommonRootObject);
+}
+
+FFolder FActorMode::GetWorldDefaultRootFolder() const
+{
+	return FFolder::GetWorldRootFolder(RepresentingWorld.Get());
+}
+
+
 #undef LOCTEXT_NAMESPACE
\ No newline at end of file
diff --git a/Engine/Source/Editor/SceneOutliner/Public/ActorBrowsingMode.h b/Engine/Source/Editor/SceneOutliner/Public/ActorBrowsingMode.h
--- a/Engine/Source/Editor/SceneOutliner/Public/ActorBrowsingMode.h
+++ b/Engine/Source/Editor/SceneOutliner/Public/ActorBrowsingMode.h
@@ -17,9 +17,6 @@ public:
 	SCENEOUTLINER_API virtual void Rebuild() override;
 	SCENEOUTLINER_API virtual FCreateSceneOutlinerMode CreateFolderPickerMode(const FFolder::FRootObject& InRootObject = FFolder::GetInvalidRootObject()) const override;
 	SCENEOUTLINER_API virtual void InitializeViewMenuExtender(TSharedPtr<FExtender> Extender) override;
-	SCENEOUTLINER_API virtual TSharedPtr<FDragDropOperation> CreateDragDropOperation(const FPointerEvent& MouseEvent, const TArray<FSceneOutlinerTreeItemPtr>& InTreeItems) const override;
-	SCENEOUTLINER_API virtual bool ParseDragDrop(FSceneOutlinerDragDropPayload& OutPayload, const FDragDropOperation& Operation) const override;
-	SCENEOUTLINER_API virtual FSceneOutlinerDragValidationInfo ValidateDrop(const ISceneOutlinerTreeItem& DropTarget, const FSceneOutlinerDragDropPayload& Payload) const override;
 	SCENEOUTLINER_API virtual TSharedPtr<SWidget> CreateContextMenu() override;
 	SCENEOUTLINER_API virtual void OnItemAdded(FSceneOutlinerTreeItemPtr Item) override;
 	SCENEOUTLINER_API virtual void OnItemRemoved(FSceneOutlinerTreeItemPtr Item) override;
@@ -29,7 +26,6 @@ public:
 	SCENEOUTLINER_API virtual void OnItemPassesFilters(const ISceneOutlinerTreeItem& Item) override;
 	SCENEOUTLINER_API virtual FReply OnKeyDown(const FKeyEvent& InKeyEvent) override;
 	SCENEOUTLINER_API virtual void OnDuplicateSelected() override;
-	SCENEOUTLINER_API virtual void OnDrop(ISceneOutlinerTreeItem& DropTarget, const FSceneOutlinerDragDropPayload& Payload, const FSceneOutlinerDragValidationInfo& ValidationInfo) const override;
 	SCENEOUTLINER_API virtual bool CanRenameItem(const ISceneOutlinerTreeItem& Item) const override;
 	SCENEOUTLINER_API virtual FText GetStatusText() const override;
 	SCENEOUTLINER_API virtual FSlateColor GetStatusTextColor() const override;
@@ -114,10 +110,7 @@ private:
 	static void RegisterContextMenu();
 	static void FillDefaultContextBaseMenu(UToolMenu* InMenu);
 	bool CanPasteFoldersOnlyFromClipboard() const;
-
-	bool GetFolderNamesFromPayload(const FSceneOutlinerDragDropPayload& InPayload, TArray<FName>& OutFolders, FFolder::FRootObject& OutCommonRootObject) const;
-	FFolder GetWorldDefaultRootFolder() const;
-
+	
 	void SynchronizeSelectedActorDescs();
 
 	void OnActorEditorContextSubsystemChanged();
diff --git a/Engine/Source/Editor/SceneOutliner/Public/ActorMode.h b/Engine/Source/Editor/SceneOutliner/Public/ActorMode.h
--- a/Engine/Source/Editor/SceneOutliner/Public/ActorMode.h
+++ b/Engine/Source/Editor/SceneOutliner/Public/ActorMode.h
@@ -7,7 +7,7 @@
 namespace SceneOutliner
 {
 	/** Functor which can be used to get weak actor pointers from a selection */
-	struct FWeakActorSelector
+	struct SCENEOUTLINER_API FWeakActorSelector
 	{
 		bool operator()(const TWeakPtr<ISceneOutlinerTreeItem>& Item, TWeakObjectPtr<AActor>& DataOut) const;
 	};
@@ -19,7 +19,7 @@ namespace SceneOutliner
 	};
 
 	/** Functor which can be used to get actor descriptors from a selection  */
-	struct FActorDescSelector
+	struct SCENEOUTLINER_API FActorDescSelector
 	{
 		bool operator()(const TWeakPtr<ISceneOutlinerTreeItem>& Item, FWorldPartitionActorDesc*& ActorDescPtrOut) const;
 	};
@@ -77,6 +77,12 @@ public:
 	virtual FFolder::FRootObject GetPasteTargetRootObject() const override;
 
 	virtual bool CanInteract(const ISceneOutlinerTreeItem& Item) const override;
+	
+	virtual TSharedPtr<FDragDropOperation> CreateDragDropOperation(const FPointerEvent& MouseEvent, const TArray<FSceneOutlinerTreeItemPtr>& InTreeItems) const override;
+	virtual bool ParseDragDrop(FSceneOutlinerDragDropPayload& OutPayload, const FDragDropOperation& Operation) const override;
+	virtual FSceneOutlinerDragValidationInfo ValidateDrop(const ISceneOutlinerTreeItem& DropTarget, const FSceneOutlinerDragDropPayload& Payload) const override;
+	virtual void OnDrop(ISceneOutlinerTreeItem& DropTarget, const FSceneOutlinerDragDropPayload& Payload, const FSceneOutlinerDragValidationInfo& ValidationInfo) const override;
+
 
 private:
 	/** Called when the user selects a world in the world picker menu */
@@ -86,6 +92,8 @@ private:
 
 	void ChooseRepresentingWorld();
 	bool IsWorldChecked(TWeakObjectPtr<UWorld> World) const;
+	bool GetFolderNamesFromPayload(const FSceneOutlinerDragDropPayload& InPayload, TArray<FName>& OutFolders, FFolder::FRootObject& OutCommonRootObject) const;
+
 protected:
 	void SynchronizeActorSelection();
 	bool IsActorDisplayable(const AActor* InActor) const;
@@ -94,6 +102,8 @@ protected:
 	void SetAsMostRecentOutliner() const;
 
 	virtual TUniquePtr<ISceneOutlinerHierarchy> CreateHierarchy() override;
+	
+	FFolder GetWorldDefaultRootFolder() const;
 protected:
 	/** The world which we are currently representing */
 	TWeakObjectPtr<UWorld> RepresentingWorld;
diff --git a/Engine/Source/Runtime/TypedElementFramework/Public/Elements/Columns/TypedElementSelectionColumns.h b/Engine/Source/Runtime/TypedElementFramework/Public/Elements/Columns/TypedElementSelectionColumns.h
--- a/Engine/Source/Runtime/TypedElementFramework/Public/Elements/Columns/TypedElementSelectionColumns.h
+++ b/Engine/Source/Runtime/TypedElementFramework/Public/Elements/Columns/TypedElementSelectionColumns.h
@@ -14,4 +14,8 @@ USTRUCT(meta = (DisplayName = "Selected"))
 struct FTypedElementSelectionColumn final : public FTypedElementDataStorageColumn
 {
 	GENERATED_BODY()
+
+	// The selection set this row belongs to, empty = default selection set
+	UPROPERTY()
+	FName SelectionSet;
 };
