Outliner: Allow the typed element registry to create the level editor's outliner if it is available and has registered a custom outliner factory

diff --git a/Engine/Plugins/Experimental/TypedElementsDataStorage/Source/TypedElementsDataStorageUI/Private/TypedElementOutliner/TypedElementOutlinerColumnIntegration.cpp b/Engine/Plugins/Experimental/TypedElementsDataStorage/Source/TypedElementsDataStorageUI/Private/TypedElementOutliner/TypedElementOutlinerColumnIntegration.cpp
--- a/Engine/Plugins/Experimental/TypedElementsDataStorage/Source/TypedElementsDataStorageUI/Private/TypedElementOutliner/TypedElementOutlinerColumnIntegration.cpp
+++ b/Engine/Plugins/Experimental/TypedElementsDataStorage/Source/TypedElementsDataStorageUI/Private/TypedElementOutliner/TypedElementOutlinerColumnIntegration.cpp
@@ -340,16 +340,29 @@ FTypedElementSceneOutlinerQueryBinder& FTypedElementSceneOutlinerQueryBinder::Ge
 
 void FTypedElementSceneOutlinerQueryBinder::AssignQuery(TypedElementQueryHandle Query, const TSharedPtr<ISceneOutliner>& Outliner)
 {
-	FTypedElementSceneOutliner* QueryMapping = SceneOutliners.Find(Outliner);
+	CleanupStaleOutliners();
+
+	TSharedPtr<FTypedElementSceneOutliner>* QueryMapping = SceneOutliners.Find(Outliner);
 	if (QueryMapping == nullptr)
 	{
-		QueryMapping = &SceneOutliners.Add(Outliner);
-		QueryMapping->Initialize(*Storage, *StorageUi, *StorageCompatibility, Outliner);
+		QueryMapping = &SceneOutliners.Add(Outliner, MakeShared<FTypedElementSceneOutliner>());
+		
+		(*QueryMapping)->Initialize(*Storage, *StorageUi, *StorageCompatibility, Outliner);
 	}
-	QueryMapping->AssignQuery(Query);
+	(*QueryMapping)->AssignQuery(Query);
 }
 
-
+void FTypedElementSceneOutlinerQueryBinder::CleanupStaleOutliners()
+{
+	for (TMap<TWeakPtr<ISceneOutliner>, TSharedPtr<FTypedElementSceneOutliner>>::TIterator It(SceneOutliners); It; ++It)
+	{
+		// Remove any query mappings where the target Outliner doesn't exist anymore
+		if(!It.Key().IsValid())
+		{
+			It.RemoveCurrent();
+		}
+	}
+}
 
 //
 // FTypedElementSceneOutliner
@@ -469,12 +482,20 @@ void FTypedElementSceneOutliner::AssignQuery(TypedElementQueryHandle Query)
 				{
 					TSharedPtr<FTypedElementWidgetConstructor> CellConstructor(Constructor.Release());
 
+					/* If we have a fallback column for this query, remove it, take over it's priority and 
+					 * replace it with the TEDS column. But also allow the TEDS-Outliner column to fallback to it for
+					 * data not in TEDS yet.
+				 	 */
 					FName FallbackColumn = Binder.FindOutlinerColumnFromTEDSColumns(ColumnTypes);
+					const FSceneOutlinerColumnInfo* FallbackColumnInfo = OutlinerPinned->GetSharedData().ColumnMap.Find(FallbackColumn);
+					int32 ColumnPriority = FallbackColumnInfo ? FallbackColumnInfo->PriorityIndex : DefaultPriorityIndex + IndexOffset;
+
+					OutlinerPinned->RemoveColumn(FallbackColumn);
 
 					FName NameId = FindLongestMatchingName(ColumnTypes, IndexOffset);
 					AddedColumns.Add(NameId);
 					OutlinerPinned->AddColumn(NameId,
-						FSceneOutlinerColumnInfo(ESceneOutlinerColumnVisibility::Visible, DefaultPriorityIndex + IndexOffset,
+						FSceneOutlinerColumnInfo(ESceneOutlinerColumnVisibility::Visible, ColumnPriority,
 							FCreateSceneOutlinerColumn::CreateLambda(
 								[this, Query, NameId, &ColumnTypes, CellConstructor, &OutlinerPinned, FallbackColumn](ISceneOutliner&)
 								{
diff --git a/Engine/Plugins/Experimental/TypedElementsDataStorage/Source/TypedElementsDataStorageUI/Public/TypedElementOutliner/TypedElementOutlinerColumnIntegration.h b/Engine/Plugins/Experimental/TypedElementsDataStorage/Source/TypedElementsDataStorageUI/Public/TypedElementOutliner/TypedElementOutlinerColumnIntegration.h
--- a/Engine/Plugins/Experimental/TypedElementsDataStorage/Source/TypedElementsDataStorageUI/Public/TypedElementOutliner/TypedElementOutlinerColumnIntegration.h
+++ b/Engine/Plugins/Experimental/TypedElementsDataStorage/Source/TypedElementsDataStorageUI/Public/TypedElementOutliner/TypedElementOutlinerColumnIntegration.h
@@ -67,8 +67,9 @@ public:
 private:
 	FTypedElementSceneOutlinerQueryBinder();
 	void SetupDefaultColumnMapping();
-
-	TMap<TWeakPtr<ISceneOutliner>, FTypedElementSceneOutliner> SceneOutliners;
+	void CleanupStaleOutliners();
+	
+	TMap<TWeakPtr<ISceneOutliner>, TSharedPtr<FTypedElementSceneOutliner>> SceneOutliners;
 
 	ITypedElementDataStorageInterface* Storage{ nullptr };
 	ITypedElementDataStorageUiInterface* StorageUi{ nullptr };
diff --git a/Engine/Source/Editor/LevelEditor/Private/SLevelEditor.cpp b/Engine/Source/Editor/LevelEditor/Private/SLevelEditor.cpp
--- a/Engine/Source/Editor/LevelEditor/Private/SLevelEditor.cpp
+++ b/Engine/Source/Editor/LevelEditor/Private/SLevelEditor.cpp
@@ -68,6 +68,7 @@
 #include "Elements/Actor/ActorElementLevelEditorCommonActionsCustomization.h"
 #include "Elements/Component/ComponentElementLevelEditorSelectionCustomization.h"
 #include "Elements/Component/ComponentElementLevelEditorCommonActionsCustomization.h"
+#include "Elements/Framework/TypedElementRegistry.h"
 #include "Elements/SMInstance/SMInstanceElementId.h"
 #include "Elements/SMInstance/SMInstanceElementLevelEditorSelectionCustomization.h"
 #include "DerivedDataEditorModule.h"
@@ -838,20 +839,31 @@ TSharedRef<ISceneOutliner> SLevelEditor::CreateSceneOutliner(FName TabIdentifier
 	
 	TSharedPtr<FLevelEditorOutlinerSettings> OutlinerSettings = LevelEditorModule.GetLevelEditorOutlinerSettings();
 	OutlinerSettings->GetOutlinerFilters(InitOptions.FilterBarOptions);
-
 	InitOptions.FilterBarOptions.CategoryToExpand = OutlinerSettings->GetFilterCategory(FLevelEditorOutlinerBuiltInCategories::Common());
 
 	FSceneOutlinerModule& SceneOutlinerModule = FModuleManager::Get().LoadModuleChecked<FSceneOutlinerModule>("SceneOutliner");
-	TSharedRef<ISceneOutliner> SceneOutlinerRef = SceneOutlinerModule.CreateActorBrowser(
-		InitOptions);
+	TSharedPtr<ISceneOutliner> NewSceneOutlinerPtr;
 
+	// Check if the Typed Element Registry has registered a custom outliner factory for us to use
+	if(SceneOutlinerModule.IsCustomSceneOutlinerFactoryRegistered(UTypedElementRegistry::GetInstance()->GetFName()))
+	{
+		NewSceneOutlinerPtr = SceneOutlinerModule.CreateCustomRegisteredOutliner(UTypedElementRegistry::GetInstance()->GetFName(), InitOptions);
+	}
+	// Fallback to the regular Actor Browser otherwise
+	else
+	{
+		NewSceneOutlinerPtr = SceneOutlinerModule.CreateActorBrowser(
+			InitOptions);
+
+	}
+	
 	// Add this to the map of all outliners
-	SceneOutliners.Add(TabIdentifier, SceneOutlinerRef);
+	SceneOutliners.Add(TabIdentifier, NewSceneOutlinerPtr);
 
 	// Update the most recently used outliner
 	SetMostRecentlyUsedSceneOutliner(TabIdentifier);
 
-	return SceneOutlinerRef;
+	return NewSceneOutlinerPtr.ToSharedRef();
 }
 
 void SLevelEditor::OnExtendSceneOutlinerTabContextMenu(FMenuBuilder& InMenuBuilder)
diff --git a/Engine/Source/Editor/SceneOutliner/Private/SceneOutlinerModule.cpp b/Engine/Source/Editor/SceneOutliner/Private/SceneOutlinerModule.cpp
--- a/Engine/Source/Editor/SceneOutliner/Private/SceneOutlinerModule.cpp
+++ b/Engine/Source/Editor/SceneOutliner/Private/SceneOutlinerModule.cpp
@@ -168,47 +168,90 @@ TSharedRef< ISceneOutliner > FSceneOutlinerModule::CreateActorBrowser(const FSce
 
 	if (InitOptions.ColumnMap.Num() == 0)
 	{
-		InitOptions.UseDefaultColumns();
-		InitOptions.ColumnMap.Add(FSceneOutlinerBuiltInColumnTypes::Gutter(), FSceneOutlinerColumnInfo(ESceneOutlinerColumnVisibility::Visible, 0, FCreateSceneOutlinerColumn(), true, TOptional<float>(), FSceneOutlinerBuiltInColumnTypes::Gutter_Localized()));
-		InitOptions.ColumnMap.Add(FSceneOutlinerBuiltInColumnTypes::ActorInfo(), FSceneOutlinerColumnInfo(ESceneOutlinerColumnVisibility::Visible, 20, FCreateSceneOutlinerColumn(), true, TOptional<float>(), FSceneOutlinerBuiltInColumnTypes::ActorInfo_Localized()));
+		CreateActorBrowserColumns(InitOptions);
+	}
 
-		ESceneOutlinerColumnVisibility SourceControlColumnVisibility = ESceneOutlinerColumnVisibility::Invisible;
-		
-		UWorld* WorldPtr = nullptr;
+	return CreateSceneOutliner(InitOptions);
+}
+
+TSharedPtr<ISceneOutliner> FSceneOutlinerModule::CreateCustomRegisteredOutliner(FName ID, FSceneOutlinerInitializationOptions InInitOptions)
+{
+	FSceneOutlinerFactory* FoundInitOptionsFactory = CustomOutlinerFactories.Find(ID);
+
+	if(!FoundInitOptionsFactory)
+	{
+		return nullptr;
+	}
+
+	return FoundInitOptionsFactory->Execute(InInitOptions);
+}
+
+void FSceneOutlinerModule::RegisterCustomSceneOutlinerFactory(FName ID,
+	FSceneOutlinerFactory InOutlinerFactory)
+{
+	CustomOutlinerFactories.Add(ID, InOutlinerFactory);
+}
+
+void FSceneOutlinerModule::UnregisterCustomSceneOutlinerFactory(FName ID)
+{
+	CustomOutlinerFactories.Remove(ID);
+}
+
+bool FSceneOutlinerModule::IsCustomSceneOutlinerFactoryRegistered(FName ID)
+{
+	return CustomOutlinerFactories.Contains(ID);
+}
+
+void FSceneOutlinerModule::CreateActorBrowserColumns(FSceneOutlinerInitializationOptions& InInitOptions, UWorld* InWorld) const
+{
+	if(!InWorld)
+	{
 		// Query the Level Editor to get the correct world based on context
 		TWeakPtr<ILevelEditor> LevelEditor = FModuleManager::GetModuleChecked<FLevelEditorModule>(TEXT("LevelEditor")).GetLevelEditorInstance();
 		if (TSharedPtr<ILevelEditor> LevelEditorPin = LevelEditor.Pin())
 		{
-			WorldPtr = LevelEditorPin->GetEditorModeManager().GetWorld();
-		}
-		ISourceControlModule& SourceControlModule = ISourceControlModule::Get();
-		if (SourceControlModule.IsEnabled())
-		{
-			if (WorldPtr && WorldPtr->PersistentLevel->IsUsingExternalActors())
-			{
-				// The source control column should be visible by default in source-controlled levels using external actors
-				SourceControlColumnVisibility = ESceneOutlinerColumnVisibility::Visible;
-			}
+			InWorld = LevelEditorPin->GetEditorModeManager().GetWorld();
 		}
-		InitOptions.ColumnMap.Add(FSceneOutlinerBuiltInColumnTypes::SourceControl(), FSceneOutlinerColumnInfo(SourceControlColumnVisibility, 30, FCreateSceneOutlinerColumn(), true, TOptional<float>(), FSceneOutlinerBuiltInColumnTypes::SourceControl_Localized()));
+	}
 
-		ESceneOutlinerColumnVisibility UnsavedColumnVisibility = ESceneOutlinerColumnVisibility::Invisible;
-		
-		if (WorldPtr && WorldPtr->IsPartitionedWorld())
-		{
-			// We don't want the pinned column in non wp levels
-			InitOptions.ColumnMap.Add(FSceneOutlinerBuiltInColumnTypes::Pinned(), FSceneOutlinerColumnInfo(ESceneOutlinerColumnVisibility::Visible, 5, FCreateSceneOutlinerColumn(), true, TOptional<float>(), FSceneOutlinerBuiltInColumnTypes::Pinned_Localized()));
+	InInitOptions.UseDefaultColumns();
+	
+	InInitOptions.UseDefaultColumns();
+	InInitOptions.ColumnMap.Add(FSceneOutlinerBuiltInColumnTypes::Gutter(), FSceneOutlinerColumnInfo(ESceneOutlinerColumnVisibility::Visible, 0, FCreateSceneOutlinerColumn(), true, TOptional<float>(), FSceneOutlinerBuiltInColumnTypes::Gutter_Localized()));
+	InInitOptions.ColumnMap.Add(FSceneOutlinerBuiltInColumnTypes::ActorInfo(), FSceneOutlinerColumnInfo(ESceneOutlinerColumnVisibility::Visible, 20, FCreateSceneOutlinerColumn(), true, TOptional<float>(), FSceneOutlinerBuiltInColumnTypes::ActorInfo_Localized()));
+
+	CreateWorldPartitionColumns(InInitOptions, InWorld);
+	CreateActorInfoColumns(InInitOptions, InWorld);
+}
+
+void FSceneOutlinerModule::CreateWorldPartitionColumns(FSceneOutlinerInitializationOptions& InInitOptions, UWorld* WorldPtr) const
+{
+	ESceneOutlinerColumnVisibility SourceControlColumnVisibility = ESceneOutlinerColumnVisibility::Invisible;
 
-			// We want the unsaved column to be visible by default in partitioned levels
-			UnsavedColumnVisibility = ESceneOutlinerColumnVisibility::Visible;
+	ISourceControlModule& SourceControlModule = ISourceControlModule::Get();
+	if (SourceControlModule.IsEnabled())
+	{
+		if (WorldPtr && WorldPtr->PersistentLevel->IsUsingExternalActors())
+		{
+			// The source control column should be visible by default in source-controlled levels using external actors
+			SourceControlColumnVisibility = ESceneOutlinerColumnVisibility::Visible;
 		}
-		CreateActorInfoColumns(InitOptions, WorldPtr);
+	}
+	InInitOptions.ColumnMap.Add(FSceneOutlinerBuiltInColumnTypes::SourceControl(), FSceneOutlinerColumnInfo(SourceControlColumnVisibility, 30, FCreateSceneOutlinerColumn(), true, TOptional<float>(), FSceneOutlinerBuiltInColumnTypes::SourceControl_Localized()));
 
-		InitOptions.ColumnMap.Add(FSceneOutlinerBuiltInColumnTypes::Unsaved(), FSceneOutlinerColumnInfo(UnsavedColumnVisibility, 1, FCreateSceneOutlinerColumn(), true, TOptional<float>(), FSceneOutlinerBuiltInColumnTypes::Unsaved_Localized()));
+	ESceneOutlinerColumnVisibility UnsavedColumnVisibility = ESceneOutlinerColumnVisibility::Invisible;
 		
+	if (WorldPtr && WorldPtr->IsPartitionedWorld())
+	{
+		// We don't want the pinned column in non wp levels
+		InInitOptions.ColumnMap.Add(FSceneOutlinerBuiltInColumnTypes::Pinned(), FSceneOutlinerColumnInfo(ESceneOutlinerColumnVisibility::Visible, 5, FCreateSceneOutlinerColumn(), true, TOptional<float>(), FSceneOutlinerBuiltInColumnTypes::Pinned_Localized()));
+
+		// We want the unsaved column to be visible by default in partitioned levels
+		UnsavedColumnVisibility = ESceneOutlinerColumnVisibility::Visible;
 	}
+	
+	InInitOptions.ColumnMap.Add(FSceneOutlinerBuiltInColumnTypes::Unsaved(), FSceneOutlinerColumnInfo(UnsavedColumnVisibility, 1, FCreateSceneOutlinerColumn(), true, TOptional<float>(), FSceneOutlinerBuiltInColumnTypes::Unsaved_Localized()));
 
-	return CreateSceneOutliner(InitOptions);
 }
 
 void FSceneOutlinerModule::CreateActorInfoColumns(FSceneOutlinerInitializationOptions& InInitOptions, UWorld *WorldPtr) const
diff --git a/Engine/Source/Editor/SceneOutliner/Public/SceneOutlinerFwd.h b/Engine/Source/Editor/SceneOutliner/Public/SceneOutlinerFwd.h
--- a/Engine/Source/Editor/SceneOutliner/Public/SceneOutlinerFwd.h
+++ b/Engine/Source/Editor/SceneOutliner/Public/SceneOutlinerFwd.h
@@ -45,4 +45,8 @@ DECLARE_DELEGATE_RetVal_OneParam( TSharedRef<ISceneOutlinerColumn>, FCreateScene
 /** A delegate used to factory a new filter type */
 DECLARE_DELEGATE_RetVal( TSharedRef<FSceneOutlinerFilter>, FCreateSceneOutlinerFilter );
 
+/** A delegate used to factory a new scene outliner using the given init options */
+DECLARE_DELEGATE_RetVal_OneParam( TSharedRef<ISceneOutliner>, FSceneOutlinerFactory, FSceneOutlinerInitializationOptions);
+
+
 class FSceneOutlinerTreeItemSCC;
diff --git a/Engine/Source/Editor/SceneOutliner/Public/SceneOutlinerModule.h b/Engine/Source/Editor/SceneOutliner/Public/SceneOutlinerModule.h
--- a/Engine/Source/Editor/SceneOutliner/Public/SceneOutlinerModule.h
+++ b/Engine/Source/Editor/SceneOutliner/Public/SceneOutlinerModule.h
@@ -59,6 +59,24 @@ public:
 		const FSceneOutlinerInitializationOptions& InInitOptions,
 		TWeakObjectPtr<UWorld> SpecifiedWorld = nullptr) const;
 
+
+	/** Register a factory to create a custom Scene Outliner */
+	virtual void RegisterCustomSceneOutlinerFactory(FName ID, FSceneOutlinerFactory InOutlinerFactory);
+
+	/** Unregister a factory to create a custom Scene Outliner */
+	virtual void UnregisterCustomSceneOutlinerFactory(FName ID);
+
+	/** Try to create a custom scene outliner using a registered factory (nullptr if ID was not registered) */
+	virtual TSharedPtr<ISceneOutliner> CreateCustomRegisteredOutliner(FName ID, FSceneOutlinerInitializationOptions InInitOptions);
+
+	/** Check if a custom scene outliner factory is registered with the given ID */
+	virtual bool IsCustomSceneOutlinerFactoryRegistered(FName ID);
+
+	/** Add the columns present in the level editor's Outliner (Actor Browser) to the given init options
+	 *  @param InWorld The world the Outliner initialized by InInitOptions will look it, defaults to the level editor's world if nullptr
+	 */
+	virtual void CreateActorBrowserColumns(FSceneOutlinerInitializationOptions& InInitOptions, UWorld* InWorld = nullptr) const;
+	
 	/** Column permission list */
 	TSharedRef<FNamePermissionList>& GetColumnPermissionList() { return ColumnPermissionList; }
 
@@ -117,6 +135,9 @@ public:
 		return nullptr;
 	}
 
+	void CreateActorInfoColumns(FSceneOutlinerInitializationOptions& InInitOptions, UWorld* WorldPtr = nullptr) const;
+	void CreateWorldPartitionColumns(FSceneOutlinerInitializationOptions& InInitOptions, UWorld* WorldPtr = nullptr) const;
+	
 	/** Map of column type name -> default column info */
 	TMap< FName, FSceneOutlinerColumnInfo> DefaultColumnMap;
 
@@ -130,9 +151,8 @@ private:
 
 	/** Delegate that broadcasts when column permission list changes. */
 	FOnColumnPermissionListChanged ColumnPermissionListChanged;
-
-	// Function to recreate what originally was FActorInfoColumn as separate columns using FTextInfoColumn
-	void CreateActorInfoColumns(FSceneOutlinerInitializationOptions& InInitOptions, UWorld* WorldPtr = nullptr) const;
+	
+	TMap< FName, FSceneOutlinerFactory> CustomOutlinerFactories;
 
 public:
 
