Editor: Fix outliner actor item class not being properly inheritable because it doesn't have a way to pass down the type identifier from derived classes.

diff --git a/Engine/Source/Editor/SceneOutliner/Private/ActorTreeItem.cpp b/Engine/Source/Editor/SceneOutliner/Private/ActorTreeItem.cpp
--- a/Engine/Source/Editor/SceneOutliner/Private/ActorTreeItem.cpp
+++ b/Engine/Source/Editor/SceneOutliner/Private/ActorTreeItem.cpp
@@ -300,7 +300,11 @@ private:
 };
 
 FActorTreeItem::FActorTreeItem(AActor* InActor)
-	: IActorBaseTreeItem(Type)
+	// Forward to the other constructor using our type identifier.
+	: FActorTreeItem(Type, InActor) {}
+
+FActorTreeItem::FActorTreeItem(FSceneOutlinerTreeItemType TypeIn, AActor* InActor)
+	: IActorBaseTreeItem(TypeIn)
 	, Actor(InActor)
 	, ID(InActor)
 {
diff --git a/Engine/Source/Editor/SceneOutliner/Public/ActorTreeItem.h b/Engine/Source/Editor/SceneOutliner/Public/ActorTreeItem.h
--- a/Engine/Source/Editor/SceneOutliner/Public/ActorTreeItem.h
+++ b/Engine/Source/Editor/SceneOutliner/Public/ActorTreeItem.h
@@ -62,6 +62,9 @@ public:
 	bool bExistsInCurrentWorldAndPIE;
 
 protected:
+	/** For use by derived classes, as it allows for passing down their own type. */
+	FActorTreeItem(FSceneOutlinerTreeItemType TypeIn, AActor* InActor);
+
 	virtual void UpdateDisplayString();
 
 	/** Cached display string */
