Scene Outliner: Check for Scene Outliner in OnDragLeave

#rnx

diff --git a/Engine/Source/Editor/SceneOutliner/Private/SOutlinerTreeView.cpp b/Engine/Source/Editor/SceneOutliner/Private/SOutlinerTreeView.cpp
--- a/Engine/Source/Editor/SceneOutliner/Private/SOutlinerTreeView.cpp
+++ b/Engine/Source/Editor/SceneOutliner/Private/SOutlinerTreeView.cpp
@@ -140,6 +140,11 @@ FReply SSceneOutlinerTreeView::OnDragOver(const FGeometry& MyGeometry, const FDr
 
 void SSceneOutlinerTreeView::OnDragLeave(const FDragDropEvent& DragDropEvent)
 {
+	if (!SceneOutlinerWeak.IsValid())
+	{
+		return;
+	}
+
 	if( SceneOutlinerWeak.Pin()->GetSharedData().bShowParentTree )
 	{
 		ResetOperationDecorator(DragDropEvent);
