FActorHierarchy: Move the CreateComponentItems to protected and export it

diff --git a/Engine/Source/Editor/SceneOutliner/Public/ActorHierarchy.h b/Engine/Source/Editor/SceneOutliner/Public/ActorHierarchy.h
--- a/Engine/Source/Editor/SceneOutliner/Public/ActorHierarchy.h
+++ b/Engine/Source/Editor/SceneOutliner/Public/ActorHierarchy.h
@@ -34,13 +34,14 @@ private:
 	bool IsShowingUnloadedActors() const;
 	/** Adds all the direct and indirect children of a world to OutItems */
 	void CreateWorldChildren(UWorld* World, TArray<FSceneOutlinerTreeItemPtr>& OutItems) const;
-	/** Create all component items for an actor if we are showing components and place them in OutItems */
-	void CreateComponentItems(const AActor* Actor, TArray<FSceneOutlinerTreeItemPtr>& OutItems) const;
 
 protected:
 	// Update the hierarchy when actor or world changing events occur
 	void OnWorldPartitionCreated(UWorld* InWorld);
 
+	/** Create all component items for an actor if we are showing components and place them in OutItems */
+	SCENEOUTLINER_API void CreateComponentItems(const AActor* Actor, TArray<FSceneOutlinerTreeItemPtr>& OutItems) const;
+	
 	SCENEOUTLINER_API virtual void OnLevelActorAdded(AActor* InActor);
 	SCENEOUTLINER_API virtual void OnLevelActorDeleted(AActor* InActor);
 		
