Actor: make IsActorLabelEditable virtual
ActorTreeItem: prevent editing label if IsActorLabelEditable returns false

#rnx

diff --git a/Engine/Source/Editor/SceneOutliner/Private/ActorTreeItem.cpp b/Engine/Source/Editor/SceneOutliner/Private/ActorTreeItem.cpp
--- a/Engine/Source/Editor/SceneOutliner/Private/ActorTreeItem.cpp
+++ b/Engine/Source/Editor/SceneOutliner/Private/ActorTreeItem.cpp
@@ -59,10 +59,7 @@ struct SActorTreeLabel : FSceneOutlinerCommonLabelData, public SCompoundWidget
 				.OnEnterEditingMode(this, &SActorTreeLabel::OnEnterEditingMode)
 				.OnExitEditingMode(this, &SActorTreeLabel::OnExitEditingMode)
 				.IsSelected(FIsSelected::CreateSP(&InRow, &STableRow<FSceneOutlinerTreeItemPtr>::IsSelectedExclusively))
-				.IsReadOnly_Lambda([Item = ActorItem.AsShared(), this]()
-				{
-					return !CanExecuteRenameRequest(Item.Get());
-				})
+				.IsReadOnly(this, &SActorTreeLabel::IsReadOnly)
 			];
 
 		if (WeakSceneOutliner.Pin()->GetMode()->IsInteractive())
@@ -296,6 +293,12 @@ private:
 		bInEditingMode = false;
 	}
 
+	bool IsReadOnly() const
+	{
+		AActor* Actor = ActorPtr.Get();
+		return !(Actor && Actor->IsActorLabelEditable() && CanExecuteRenameRequest(*TreeItemPtr.Pin()));
+	}
+
 	bool bInEditingMode = false;
 };
 
diff --git a/Engine/Source/Runtime/Engine/Classes/GameFramework/Actor.h b/Engine/Source/Runtime/Engine/Classes/GameFramework/Actor.h
--- a/Engine/Source/Runtime/Engine/Classes/GameFramework/Actor.h
+++ b/Engine/Source/Runtime/Engine/Classes/GameFramework/Actor.h
@@ -2507,7 +2507,7 @@ public:
 	 * Returns if this actor's current label is editable.  Actor labels are only available in development builds.
 	 * @return	The editable status of the actor's label
 	 */
-	ENGINE_API bool IsActorLabelEditable() const;
+	ENGINE_API virtual bool IsActorLabelEditable() const;
 
 	/** Returns this actor's folder path. Actor folder paths are only available in development builds. */
 	UFUNCTION(BlueprintCallable, Category = "Editor Scripting | Actor Editing")
