Outliner: Fix infinite recursion crash with outliner, GetAttachedActors handles recursion checks but may still return unexpected actors such as the parent so this code may need additional changes

diff --git a/Engine/Source/Editor/SceneOutliner/Private/ActorHierarchy.cpp b/Engine/Source/Editor/SceneOutliner/Private/ActorHierarchy.cpp
--- a/Engine/Source/Editor/SceneOutliner/Private/ActorHierarchy.cpp
+++ b/Engine/Source/Editor/SceneOutliner/Private/ActorHierarchy.cpp
@@ -602,17 +602,9 @@ void FActorHierarchy::CreateChildren(const FSceneOutlinerTreeItemPtr& Item, TArr
 		}
 		else
 		{
-			TFunction<bool(AActor*)> GetAttachedActors = [&ChildActors, &GetAttachedActors](AActor* Child)
-			{
-				ChildActors.Add(Child);
-				Child->ForEachAttachedActors(GetAttachedActors);
-
-				// Always continue
-				return true;
-			};
-
-			// Grab all direct/indirect children of an actor
-			ParentActor->ForEachAttachedActors(GetAttachedActors);
+			// NOTE: GetAttachedActors can sometimes return the actor itself or its parent depending on how the component ownership is setup
+			// but we don't have to check for that here because the Outliner will simply discard duplicates
+			ParentActor->GetAttachedActors(ChildActors, true, true);
 		}
 
 		for (auto ChildActor : ChildActors)
