Outliner: Keep folders selected when changing / clearing search text

diff --git a/Engine/Source/Editor/SceneOutliner/Private/SSceneOutliner.cpp b/Engine/Source/Editor/SceneOutliner/Private/SSceneOutliner.cpp
--- a/Engine/Source/Editor/SceneOutliner/Private/SSceneOutliner.cpp
+++ b/Engine/Source/Editor/SceneOutliner/Private/SSceneOutliner.cpp
@@ -693,6 +693,16 @@ void SSceneOutliner::Populate()
 	bool bMadeAnySignificantChanges = false;
 	if (bFullRefresh)
 	{
+		// Remember the selected folders
+		TArray<TSharedPtr<ISceneOutlinerTreeItem>> SelectedItems = OutlinerTreeView->GetSelectedItems();
+		for (const TSharedPtr<ISceneOutlinerTreeItem>& SelectedItem : SelectedItems)
+		{
+			if (const FFolderTreeItem* FolderItem = SelectedItem->CastTo<FFolderTreeItem>())
+			{
+				PendingFoldersSelect.Add(FolderItem->GetFolder());
+			}
+		}
+
 		// Clear the selection here - RepopulateEntireTree will reconstruct it.
 		OutlinerTreeView->ClearSelection();
 
