- Revert UUnrealEdEngine::SelectActor allowing SubRootSelection
- Instead convert FActorBrowsingMode::OnItemSelectionChanged to use TypedElement selection and allow SubRootSelection there

#rnx

diff --git a/Engine/Source/Editor/SceneOutliner/Private/ActorBrowsingMode.cpp b/Engine/Source/Editor/SceneOutliner/Private/ActorBrowsingMode.cpp
--- a/Engine/Source/Editor/SceneOutliner/Private/ActorBrowsingMode.cpp
+++ b/Engine/Source/Editor/SceneOutliner/Private/ActorBrowsingMode.cpp
@@ -49,6 +49,8 @@
 #include "EditorViewportCommands.h"
 #include "SceneOutlinerActorSCCColumn.h"
 #include "Misc/ScopedSlowTask.h"
+#include "Elements/Framework/EngineElementsLibrary.h"
+#include "Elements/Framework/TypedElementHandle.h"
 
 DEFINE_LOG_CATEGORY_STATIC(LogActorBrowser, Log, All);
 
@@ -248,6 +250,11 @@ FActorBrowsingMode::FActorBrowsingMode(SSceneOutliner* InSceneOutliner, TWeakObj
 	// Add a filter which sets the interactive mode of LevelInstance items and their children
 	SceneOutliner->AddInteractiveFilter(MakeShared<FActorFilter>(FActorTreeItem::FFilterPredicate::CreateLambda([this](const AActor* Actor)
 		{
+			if (Actor->SupportsSubRootSelection())
+			{
+				return true;
+			}
+
 			if (!bHideLevelInstanceHierarchy)
 			{
 				if (const ULevelInstanceSubsystem* LevelInstanceSubsystem = RepresentingWorld->GetSubsystem<ULevelInstanceSubsystem>())
@@ -995,85 +1002,98 @@ void FActorBrowsingMode::OnActorDescRemoved(FWorldPartitionActorDesc* InActorDes
 
 void FActorBrowsingMode::OnItemSelectionChanged(FSceneOutlinerTreeItemPtr TreeItem, ESelectInfo::Type SelectionType, const FSceneOutlinerItemSelection& Selection)
 {
-	TArray<AActor*> SelectedActors = Selection.GetData<AActor*>(SceneOutliner::FActorSelector());
+	TSet<AActor*> OutlinerSelectedActors(Selection.GetData<AActor*>(SceneOutliner::FActorSelector()));
+	TSet<AActor*> ActorsToSelect;
+	ActorsToSelect.Reserve(OutlinerSelectedActors.Num());
 
 	SynchronizeSelectedActorDescs();
 
-	bool bChanged = false;
-	bool bAnyInPIE = false;
-	for (auto* Actor : SelectedActors)
+	USelection* ActorSelection = GEditor->GetSelectedActors();
+	if (UTypedElementSelectionSet* SelectionSet = ActorSelection->GetElementSelectionSet())
 	{
-		if (!bAnyInPIE && Actor && Actor->GetPackage()->HasAnyPackageFlags(PKG_PlayInEditor))
-		{
-			bAnyInPIE = true;
-		}
-		if (!GEditor->GetSelectedActors()->IsSelected(Actor))
-		{
-			bChanged = true;
-			break;
-		}
-	}
+		TSet<AActor*> EditorSelectedActors(SelectionSet->GetSelectedObjects<AActor>());
 
-	for (FSelectionIterator SelectionIt(*GEditor->GetSelectedActors()); SelectionIt && !bChanged; ++SelectionIt)
-	{
-		const AActor* Actor = CastChecked< AActor >(*SelectionIt);
-		if (!bAnyInPIE && Actor->GetPackage()->HasAnyPackageFlags(PKG_PlayInEditor))
+		bool bChanged = false;
+		bool bAnyInPIE = false;
+		for (AActor* Actor : OutlinerSelectedActors)
 		{
-			bAnyInPIE = true;
+			if (!bAnyInPIE && Actor && Actor->GetPackage()->HasAnyPackageFlags(PKG_PlayInEditor))
+			{
+				bAnyInPIE = true;
+			}
+
+			if (!EditorSelectedActors.Contains(Actor))
+			{
+				bChanged = true;
+			}
+
+			// Allow selection of Sub Roots only if Roots aren't in the outliners selection
+			AActor* RootSelectionParent = Actor ? Actor->GetRootSelectionParent() : nullptr;
+			if (!RootSelectionParent || !OutlinerSelectedActors.Contains(RootSelectionParent))
+			{
+				ActorsToSelect.Add(Actor);
+			}
 		}
-		if (!SelectedActors.Contains(Actor))
+
+		for (const AActor* SelectedActor : EditorSelectedActors)
 		{
-			// Actor has been deselected
-			bChanged = true;
+			if (!bAnyInPIE && SelectedActor && SelectedActor->GetPackage()->HasAnyPackageFlags(PKG_PlayInEditor))
+			{
+				bAnyInPIE = true;
+			}
 
-			// If actor was a group actor, remove its members from the ActorsToSelect list
-			const AGroupActor* DeselectedGroupActor = Cast<AGroupActor>(Actor);
-			if (DeselectedGroupActor)
+			if (!ActorsToSelect.Contains(SelectedActor))
 			{
-				TArray<AActor*> GroupActors;
-				DeselectedGroupActor->GetGroupActors(GroupActors);
+				// Actor has been deselected
+				bChanged = true;
 
-				for (auto* GroupActor : GroupActors)
+				// If actor was a group actor, remove its members from the ActorsToSelect list
+				if (const AGroupActor* DeselectedGroupActor = Cast<AGroupActor>(SelectedActor))
 				{
-					SelectedActors.Remove(GroupActor);
-				}
+					TArray<AActor*> GroupActors;
+					DeselectedGroupActor->GetGroupActors(GroupActors);
 
+					for (AActor* GroupActor : GroupActors)
+					{
+						ActorsToSelect.Remove(GroupActor);
+					}
+				}
 			}
 		}
-	}
-
-	// If there's a discrepancy, update the selected actors to reflect this list.
-	if (bChanged)
-	{
-		const FScopedTransaction Transaction(NSLOCTEXT("UnrealEd", "ClickingOnActors", "Clicking on Actors"), !bAnyInPIE);
-		GEditor->GetSelectedActors()->Modify();
 
-		// We'll batch selection changes instead by using BeginBatchSelectOperation()
-		GEditor->GetSelectedActors()->BeginBatchSelectOperation();
+		// If there's a discrepancy, update the selected actors to reflect this list.
+		if (bChanged)
+		{
+			const FScopedTransaction Transaction(NSLOCTEXT("UnrealEd", "ClickingOnActors", "Clicking on Actors"), !bAnyInPIE);
 
-		// Clear the selection.
-		GEditor->SelectNone(false, true, true);
+			TArray<FTypedElementHandle> ElementsToSelect;
+			for (auto* Actor : ActorsToSelect)
+			{
+				UE_LOG(LogActorBrowser, Verbose, TEXT("Clicking on Actor (world outliner): %s (%s)"), *Actor->GetClass()->GetName(), *Actor->GetActorLabel());
+				ElementsToSelect.Add(UEngineElementsLibrary::AcquireEditorActorElementHandle(Actor));
+			}
 
-		const bool bShouldSelect = true;
-		const bool bNotifyAfterSelect = false;
-		const bool bSelectEvenIfHidden = true;	// @todo outliner: Is this actually OK?
-		for (auto* Actor : SelectedActors)
-		{
-			UE_LOG(LogActorBrowser, Verbose, TEXT("Clicking on Actor (world outliner): %s (%s)"), *Actor->GetClass()->GetName(), *Actor->GetActorLabel());
-			GEditor->SelectActor(Actor, bShouldSelect, bNotifyAfterSelect, bSelectEvenIfHidden);
-		}
+			{
+				const FTypedElementSelectionOptions SelectionOptions = FTypedElementSelectionOptions()
+					.SetAllowHidden(true)
+					.SetWarnIfLocked(false)
+					.SetAllowLegacyNotifications(false)
+					.SetAllowSubRootSelection(true);
+
+				// Avoid senting out notification via typed element and call NoteSelectionChange to preserve previous behavior
+				FTypedElementList::FScopedClearNewPendingChange ClearNewPendingChange = SelectionSet->GetScopedClearNewPendingChange();
+				SelectionSet->SetSelection(ElementsToSelect, SelectionOptions);
+			}
 
-		// Commit selection changes
-		GEditor->GetSelectedActors()->EndBatchSelectOperation(/*bNotify*/false);
+			// Fire selection changed event
+			GEditor->NoteSelectionChange();
 
-		// Fire selection changed event
-		GEditor->NoteSelectionChange();
+			// Set this outliner as the most recently interacted with
+			SetAsMostRecentOutliner();
+		}
 
-		// Set this outliner as the most recently interacted with
-		SetAsMostRecentOutliner();
+		SceneOutliner->RefreshSelection();
 	}
-
-	SceneOutliner->RefreshSelection();
 }
 
 void FActorBrowsingMode::OnItemDoubleClick(FSceneOutlinerTreeItemPtr Item)
diff --git a/Engine/Source/Editor/SceneOutliner/SceneOutliner.Build.cs b/Engine/Source/Editor/SceneOutliner/SceneOutliner.Build.cs
--- a/Engine/Source/Editor/SceneOutliner/SceneOutliner.Build.cs
+++ b/Engine/Source/Editor/SceneOutliner/SceneOutliner.Build.cs
@@ -26,7 +26,9 @@ public class SceneOutliner : ModuleRules
 				"EditorWidgets",
 				"ToolWidgets",
 				"UnsavedAssetsTracker",
-				"AssetDefinition"
+				"AssetDefinition",
+				"TypedElementFramework",
+				"TypedElementRuntime",
 			}
 		);
 		
diff --git a/Engine/Source/Editor/UnrealEd/Private/EditorSelectUtils.cpp b/Engine/Source/Editor/UnrealEd/Private/EditorSelectUtils.cpp
--- a/Engine/Source/Editor/UnrealEd/Private/EditorSelectUtils.cpp
+++ b/Engine/Source/Editor/UnrealEd/Private/EditorSelectUtils.cpp
@@ -571,7 +571,7 @@ bool UUnrealEdEngine::CanSelectActor(AActor* Actor, bool bInSelected, bool bSele
 
 void UUnrealEdEngine::SelectActor(AActor* Actor, bool bInSelected, bool bNotify, bool bSelectEvenIfHidden, bool bForceRefresh)
 {
-	if (!Actor || (Actor->GetRootSelectionParent() != nullptr && !Actor->SupportsSubRootSelection()))
+	if (!Actor || Actor->GetRootSelectionParent() != nullptr)
 	{
 		return;
 	}
@@ -590,7 +590,7 @@ void UUnrealEdEngine::SelectActor(AActor* Actor, bool bInSelected, bool bNotify,
 			.SetAllowHidden(bSelectEvenIfHidden)
 			.SetWarnIfLocked(true)
 			.SetAllowLegacyNotifications(false)
-			.SetAllowSubRootSelection(Actor->SupportsSubRootSelection());
+			.SetAllowSubRootSelection(false);
 
 		const bool bSelectionChanged = bInSelected
 			? SelectionSet->SelectElement(UEngineElementsLibrary::AcquireEditorActorElementHandle(Actor), SelectionOptions)
