Outliner: Fix a crash in the "level" column name selection code when you tried to exit PIE. The column name for the level column is now created once on Outliner creation instead of being a lambda.

diff --git a/Engine/Source/Editor/SceneOutliner/Private/SceneOutlinerModule.cpp b/Engine/Source/Editor/SceneOutliner/Private/SceneOutlinerModule.cpp
--- a/Engine/Source/Editor/SceneOutliner/Private/SceneOutlinerModule.cpp
+++ b/Engine/Source/Editor/SceneOutliner/Private/SceneOutlinerModule.cpp
@@ -546,15 +546,15 @@ void FSceneOutlinerModule::CreateActorInfoColumns(FSceneOutlinerInitializationOp
 	};
 
 	// The "Level" column should be named "Package Short Name" in wp enabled levels
-	auto LevelColumnName = TAttribute<FText>::CreateLambda([WorldPtr]() -> FText
+	FText LevelColumnName;
+	if(WorldPtr && WorldPtr->PersistentLevel && WorldPtr->PersistentLevel->IsUsingExternalActors())
 	{
-		if (WorldPtr && WorldPtr->PersistentLevel->IsUsingExternalActors())
-		{
-			return FSceneOutlinerBuiltInColumnTypes::PackageShortName_Localized();
-		}
-
-		return FSceneOutlinerBuiltInColumnTypes::Level_Localized();
-	});
+		LevelColumnName = FSceneOutlinerBuiltInColumnTypes::PackageShortName_Localized();
+	}
+	else
+	{
+		LevelColumnName = FSceneOutlinerBuiltInColumnTypes::Level_Localized();
+	}
 
 	AddTextInfoColumn(FSceneOutlinerBuiltInColumnTypes::Mobility(), FSceneOutlinerBuiltInColumnTypes::Mobility_Localized(), MobilityInfoText);
 	AddTextInfoColumn(FSceneOutlinerBuiltInColumnTypes::Level(), LevelColumnName, LevelInfoText);
