Scene Outliner: Fix crash in folder context menu generation.

Callstack:
UnrealEditor_SceneOutliner!SSceneOutliner::AddMoveToFolderOutliner()
UnrealEditor_SceneOutliner!SSceneOutliner::FillFoldersSubMenu()
UnrealEditor_SceneOutliner!TBaseSPMethodDelegateInstance<1,SSceneOutliner const ,1,void __cdecl(UToolMenu * __ptr64),FDefaultDelegateUserPolicy>::Execute()
UnrealEditor_ToolMenus!UToolMenus::GenerateSubMenu()
UnrealEditor_ToolMenus!UToolMenus::PopulateSubMenu()
UnrealEditor_ToolMenus!TBaseUObjectMethodDelegateInstance<0,UToolMenus,void __cdecl(FMenuBuilder & __ptr64),FDefaultDelegateUserPolicy,TWeakObjectPtr<UToolMenu,FWeakObjectPtr>,FName>::Execute()
UnrealEditor_Slate!SMenuEntryBlock::MakeNewMenuWidget()
UnrealEditor_Slate!TBaseSPMethodDelegateInstance<1,SMenuEntryBlock const ,1,TSharedRef<SWidget,1> __cdecl(void),FDefaultDelegateUserPolicy>::Execute()
UnrealEditor_Slate!TDelegate<TSharedRef<SWidget,1> __cdecl(void),FDefaultDelegateUserPolicy>::Execute()
UnrealEditor_Slate!SMenuAnchor::SetIsOpen()
UnrealEditor_Slate!SMenuEntryBlock::UpdateSubMenuState()

diff --git a/Engine/Source/Editor/SceneOutliner/Private/SSceneOutliner.cpp b/Engine/Source/Editor/SceneOutliner/Private/SSceneOutliner.cpp
--- a/Engine/Source/Editor/SceneOutliner/Private/SSceneOutliner.cpp
+++ b/Engine/Source/Editor/SceneOutliner/Private/SSceneOutliner.cpp
@@ -1325,7 +1325,10 @@ TSharedRef<TSet<FFolder>> SSceneOutliner::GatherInvalidMoveToDestinations() cons
 
 	for (const auto& Item : OutlinerTreeView->GetSelectedItems())
 	{
-		if (FFolderTreeItem* ParentFolderItem = Item->GetParent()->CastTo<FFolderTreeItem>())
+		const FSceneOutlinerTreeItemPtr Parent = Item->GetParent();
+		const FFolderTreeItem* ParentFolderItem = Parent.IsValid() ? Parent->CastTo<FFolderTreeItem>() : nullptr;
+
+		if (ParentFolderItem)
 		{
 			auto FolderHasOtherSubFolders = [&Item](const TWeakPtr<ISceneOutlinerTreeItem>& WeakItem)
 			{
