Scene Outliner: Fixed invalid numeric cast

diff --git a/Engine/Source/Editor/SceneOutliner/Public/SceneOutlinerPublicTypes.h b/Engine/Source/Editor/SceneOutliner/Public/SceneOutlinerPublicTypes.h
--- a/Engine/Source/Editor/SceneOutliner/Public/SceneOutlinerPublicTypes.h
+++ b/Engine/Source/Editor/SceneOutliner/Public/SceneOutlinerPublicTypes.h
@@ -74,7 +74,7 @@ enum class ESceneOutlinerColumnVisibility : uint8
 struct FSceneOutlinerColumnInfo
 {
 	FSceneOutlinerColumnInfo(
-		ESceneOutlinerColumnVisibility InVisibility, int32 InPriorityIndex, const FCreateSceneOutlinerColumn& InFactory = FCreateSceneOutlinerColumn()
+		ESceneOutlinerColumnVisibility InVisibility, uint8 InPriorityIndex, const FCreateSceneOutlinerColumn& InFactory = FCreateSceneOutlinerColumn()
 		, bool inCanBeHidden = true, TOptional<float> InFillSize = TOptional<float>(), TAttribute<FText> InColumnLabel = TAttribute<FText>()
 		, EHeaderComboVisibility InHeaderComboVisibility = EHeaderComboVisibility::OnHover, FOnGetContent InOnGetMenuContent = FOnGetContent())
 		: Visibility(InVisibility), PriorityIndex(InPriorityIndex), bCanBeHidden(inCanBeHidden)
