WorldPartition: Add UActorDescContainerInstance/FWorldPartitionActorDescInstance which take transient/context data from UActorDescContainer/FWorldPartitionActorDesc leaving them as pure on disk repres

diff --git a/Engine/Source/Editor/SceneOutliner/Private/ActorDescTreeItem.cpp b/Engine/Source/Editor/SceneOutliner/Private/ActorDescTreeItem.cpp
--- a/Engine/Source/Editor/SceneOutliner/Private/ActorDescTreeItem.cpp
+++ b/Engine/Source/Editor/SceneOutliner/Private/ActorDescTreeItem.cpp
@@ -1,9 +1,10 @@
 // Copyright Epic Games, Inc. All Rights Reserved.
 
 #include "ActorDescTreeItem.h"
-#include "WorldPartition/ActorDescContainer.h"
+#include "WorldPartition/ActorDescContainerInstance.h"
 #include "WorldPartition/WorldPartition.h"
 #include "WorldPartition/WorldPartitionActorDesc.h"
+#include "WorldPartition/WorldPartitionActorDescInstance.h"
 #include "WorldPartition/LoaderAdapter/LoaderAdapterPinnedActors.h"
 #include "Widgets/DeclarativeSyntaxSupport.h"
 #include "Layout/WidgetPath.h"
@@ -121,11 +122,11 @@ private:
 	{
 		if (TSharedPtr<FActorDescTreeItem> TreeItem = TreeItemPtr.Pin())
 		{
-			if (FWorldPartitionActorDesc* ActorDesc = TreeItem->ActorDescHandle.Get())
+			if (FWorldPartitionActorDescInstance* ActorDescInstance = TreeItem->ActorDescHandle.GetInstance())
 			{
 				FFormatNamedArguments Args;
 				Args.Add(TEXT("ActorLabel"), FText::FromString(TreeItem->GetDisplayString()));
-				Args.Add(TEXT("UnloadState"), ActorDesc->GetUnloadedReason());
+				Args.Add(TEXT("UnloadState"), ActorDescInstance->GetUnloadedReason());
 
 				return FText::Format(LOCTEXT("UnloadedActorDisplay", "{ActorLabel} ({UnloadState})"), Args);
 			}
@@ -142,9 +143,9 @@ private:
 	{
 		if (TSharedPtr<FActorDescTreeItem> TreeItem = TreeItemPtr.Pin())
 		{
-			if (const FWorldPartitionActorDesc* ActorDesc = TreeItem->ActorDescHandle.Get())
+			if (const FWorldPartitionActorDescInstance* ActorDescInstance = TreeItem->ActorDescHandle.GetInstance())
 			{
-				return FText::FromName(ActorDesc->GetDisplayClassName());
+				return FText::FromName(ActorDescInstance->GetDisplayClassName());
 			}
 		}
 
@@ -162,9 +163,9 @@ private:
 
 		if (TreeItem.IsValid() && WeakSceneOutliner.IsValid())
 		{
-			if (const FWorldPartitionActorDesc* ActorDesc = TreeItem->ActorDescHandle.Get())
+			if (const FWorldPartitionActorDescInstance* ActorDescInstance = TreeItem->ActorDescHandle.GetInstance())
 			{
-				const FName IconName = ActorDesc->GetDisplayClassName();
+				const FName IconName = ActorDescInstance->GetDisplayClassName();
 
 				const FSlateBrush* CachedBrush = WeakSceneOutliner.Pin()->GetCachedIconForClass(IconName);
 				if (CachedBrush != nullptr)
@@ -173,7 +174,7 @@ private:
 				}
 				else if (IconName != NAME_None)
 				{
-					const FSlateBrush* FoundSlateBrush = FSlateIconFinder::FindIconForClass(ActorDesc->GetActorNativeClass()).GetIcon();
+					const FSlateBrush* FoundSlateBrush = FSlateIconFinder::FindIconForClass(ActorDescInstance->GetActorNativeClass()).GetIcon();
 					WeakSceneOutliner.Pin()->CacheIconForClass(IconName, FoundSlateBrush);
 					return FoundSlateBrush;
 				}
@@ -227,15 +228,15 @@ private:
 	}
 };
 
-FActorDescTreeItem::FActorDescTreeItem(const FGuid& InActorGuid, UActorDescContainer* InContainer)
+FActorDescTreeItem::FActorDescTreeItem(const FGuid& InActorGuid, UActorDescContainerInstance* InContainerInstance)
 	: IActorBaseTreeItem(Type)
-	, ActorDescHandle(InContainer, InActorGuid)
-	, ID(ComputeTreeItemID(InActorGuid, InContainer))
+	, ActorDescHandle(InContainerInstance, InActorGuid)
+	, ID(ComputeTreeItemID(InActorGuid, InContainerInstance))
 	, ActorGuid(InActorGuid)
 {
-	if (const FWorldPartitionActorDesc* const ActorDesc = ActorDescHandle.Get())
+	if (const FWorldPartitionActorDescInstance* const ActorDescInstance = ActorDescHandle.GetInstance())
 	{
-		DisplayString = ActorDesc->GetActorLabel().ToString();
+		DisplayString = ActorDescInstance->GetActorLabel().ToString();
 	}
 	else
 	{
@@ -243,20 +244,20 @@ FActorDescTreeItem::FActorDescTreeItem(const FGuid& InActorGuid, UActorDescConta
 	}
 }
 
-FSceneOutlinerTreeItemID FActorDescTreeItem::ComputeTreeItemID(FGuid InActorGuid, UActorDescContainer* InContainer)
+FSceneOutlinerTreeItemID FActorDescTreeItem::ComputeTreeItemID(FGuid InActorGuid, UActorDescContainerInstance* InContainerInstance)
 {
 	FArchiveMD5 Ar;
 	Ar << InActorGuid;
 
-	FObjectKey ContainerKey(InContainer);
+	FObjectKey ContainerKey(InContainerInstance);
 	Ar << ContainerKey;
 
 	return FSceneOutlinerTreeItemID(Ar.GetGuidFromHash());
 }
 
-bool FActorDescTreeItem::ShouldDisplayInOutliner(const FWorldPartitionActorDesc* ActorDesc)
+bool FActorDescTreeItem::ShouldDisplayInOutliner(const FWorldPartitionActorDescInstance* InActorDescInstance)
 {
-	return ActorDesc && ActorDesc->IsListedInSceneOutliner() && (ActorDesc->GetActorIsRuntimeOnly() || ActorDesc->IsEditorRelevant());
+	return InActorDescInstance && InActorDescInstance->IsListedInSceneOutliner() && (InActorDescInstance->GetActorIsRuntimeOnly() || InActorDescInstance->IsEditorRelevant());
 }
 
 FSceneOutlinerTreeItemID FActorDescTreeItem::GetID() const
@@ -281,10 +282,10 @@ TSharedRef<SWidget> FActorDescTreeItem::GenerateLabelWidget(ISceneOutliner& Outl
 
 void FActorDescTreeItem::FocusActorBounds() const
 {
-	if (FWorldPartitionActorDesc const* ActorDesc = ActorDescHandle.Get())
+	if (FWorldPartitionActorDescInstance const* ActorDescInstance = ActorDescHandle.GetInstance())
 	{
 		const bool bActiveViewportOnly = true;
-		const FBox EditorBounds = ActorDesc->GetEditorBounds();
+		const FBox EditorBounds = ActorDescInstance->GetEditorBounds();
 		if (EditorBounds.IsValid)
 		{
 			GEditor->MoveViewportCamerasToBox(EditorBounds, bActiveViewportOnly, 0.5f);
@@ -294,10 +295,10 @@ void FActorDescTreeItem::FocusActorBounds() const
 
 void FActorDescTreeItem::CopyActorFilePathtoClipboard() const
 {
-	if (FWorldPartitionActorDesc const* ActorDesc = ActorDescHandle.Get())
+	if (FWorldPartitionActorDescInstance const* ActorDescInstance = ActorDescHandle.GetInstance())
 	{
 		FString PackageFilename;
-		if (FPackageName::TryConvertLongPackageNameToFilename(ActorDesc->GetActorPackage().ToString(), PackageFilename, FPackageName::GetAssetPackageExtension()))
+		if (FPackageName::TryConvertLongPackageNameToFilename(ActorDescInstance->GetActorPackage().ToString(), PackageFilename, FPackageName::GetAssetPackageExtension()))
 		{
 			FString Result = FPaths::ConvertRelativePathToFull(PackageFilename);
 			FPlatformApplicationMisc::ClipboardCopy(*Result);
@@ -319,14 +320,14 @@ bool FActorDescTreeItem::GetVisibility() const
 
 bool FActorDescTreeItem::ShouldShowPinnedState() const
 {
-	return FLoaderAdapterPinnedActors::SupportsPinning(ActorDescHandle.Get());
+	return FLoaderAdapterPinnedActors::SupportsPinning(ActorDescHandle.GetInstance());
 }
 
 bool FActorDescTreeItem::GetPinnedState() const
 {
-	if (ActorDescHandle.IsValid() && ActorDescHandle->GetContainer())
+	if (ActorDescHandle.IsValid() && ActorDescHandle.GetContainerInstance())
 	{
-		UWorldPartition* WorldPartition = ActorDescHandle->GetContainer()->GetWorldPartition();
+		UWorldPartition* WorldPartition = ActorDescHandle.GetContainerInstance()->GetWorldPartition();
 		return WorldPartition ? WorldPartition->IsActorPinned(GetGuid()) : false;
 	}
 	return false;
diff --git a/Engine/Source/Editor/SceneOutliner/Private/ActorHierarchy.cpp b/Engine/Source/Editor/SceneOutliner/Private/ActorHierarchy.cpp
--- a/Engine/Source/Editor/SceneOutliner/Private/ActorHierarchy.cpp
+++ b/Engine/Source/Editor/SceneOutliner/Private/ActorHierarchy.cpp
@@ -20,7 +20,8 @@
 #include "LevelInstance/LevelInstanceInterface.h"
 #include "LevelInstance/LevelInstanceSubsystem.h"
 #include "Modules/ModuleManager.h"
-#include "WorldPartition/ActorDescContainer.h"
+#include "WorldPartition/ActorDescContainerInstance.h"
+#include "WorldPartition/WorldPartitionActorDescInstance.h"
 #include "WorldPartition/WorldPartition.h"
 #include "WorldPartition/WorldPartitionHelpers.h"
 #include "WorldPartition/WorldPartitionActorDesc.h"
@@ -66,8 +67,8 @@ void FActorHierarchy::Create_Internal(FActorHierarchy* Hierarchy, const TWeakObj
 
 		if (UWorldPartition* WorldPartition = World->GetWorldPartition())
 		{
-			WorldPartition->OnActorDescAddedEvent.AddRaw(Hierarchy, &FActorHierarchy::OnActorDescAdded);
-			WorldPartition->OnActorDescRemovedEvent.AddRaw(Hierarchy, &FActorHierarchy::OnActorDescRemoved);
+			WorldPartition->OnActorDescInstanceAddedEvent.AddRaw(Hierarchy, &FActorHierarchy::OnActorDescInstanceAdded);
+			WorldPartition->OnActorDescInstanceRemovedEvent.AddRaw(Hierarchy, &FActorHierarchy::OnActorDescInstanceRemoved);
 		}
 	}
 
@@ -116,8 +117,8 @@ FActorHierarchy::~FActorHierarchy()
 
 		if (UWorldPartition* WorldPartition = RepresentingWorld->GetWorldPartition())
 		{
-			WorldPartition->OnActorDescAddedEvent.RemoveAll(this);
-			WorldPartition->OnActorDescRemovedEvent.RemoveAll(this);
+			WorldPartition->OnActorDescInstanceAddedEvent.RemoveAll(this);
+			WorldPartition->OnActorDescInstanceRemovedEvent.RemoveAll(this);
 		}
 	}
 
@@ -273,11 +274,11 @@ FSceneOutlinerTreeItemPtr FActorHierarchy::FindOrCreateParentItem(const ISceneOu
 	}
 	else if (const FActorDescTreeItem* ActorDescItem = Item.CastTo<FActorDescTreeItem>())
 	{
-		if (const FWorldPartitionActorDesc* ActorDesc = ActorDescItem->ActorDescHandle.Get())
+		if (const FWorldPartitionActorDescInstance* ActorDescInstance = ActorDescItem->ActorDescHandle.GetInstance())
 		{
 			if (UWorld* RepresentingWorldPtr = RepresentingWorld.Get())
 			{
-				const FFolder ActorDescFolder = FActorFolders::GetActorDescFolder(*RepresentingWorldPtr, ActorDesc);
+				const FFolder ActorDescFolder = FActorFolders::GetActorDescInstanceFolder(*RepresentingWorldPtr, ActorDescInstance);
 				if (!ActorDescFolder.IsNone())
 				{
 					if (const FSceneOutlinerTreeItemPtr* ParentItem = Items.Find(ActorDescFolder))
@@ -292,12 +293,12 @@ FSceneOutlinerTreeItemPtr FActorHierarchy::FindOrCreateParentItem(const ISceneOu
 			}
 
 			// Parent Actor (Actor attachement / parenting)
-			const FGuid& ParentActorGuid = ActorDesc->GetSceneOutlinerParent();
+			const FGuid& ParentActorGuid = ActorDescInstance->GetSceneOutlinerParent();
 			if (ParentActorGuid.IsValid())
 			{
-				if (UActorDescContainer* ActorDescContainer = ActorDesc->GetContainer())
+				if (UActorDescContainerInstance* ContainerInstance = ActorDescInstance->GetContainerInstance())
 				{
-					if (const FWorldPartitionActorDesc* ParentActorDesc = ActorDescContainer->GetActorDesc(ParentActorGuid))
+					if (const FWorldPartitionActorDescInstance* ParentActorDesc = ContainerInstance->GetActorDescInstance(ParentActorGuid))
 					{
 						// If parent actor is loaded
 						if (AActor* ParentActor = ParentActorDesc->GetActor())
@@ -314,21 +315,21 @@ FSceneOutlinerTreeItemPtr FActorHierarchy::FindOrCreateParentItem(const ISceneOu
 						}
 
 						// Find unloaded parent actor node (from the guid)
-						if (const FSceneOutlinerTreeItemPtr* ParentItem = Items.Find(FActorDescTreeItem::ComputeTreeItemID(ParentActorGuid, ActorDescContainer)))
+						if (const FSceneOutlinerTreeItemPtr* ParentItem = Items.Find(FActorDescTreeItem::ComputeTreeItemID(ParentActorGuid, ContainerInstance)))
 						{
 							return *ParentItem;
 						}
 						else
 						{
-							return bCreate ? Mode->CreateItemFor<FActorDescTreeItem>(FActorDescTreeItem(ParentActorGuid, ActorDescContainer)) : nullptr;
+							return bCreate ? Mode->CreateItemFor<FActorDescTreeItem>(FActorDescTreeItem(ParentActorGuid, ContainerInstance)) : nullptr;
 						}
 					}
 				}
 			}
-			else if (UActorDescContainer* ActorDescContainer = ActorDesc->GetContainer())
+			else if (UActorDescContainerInstance* ContainerInstance = ActorDescInstance->GetContainerInstance())
 			{
 				const ULevelInstanceSubsystem* LevelInstanceSubsystem = UWorld::GetSubsystem<ULevelInstanceSubsystem>(RepresentingWorld.Get());
-				UWorld* OuterWorld = ActorDescContainer->GetTypedOuter<UWorld>();
+				UWorld* OuterWorld = ContainerInstance->GetTypedOuter<UWorld>();
 				// If parent actor is loaded
 				if (AActor* ParentActor = OuterWorld ? Cast<AActor>(LevelInstanceSubsystem->GetOwningLevelInstance(OuterWorld->PersistentLevel)) : nullptr)
 				{
@@ -481,11 +482,11 @@ void FActorHierarchy::CreateUnloadedItems(UWorld* World, TArray<FSceneOutlinerTr
 
 				const TSet<FGuid>& LoadedActors = LoadedActorCache.GetLoadedActorsForLevel(OuterLevel);
 				
-				FWorldPartitionHelpers::ForEachActorDesc(WorldPartition, [this, &LoadedActors, &OutItems](const FWorldPartitionActorDesc* ActorDesc)
+				FWorldPartitionHelpers::ForEachActorDescInstance(WorldPartition, [this, &LoadedActors, &OutItems](const FWorldPartitionActorDescInstance* ActorDescInstance)
 				{
-					if (ActorDesc != nullptr && !LoadedActors.Contains(ActorDesc->GetGuid()) && FActorDescTreeItem::ShouldDisplayInOutliner(ActorDesc))
+					if (ActorDescInstance != nullptr && !LoadedActors.Contains(ActorDescInstance->GetGuid()) && FActorDescTreeItem::ShouldDisplayInOutliner(ActorDescInstance))
 					{
-						if (const FSceneOutlinerTreeItemPtr ActorDescItem = Mode->CreateItemFor<FActorDescTreeItem>(FActorDescTreeItem(ActorDesc->GetGuid(), ActorDesc->GetContainer())))
+						if (const FSceneOutlinerTreeItemPtr ActorDescItem = Mode->CreateItemFor<FActorDescTreeItem>(FActorDescTreeItem(ActorDescInstance->GetGuid(), ActorDescInstance->GetContainerInstance())))
 						{
 							OutItems.Add(ActorDescItem);
 						}
@@ -694,11 +695,11 @@ void FActorHierarchy::RemoveActorDesc(AActor& InActor)
 	if (UWorldPartition* WorldPartition = FWorldPartitionHelpers::GetWorldPartition(&InActor))
 	{
 		const FGuid& ActorGuid = InActor.GetActorGuid();
-		if (FWorldPartitionActorDesc* ActorDesc = WorldPartition->GetActorDesc(ActorGuid))
+		if (FWorldPartitionActorDescInstance* ActorDescInstance = WorldPartition->GetActorDescInstance(ActorGuid))
 		{
 			FSceneOutlinerHierarchyChangedData EventData;
 			EventData.Type = FSceneOutlinerHierarchyChangedData::Removed;
-			EventData.ItemIDs.Add(FActorDescTreeItem::ComputeTreeItemID(ActorGuid, ActorDesc->GetContainer()));
+			EventData.ItemIDs.Add(FActorDescTreeItem::ComputeTreeItemID(ActorGuid, ActorDescInstance->GetContainerInstance()));
 			HierarchyChangedEvent.Broadcast(EventData);
 		}
 	}
@@ -711,11 +712,11 @@ void FActorHierarchy::AddActorDesc(AActor& InActor)
 		if (UWorldPartition* WorldPartition = FWorldPartitionHelpers::GetWorldPartition(&InActor))
 		{
 			const FGuid& ActorGuid = InActor.GetActorGuid();
-			if (FWorldPartitionActorDesc* ActorDesc = WorldPartition->GetActorDesc(ActorGuid); FActorDescTreeItem::ShouldDisplayInOutliner(ActorDesc))
+			if (FWorldPartitionActorDescInstance* ActorDescInstance = WorldPartition->GetActorDescInstance(ActorGuid); FActorDescTreeItem::ShouldDisplayInOutliner(ActorDescInstance))
 			{
 				FSceneOutlinerHierarchyChangedData EventData;
 				EventData.Type = FSceneOutlinerHierarchyChangedData::Added;
-				EventData.Items.Add(Mode->CreateItemFor<FActorDescTreeItem>(FActorDescTreeItem(ActorDesc->GetGuid(), ActorDesc->GetContainer())));
+				EventData.Items.Add(Mode->CreateItemFor<FActorDescTreeItem>(FActorDescTreeItem(ActorDescInstance->GetGuid(), ActorDescInstance->GetContainerInstance())));
 				HierarchyChangedEvent.Broadcast(EventData);
 			}
 		}
@@ -736,24 +737,24 @@ void FActorHierarchy::OnLoadedActorRemoved(AActor& InActor)
 	AddActorDesc(InActor);
 }
 
-void FActorHierarchy::OnActorDescAdded(FWorldPartitionActorDesc* ActorDesc)
+void FActorHierarchy::OnActorDescInstanceAdded(FWorldPartitionActorDescInstance* InActorDescInstance)
 {
-	if (IsShowingUnloadedActors() && ActorDesc && !ActorDesc->IsLoaded(true) && FActorDescTreeItem::ShouldDisplayInOutliner(ActorDesc))
+	if (IsShowingUnloadedActors() && InActorDescInstance && !InActorDescInstance->IsLoaded(true) && FActorDescTreeItem::ShouldDisplayInOutliner(InActorDescInstance))
 	{
 		FSceneOutlinerHierarchyChangedData EventData;
 		EventData.Type = FSceneOutlinerHierarchyChangedData::Added;
-		EventData.Items.Add(Mode->CreateItemFor<FActorDescTreeItem>(FActorDescTreeItem(ActorDesc->GetGuid(), ActorDesc->GetContainer())));
+		EventData.Items.Add(Mode->CreateItemFor<FActorDescTreeItem>(FActorDescTreeItem(InActorDescInstance->GetGuid(), InActorDescInstance->GetContainerInstance())));
 		HierarchyChangedEvent.Broadcast(EventData);
 	}
 }
 
-void FActorHierarchy::OnActorDescRemoved(FWorldPartitionActorDesc* ActorDesc)
+void FActorHierarchy::OnActorDescInstanceRemoved(FWorldPartitionActorDescInstance* InActorDescInstance)
 {
-	if (IsShowingUnloadedActors() && ActorDesc)
+	if (IsShowingUnloadedActors() && InActorDescInstance)
 	{
 		FSceneOutlinerHierarchyChangedData EventData;
 		EventData.Type = FSceneOutlinerHierarchyChangedData::Removed;
-		EventData.ItemIDs.Add(FActorDescTreeItem::ComputeTreeItemID(ActorDesc->GetGuid(), ActorDesc->GetContainer()));
+		EventData.ItemIDs.Add(FActorDescTreeItem::ComputeTreeItemID(InActorDescInstance->GetGuid(), InActorDescInstance->GetContainerInstance()));
 		HierarchyChangedEvent.Broadcast(EventData);
 	}
 }
diff --git a/Engine/Source/Editor/SceneOutliner/Private/ActorMode.cpp b/Engine/Source/Editor/SceneOutliner/Private/ActorMode.cpp
--- a/Engine/Source/Editor/SceneOutliner/Private/ActorMode.cpp
+++ b/Engine/Source/Editor/SceneOutliner/Private/ActorMode.cpp
@@ -96,17 +96,14 @@ namespace SceneOutliner
 		return false;
 	}
 
-	bool FActorDescSelector::operator()(const TWeakPtr<ISceneOutlinerTreeItem>& Item, FWorldPartitionActorDesc*& ActorDescPtrOut) const
+	bool FActorHandleSelector::operator()(const TWeakPtr<ISceneOutlinerTreeItem>& Item, FWorldPartitionHandle& ActorHandleOut) const
 	{
 		if (TSharedPtr<ISceneOutlinerTreeItem> ItemPtr = Item.Pin())
 		{
 			if (FActorDescTreeItem* ActorDescItem = ItemPtr->CastTo<FActorDescTreeItem>())
 			{
-				if (FWorldPartitionActorDesc* ActorDesc = ActorDescItem->ActorDescHandle.Get())
-				{
-					ActorDescPtrOut = ActorDesc;
-					return true;
-				}
+				ActorHandleOut = ActorDescItem->ActorDescHandle;
+				return true;
 			}
 		}
 
diff --git a/Engine/Source/Editor/SceneOutliner/Private/SceneOutlinerActorInfoColumn.cpp b/Engine/Source/Editor/SceneOutliner/Private/SceneOutlinerActorInfoColumn.cpp
--- a/Engine/Source/Editor/SceneOutliner/Private/SceneOutlinerActorInfoColumn.cpp
+++ b/Engine/Source/Editor/SceneOutliner/Private/SceneOutlinerActorInfoColumn.cpp
@@ -23,7 +23,7 @@
 #include "WorldTreeItem.h"
 #include "WorldPartition/DataLayer/DataLayerInstance.h"
 #include "WorldPartition/DataLayer/WorldDataLayers.h"
-#include "WorldPartition/WorldPartitionActorDesc.h"
+#include "WorldPartition/WorldPartitionActorDescInstance.h"
 #include "Styling/StyleColors.h"
 
 #define LOCTEXT_NAMESPACE "SceneOutlinerActorInfoColumn"
@@ -61,9 +61,9 @@ struct FGetInfo
 		}
 		else if (const FActorDescTreeItem* ActorDescItem = Item.CastTo<FActorDescTreeItem>())
 		{
-			if (const FWorldPartitionActorDesc* ActorDesc = ActorDescItem->ActorDescHandle.Get())
+			if (const FWorldPartitionActorDescInstance* ActorDescInstance = ActorDescItem->ActorDescHandle.GetInstance())
 			{
-				return ActorDesc->GetDisplayClassName().ToString();
+				return ActorDescInstance->GetDisplayClassName().ToString();
 			}
 		}
 
diff --git a/Engine/Source/Editor/SceneOutliner/Private/SceneOutlinerHelpers.cpp b/Engine/Source/Editor/SceneOutliner/Private/SceneOutlinerHelpers.cpp
--- a/Engine/Source/Editor/SceneOutliner/Private/SceneOutlinerHelpers.cpp
+++ b/Engine/Source/Editor/SceneOutliner/Private/SceneOutlinerHelpers.cpp
@@ -8,7 +8,7 @@
 #include "Engine/Blueprint.h"
 #include "Misc/PackageName.h"
 #include "Modules/ModuleManager.h"
-#include "WorldPartition/WorldPartitionActorDesc.h"
+#include "WorldPartition/WorldPartitionActorDescInstance.h"
 #include "UObject/Package.h"
 
 #include "SourceControlHelpers.h"
@@ -40,9 +40,9 @@ namespace SceneOutliner
 		}
 		else if (const FActorDescTreeItem* ActorDescItem = TreeItem.CastTo<FActorDescTreeItem>())
 		{
-			if (const FWorldPartitionActorDesc* ActorDesc = ActorDescItem->ActorDescHandle.Get())
+			if (const FWorldPartitionActorDescInstance* ActorDescInstance = ActorDescItem->ActorDescHandle.GetInstance())
 			{
-				return ActorDesc->GetActorPackage().ToString();
+				return ActorDescInstance->GetActorPackage().ToString();
 			}
 		}
 
@@ -73,9 +73,9 @@ namespace SceneOutliner
 		}
 		else if (const FActorDescTreeItem* ActorDescItem = TreeItem.CastTo<FActorDescTreeItem>())
 		{
-			if (const FWorldPartitionActorDesc* ActorDesc = ActorDescItem->ActorDescHandle.Get())
+			if (const FWorldPartitionActorDescInstance* ActorDescInstance = ActorDescItem->ActorDescHandle.GetInstance())
 			{
-				return FindPackage(nullptr, *ActorDesc->GetActorPackage().ToString());
+				return FindPackage(nullptr, *ActorDescInstance->GetActorPackage().ToString());
 			}
 		}
 
diff --git a/Engine/Source/Editor/SceneOutliner/Private/SceneOutlinerModule.cpp b/Engine/Source/Editor/SceneOutliner/Private/SceneOutlinerModule.cpp
--- a/Engine/Source/Editor/SceneOutliner/Private/SceneOutlinerModule.cpp
+++ b/Engine/Source/Editor/SceneOutliner/Private/SceneOutlinerModule.cpp
@@ -23,7 +23,8 @@
 #include "ActorDescTreeItem.h"
 #include "FolderTreeItem.h"
 #include "WorldTreeItem.h"
-#include "WorldPartition/ActorDescContainer.h"
+#include "WorldPartition/ActorDescContainerInstance.h"
+#include "WorldPartition/WorldPartitionActorDescInstance.h"
 #include "WorldPartition/DataLayer/DataLayerInstance.h"
 #include "WorldPartition/DataLayer/DataLayerManager.h"
 #include "WorldPartition/ContentBundle/ContentBundleEngineSubsystem.h"
@@ -352,18 +353,18 @@ void FSceneOutlinerModule::CreateActorInfoColumns(FSceneOutlinerInitializationOp
 			}
 			else if (const FActorDescTreeItem* ActorDescItem = Item.CastTo<FActorDescTreeItem>())
 			{
-				if (const FWorldPartitionActorDesc* ActorDesc = ActorDescItem->ActorDescHandle.Get(); ActorDesc && !ActorDesc->GetDataLayerInstanceNames().IsEmpty())
+				if (const FWorldPartitionActorDescInstance* ActorDescInstance = ActorDescItem->ActorDescHandle.GetInstance(); ActorDescInstance && !ActorDescInstance->GetDataLayerInstanceNames().IsEmpty())
 				{
-					if (const UActorDescContainer* ActorDescContainer = ActorDescItem->ActorDescHandle.Container.Get())
+					if (const UActorDescContainerInstance* ActorDescContainerInstance = ActorDescInstance->GetContainerInstance())
 					{
-						const UWorld* OwningWorld = ActorDescContainer->GetWorldPartition()->GetWorld();
+						const UWorld* OwningWorld = ActorDescContainerInstance->GetWorldPartition()->GetWorld();
 						if (const UDataLayerManager* DataLayerManager = UDataLayerManager::GetDataLayerManager(OwningWorld))
 						{
 							TSet<const UDataLayerInstance*> DataLayerInstances;
-							DataLayerInstances.Append(DataLayerManager->GetDataLayerInstances(ActorDesc->GetDataLayerInstanceNames()));
+							DataLayerInstances.Append(DataLayerManager->GetDataLayerInstances(ActorDescInstance->GetDataLayerInstanceNames()));
 							if (ULevelInstanceSubsystem* LevelInstanceSubsystem = UWorld::GetSubsystem<ULevelInstanceSubsystem>(OwningWorld))
 							{
-								UWorld* OuterWorld = ActorDescContainer->GetTypedOuter<UWorld>();
+								UWorld* OuterWorld = ActorDescContainerInstance->GetTypedOuter<UWorld>();
 								// Add parent container Data Layer Instances
 								AActor* CurrentActor = OuterWorld ? Cast<AActor>(LevelInstanceSubsystem->GetOwningLevelInstance(OuterWorld->PersistentLevel)) : nullptr;
 								while (CurrentActor)
@@ -401,9 +402,9 @@ void FSceneOutlinerModule::CreateActorInfoColumns(FSceneOutlinerInitializationOp
 		}
 		else if (const FActorDescTreeItem* ActorDescItem = Item.CastTo<FActorDescTreeItem>())
 		{
-			if (const FWorldPartitionActorDesc* ActorDesc = ActorDescItem->ActorDescHandle.Get())
+			if (const FWorldPartitionActorDescInstance* ActorDescInstance = ActorDescItem->ActorDescHandle.GetInstance())
 			{
-				if (const UContentBundleDescriptor* Descriptor = ContentBundleEngineSubsystem->GetContentBundleDescriptor(ActorDesc->GetContentBundleGuid()))
+				if (const UContentBundleDescriptor* Descriptor = ContentBundleEngineSubsystem->GetContentBundleDescriptor(ActorDescInstance->GetContentBundleGuid()))
 				{
 					return Descriptor->GetDisplayName();
 				}
@@ -427,11 +428,11 @@ void FSceneOutlinerModule::CreateActorInfoColumns(FSceneOutlinerInitializationOp
 		}
 		else if (const FActorDescTreeItem* ActorDescItem = Item.CastTo<FActorDescTreeItem>())
 		{
-			if (const FWorldPartitionActorDesc* ActorDesc = ActorDescItem->ActorDescHandle.Get())
+			if (const FWorldPartitionActorDescInstance* ActorDescInstance = ActorDescItem->ActorDescHandle.GetInstance())
 			{
-				if (FName LevelPackage = ActorDesc->GetContainerPackage(); !LevelPackage.IsNone())
+				if (FName LevelPackage = ActorDescInstance->GetChildContainerPackage(); !LevelPackage.IsNone())
 				{
-					return ActorDesc->GetContainerPackage().ToString();
+					return ActorDescInstance->GetChildContainerPackage().ToString();
 				}
 			}
 		}
@@ -472,9 +473,9 @@ void FSceneOutlinerModule::CreateActorInfoColumns(FSceneOutlinerInitializationOp
 		}
 		else if (const FActorDescTreeItem* ActorDescItem = Item.CastTo<FActorDescTreeItem>())
 		{
-			if (const FWorldPartitionActorDesc* ActorDesc = ActorDescItem->ActorDescHandle.Get())
+			if (const FWorldPartitionActorDescInstance* ActorDescInstance = ActorDescItem->ActorDescHandle.GetInstance())
 			{
-				return ActorDesc->GetActorName().ToString();
+				return ActorDescInstance->GetActorName().ToString();
 			}
 		}
 		else if (const FActorFolderTreeItem* ActorFolderItem = Item.CastTo<FActorFolderTreeItem>())
@@ -499,9 +500,9 @@ void FSceneOutlinerModule::CreateActorInfoColumns(FSceneOutlinerInitializationOp
 		}
 		else if (const FActorDescTreeItem* ActorDescItem = Item.CastTo<FActorDescTreeItem>())
 		{
-			if (const FWorldPartitionActorDesc* ActorDesc = ActorDescItem->ActorDescHandle.Get())
+			if (const FWorldPartitionActorDescInstance* ActorDescInstance = ActorDescItem->ActorDescHandle.GetInstance())
 			{
-				return FPackageName::GetShortName(ActorDesc->GetActorPackage());
+				return FPackageName::GetShortName(ActorDescInstance->GetActorPackage());
 			}
 		}
 		else if (const FActorFolderTreeItem* ActorFolderItem = Item.CastTo<FActorFolderTreeItem>())
diff --git a/Engine/Source/Editor/SceneOutliner/Public/ActorBrowsingMode.h b/Engine/Source/Editor/SceneOutliner/Public/ActorBrowsingMode.h
--- a/Engine/Source/Editor/SceneOutliner/Public/ActorBrowsingMode.h
+++ b/Engine/Source/Editor/SceneOutliner/Public/ActorBrowsingMode.h
@@ -3,9 +3,11 @@
 #pragma once
 
 #include "ActorModeInteractive.h"
+#include "WorldPartition/WorldPartitionHandle.h"
 
 class UActorBrowsingModeSettings;
 class IWorldPartitionEditorModule;
+class FWorldPartitionActorDescInstance;
 
 class FActorBrowsingMode : public FActorModeInteractive
 {
@@ -69,7 +71,11 @@ public:
 	void OnSelectUnloadedActors(const TArray<FGuid>& ActorGuids);
 	
 	/** Called when an actor desc is removed */
-	void OnActorDescRemoved(FWorldPartitionActorDesc* InActorDesc);
+	UE_DEPRECATED(5.4, "Use OnActorDescInstanceRemoved")
+	void OnActorDescRemoved(FWorldPartitionActorDesc* InActorDesc) {}
+
+	/** Called when an actor desc instance is removed */
+	void OnActorDescInstanceRemoved(FWorldPartitionActorDescInstance* InActorDescInstance);
 
 	/** Called by engine when edit cut actors begins */
 	void OnEditCutActorsBegin();
@@ -159,7 +165,7 @@ private:
 	/** Number of unloaded actors which have passed through all the filters */
 	uint32 FilteredUnloadedActorCount = 0;
 	/** List of unloaded actors which passed through the regular filters and may or may not have passed the search filter */
-	TSet<const FWorldPartitionActorDesc*> ApplicableUnloadedActors;
+	TSet<FWorldPartitionHandle> ApplicableUnloadedActors;
 	/** List of actors which passed the regular filters and may or may not have passed the search filter */
 	TSet<TWeakObjectPtr<AActor>> ApplicableActors;
 
diff --git a/Engine/Source/Editor/SceneOutliner/Public/ActorDescTreeItem.h b/Engine/Source/Editor/SceneOutliner/Public/ActorDescTreeItem.h
--- a/Engine/Source/Editor/SceneOutliner/Public/ActorDescTreeItem.h
+++ b/Engine/Source/Editor/SceneOutliner/Public/ActorDescTreeItem.h
@@ -8,7 +8,7 @@
 #include "UObject/ObjectKey.h"
 #include "WorldPartition/WorldPartitionHandle.h"
 
-class UActorDescContainer;
+class UActorDescContainerInstance;
 class FWorldPartitionActorDesc;
 class UToolMenu;
 
@@ -16,17 +16,17 @@ class UToolMenu;
 struct SCENEOUTLINER_API FActorDescTreeItem : IActorBaseTreeItem
 {
 public:
-	DECLARE_DELEGATE_RetVal_OneParam(bool, FFilterPredicate, const FWorldPartitionActorDesc*);
-	DECLARE_DELEGATE_RetVal_OneParam(bool, FInteractivePredicate, const FWorldPartitionActorDesc*);
+	DECLARE_DELEGATE_RetVal_OneParam(bool, FFilterPredicate, const FWorldPartitionActorDescInstance*);
+	DECLARE_DELEGATE_RetVal_OneParam(bool, FInteractivePredicate, const FWorldPartitionActorDescInstance*);
 
 	bool Filter(FFilterPredicate Pred) const
 	{
-		return Pred.Execute(ActorDescHandle.Get());
+		return Pred.Execute(ActorDescHandle.GetInstance());
 	}
 
 	bool GetInteractiveState(FInteractivePredicate Pred) const
 	{
-		return Pred.Execute(ActorDescHandle.Get());
+		return Pred.Execute(ActorDescHandle.GetInstance());
 	}
 
 	/** The actor desc this tree item is associated with. */
@@ -35,17 +35,24 @@ public:
 	/** Constant identifier for this tree item */
 	FSceneOutlinerTreeItemID ID;
 
-	static FSceneOutlinerTreeItemID ComputeTreeItemID(FGuid InActorGuid, UActorDescContainer* InContainer);
-	static bool ShouldDisplayInOutliner(const FWorldPartitionActorDesc* ActorDesc);
+	static FSceneOutlinerTreeItemID ComputeTreeItemID(FGuid InActorGuid, UActorDescContainerInstance* InContainer);
+
+	static bool ShouldDisplayInOutliner(const FWorldPartitionActorDescInstance* InActorDescInstance);
+
+	UE_DEPRECATED(5.4, "Use FWorldPartitionActorDescInstance version instead")
+	static bool ShouldDisplayInOutliner(const FWorldPartitionActorDesc* ActorDesc) { return false;}
 
 	/** Static type identifier for this tree item class */
 	static const FSceneOutlinerTreeItemType Type;
 
 	/** Construct this item from an actor desc */
-	FActorDescTreeItem(const FGuid& InActorGuid, UActorDescContainer* InContainer);
+	FActorDescTreeItem(const FGuid& InActorGuid, UActorDescContainerInstance* InContainerInstance);
+
+	UE_DEPRECATED(5.4, "Use FWorldPartitionActorDescInstance version instead")
+	FActorDescTreeItem(const FGuid& InActorGuid, class UActorDescContainer* InContainer) : IActorBaseTreeItem(Type) {}
 
 	/* Begin ISceneOutlinerTreeItem Implementation */
-	virtual bool IsValid() const override { return ActorDescHandle.Get() != nullptr; }
+	virtual bool IsValid() const override { return ActorDescHandle.GetInstance() != nullptr; }
 	virtual FSceneOutlinerTreeItemID GetID() const override;
 	virtual FString GetDisplayString() const override;
 	virtual bool CanInteract() const override;
diff --git a/Engine/Source/Editor/SceneOutliner/Public/ActorHierarchy.h b/Engine/Source/Editor/SceneOutliner/Public/ActorHierarchy.h
--- a/Engine/Source/Editor/SceneOutliner/Public/ActorHierarchy.h
+++ b/Engine/Source/Editor/SceneOutliner/Public/ActorHierarchy.h
@@ -7,6 +7,7 @@
 #include "Folder.h"
 
 class FWorldPartitionActorDesc;
+class FWorldPartitionActorDescInstance;
 class UActorFolder;
 class UWorldPartition;
 
@@ -53,8 +54,14 @@ protected:
 	SCENEOUTLINER_API virtual void OnLoadedActorAdded(AActor& InActor);
 	SCENEOUTLINER_API virtual void OnLoadedActorRemoved(AActor& InActor);
 
-	SCENEOUTLINER_API virtual void OnActorDescAdded(FWorldPartitionActorDesc* ActorDesc);
-	SCENEOUTLINER_API virtual void OnActorDescRemoved(FWorldPartitionActorDesc* ActorDesc);
+	SCENEOUTLINER_API virtual void OnActorDescInstanceAdded(FWorldPartitionActorDescInstance* InActorDescInstance);
+	SCENEOUTLINER_API virtual void OnActorDescInstanceRemoved(FWorldPartitionActorDescInstance* InActorDescInstance);
+
+	UE_DEPRECATED(5.4, "Use OnActorDescInstanceAdded instead")
+	SCENEOUTLINER_API virtual void OnActorDescAdded(FWorldPartitionActorDesc* ActorDesc) {}
+
+	UE_DEPRECATED(5.4, "Use OnActorDescInstanceRemoved instead")
+	SCENEOUTLINER_API virtual void OnActorDescRemoved(FWorldPartitionActorDesc* ActorDesc) {}
 	
 	void OnComponentsUpdated();
 
diff --git a/Engine/Source/Editor/SceneOutliner/Public/ActorMode.h b/Engine/Source/Editor/SceneOutliner/Public/ActorMode.h
--- a/Engine/Source/Editor/SceneOutliner/Public/ActorMode.h
+++ b/Engine/Source/Editor/SceneOutliner/Public/ActorMode.h
@@ -3,6 +3,7 @@
 #pragma once
 
 #include "ISceneOutlinerMode.h"
+#include "WorldPartition/WorldPartitionHandle.h"
 
 namespace SceneOutliner
 {
@@ -18,10 +19,15 @@ namespace SceneOutliner
 		bool operator()(const TWeakPtr<ISceneOutlinerTreeItem>& Item, AActor*& ActorPtrOut) const;
 	};
 
+	struct UE_DEPRECATED(5.4, "Use FActorHandleSelector instead") FActorDescSelector
+	{
+		bool operator()(const TWeakPtr<ISceneOutlinerTreeItem>&Item, FWorldPartitionActorDesc * &ActorDescPtrOut) const { return false; }
+	};
+
 	/** Functor which can be used to get actor descriptors from a selection  */
-	struct SCENEOUTLINER_API FActorDescSelector
+	struct SCENEOUTLINER_API FActorHandleSelector
 	{
-		bool operator()(const TWeakPtr<ISceneOutlinerTreeItem>& Item, FWorldPartitionActorDesc*& ActorDescPtrOut) const;
+		bool operator()(const TWeakPtr<ISceneOutlinerTreeItem>& Item, FWorldPartitionHandle& ActorHandleOut) const;
 	};
 }
 
diff --git a/Engine/Source/Editor/UnrealEd/Classes/Commandlets/BaseIteratePackagesCommandlet.h b/Engine/Source/Editor/UnrealEd/Classes/Commandlets/BaseIteratePackagesCommandlet.h
--- a/Engine/Source/Editor/UnrealEd/Classes/Commandlets/BaseIteratePackagesCommandlet.h
+++ b/Engine/Source/Editor/UnrealEd/Classes/Commandlets/BaseIteratePackagesCommandlet.h
@@ -11,6 +11,8 @@
 #include "WorldPartition/WorldPartitionBuilder.h"
 #include "WorldPartition/WorldPartitionActorDesc.h"
 
+class FWorldPartitionActorDescInstance;
+
 #include "BaseIteratePackagesCommandlet.generated.h"
 
 
@@ -242,7 +244,10 @@ protected:
 
 	virtual void PerformWorldBuilderAdditionalOperations(class AActor* Actor, bool& bSavePackage) { }
 	virtual void PerformWorldBuilderAdditionalOperations(class UObject* Object, bool& bSavePackage) { }
+	
+	UE_DEPRECATED(5.4, "Use WorldBuilderFailedLoadingActor with FWorldPartitionActorDescInstance")
 	virtual void WorldBuilderFailedLoadingActor(const FWorldPartitionActorDesc* ActorDesc) { } 
+	virtual void WorldBuilderFailedLoadingActor(const FWorldPartitionActorDescInstance* ActorDescInstance) { }
 
 	virtual void PostPerformAdditionalOperations(class UPackage* Package) { }
 
diff --git a/Engine/Source/Editor/UnrealEd/Private/Commandlets/BaseIteratePackagesCommandlet.cpp b/Engine/Source/Editor/UnrealEd/Private/Commandlets/BaseIteratePackagesCommandlet.cpp
--- a/Engine/Source/Editor/UnrealEd/Private/Commandlets/BaseIteratePackagesCommandlet.cpp
+++ b/Engine/Source/Editor/UnrealEd/Private/Commandlets/BaseIteratePackagesCommandlet.cpp
@@ -66,6 +66,8 @@ DEFINE_LOG_CATEGORY(LogIteratePackagesCommandlet);
 // world partition includes
 #include "WorldPartition/WorldPartitionHelpers.h"
 #include "WorldPartition/WorldPartitionActorDesc.h"
+#include "WorldPartition/ActorDescContainerInstance.h"
+#include "WorldPartition/WorldPartitionActorDescInstance.h"
 #include "WorldPartition/LoaderAdapter/LoaderAdapterShape.h"
 #include "UObject/GCObjectScopeGuard.h"
 
@@ -715,13 +717,13 @@ void UBaseIteratePackagesCommandlet::LoadAndSaveOnePackage(const FString& Filena
 
 				UWorldPartition* WorldPartition = World->GetWorldPartition();
 
-				FWorldPartitionHelpers::ForEachActorWithLoading(WorldPartition, [&PackagesToSave, this](const FWorldPartitionActorDesc* ActorDesc)
+				FWorldPartitionHelpers::ForEachActorWithLoading(WorldPartition, [&PackagesToSave, this](const FWorldPartitionActorDescInstance* ActorDescInstance)
 					{
-						AActor* Actor = ActorDesc->GetActor();
+						AActor* Actor = ActorDescInstance->GetActor();
 
 						if (!Actor)
 						{
-							WorldBuilderFailedLoadingActor(ActorDesc);
+							WorldBuilderFailedLoadingActor(ActorDescInstance);
 							return true;
 						}
 
diff --git a/Engine/Source/Editor/UnrealEd/Private/Commandlets/ContentCommandlets.cpp b/Engine/Source/Editor/UnrealEd/Private/Commandlets/ContentCommandlets.cpp
--- a/Engine/Source/Editor/UnrealEd/Private/Commandlets/ContentCommandlets.cpp
+++ b/Engine/Source/Editor/UnrealEd/Private/Commandlets/ContentCommandlets.cpp
@@ -60,7 +60,7 @@
 #include "UObject/SavePackage.h"
 #include "UObject/UObjectHash.h"
 #include "UObject/UObjectIterator.h"
-#include "WorldPartition/ActorDescContainer.h"
+#include "WorldPartition/WorldPartitionActorDescInstance.h"
 #include "WorldPartition/WorldPartition.h"
 #include "WorldPartition/WorldPartitionHelpers.h"
 #include "Virtualization/VirtualizationSystem.h"
@@ -1926,12 +1926,12 @@ void UResavePackagesCommandlet::PerformAdditionalOperations(class UWorld* World,
 		// Use a default GC frequency for external actors if GarbageCollectionFrequency is 0.
 		TGuardValue<int32> ScopedGCFreq(GarbageCollectionFrequency, GarbageCollectionFrequency ? GarbageCollectionFrequency : DefaultExternalActorGCFreq);
 
-		FWorldPartitionHelpers::ForEachActorDesc(WorldPartition, [this, WorldPartition](const FWorldPartitionActorDesc* ActorDesc)
+		FWorldPartitionHelpers::ForEachActorDescInstance(WorldPartition, [this, WorldPartition](const FWorldPartitionActorDescInstance* ActorDescInstance)
 		{
 			++TotalPackagesForResave;
 			// Load & Register World Partition Actor
-			FWorldPartitionReference LoadedActor(WorldPartition, ActorDesc->GetGuid());
-			AActor* Actor = LoadedActor->GetActor();
+			FWorldPartitionReference LoadedActor(WorldPartition, ActorDescInstance->GetGuid());
+			AActor* Actor = LoadedActor.GetActor();
 			UPackage* Package = Actor ? Actor->GetExternalPackage() : nullptr;
 			if (Package == nullptr)
 			{
diff --git a/Engine/Source/Editor/UnrealEd/Private/Commandlets/DerivedDataCacheCommandlet.cpp b/Engine/Source/Editor/UnrealEd/Private/Commandlets/DerivedDataCacheCommandlet.cpp
--- a/Engine/Source/Editor/UnrealEd/Private/Commandlets/DerivedDataCacheCommandlet.cpp
+++ b/Engine/Source/Editor/UnrealEd/Private/Commandlets/DerivedDataCacheCommandlet.cpp
@@ -38,6 +38,7 @@ DerivedDataCacheCommandlet.cpp: Commandlet for DDC maintenence
 #include "WorldPartition/WorldPartition.h"
 #include "WorldPartition/WorldPartitionHelpers.h"
 #include "WorldPartition/WorldPartitionSubsystem.h"
+#include "WorldPartition/WorldPartitionActorDescInstance.h"
 DEFINE_LOG_CATEGORY_STATIC(LogDerivedDataCacheCommandlet, Log, All);
 
 class UDerivedDataCacheCommandlet::FObjectReferencer : public FGCObject
@@ -444,9 +445,9 @@ void UDerivedDataCacheCommandlet::CacheWorldPackages(UWorld* World, uint8 Packag
 		UWorldPartition* WorldPartition = World->GetWorldPartition();
 		check(WorldPartition);
 
-		FWorldPartitionHelpers::ForEachActorWithLoading(WorldPartition, [this, PackageFilter, &OutNewProcessedPackages](const FWorldPartitionActorDesc* ActorDesc)
+		FWorldPartitionHelpers::ForEachActorWithLoading(WorldPartition, [this, PackageFilter, &OutNewProcessedPackages](const FWorldPartitionActorDescInstance* ActorDescInstance)
 		{
-			if (AActor* Actor = ActorDesc->GetActor())
+			if (AActor* Actor = ActorDescInstance->GetActor())
 			{
 				UE_LOG(LogDerivedDataCacheCommandlet, Display, TEXT("Loaded actor %s"), *Actor->GetName());
 				CacheLoadedPackages(Actor->GetPackage(), PackageFilter, OutNewProcessedPackages);
diff --git a/Engine/Source/Editor/UnrealEd/Private/Commandlets/GatherTextFromAssetsCommandlet.cpp b/Engine/Source/Editor/UnrealEd/Private/Commandlets/GatherTextFromAssetsCommandlet.cpp
--- a/Engine/Source/Editor/UnrealEd/Private/Commandlets/GatherTextFromAssetsCommandlet.cpp
+++ b/Engine/Source/Editor/UnrealEd/Private/Commandlets/GatherTextFromAssetsCommandlet.cpp
@@ -27,6 +27,7 @@
 #include "WorldPartition/WorldPartition.h"
 #include "WorldPartition/WorldPartitionHelpers.h"
 #include "WorldPartition/WorldPartitionActorDescUtils.h"
+#include "WorldPartition/WorldPartitionActorDescInstance.h"
 #include "Engine/Level.h"
 #include "Engine/World.h"
 #include "Editor.h"
@@ -1125,9 +1126,9 @@ int32 UGatherTextFromAssetsCommandlet::Main(const FString& Params)
 						FWorldPartitionHelpers::FForEachActorWithLoadingParams ForEachActorParams;
 						ForEachActorParams.ActorGuids = PackagePendingGather.ExternalActors.Array();
 
-						FWorldPartitionHelpers::ForEachActorWithLoading(WorldPartition, [&GatherableTextDataArray](const FWorldPartitionActorDesc* ActorDesc)
+						FWorldPartitionHelpers::ForEachActorWithLoading(WorldPartition, [&GatherableTextDataArray](const FWorldPartitionActorDescInstance* ActorDescInstance)
 						{
-							if (const AActor* Actor = ActorDesc->GetActor())
+							if (const AActor* Actor = ActorDescInstance->GetActor())
 							{
 								EPropertyLocalizationGathererResultFlags ActorGatherableTextResultFlags = EPropertyLocalizationGathererResultFlags::Empty;
 								FPropertyLocalizationDataGatherer(GatherableTextDataArray, Actor->GetExternalPackage(), ActorGatherableTextResultFlags);
diff --git a/Engine/Source/Editor/UnrealEd/Private/Commandlets/WorldPartitionDataLayerToAssetCommandlet.cpp b/Engine/Source/Editor/UnrealEd/Private/Commandlets/WorldPartitionDataLayerToAssetCommandlet.cpp
--- a/Engine/Source/Editor/UnrealEd/Private/Commandlets/WorldPartitionDataLayerToAssetCommandlet.cpp
+++ b/Engine/Source/Editor/UnrealEd/Private/Commandlets/WorldPartitionDataLayerToAssetCommandlet.cpp
@@ -17,7 +17,7 @@
 #include "ProfilingDebugging/ScopedTimers.h"
 #include "WorldPartition/DataLayer/DataLayerInstanceWithAsset.h"
 #include "WorldPartition/DataLayer/WorldDataLayers.h"
-#include "WorldPartition/WorldPartitionActorDesc.h"
+#include "WorldPartition/WorldPartitionActorDescInstance.h"
 #include "WorldPartition/WorldPartitionHelpers.h"
 
 DEFINE_LOG_CATEGORY(LogDataLayerToAssetCommandlet);
@@ -481,10 +481,10 @@ bool UDataLayerToAssetCommandlet::RemapActorDataLayersToAssets(TStrongObjectPtr<
 	UE_LOG(LogDataLayerToAssetCommandlet, Log, TEXT("Starting Actor Data Layer Remapping To Data Layer Asset. This can take a while."));
 
 	uint32 ErrorCount = 0;
-	FWorldPartitionHelpers::ForEachActorWithLoading(MainWorld->GetWorldPartition(), [&ErrorCount, &CommandletContext, this, &PackageHelper](const FWorldPartitionActorDesc* ActorDesc)
+	FWorldPartitionHelpers::ForEachActorWithLoading(MainWorld->GetWorldPartition(), [&ErrorCount, &CommandletContext, this, &PackageHelper](const FWorldPartitionActorDescInstance* ActorDescInstance)
 	{
 		uint32 ActorConversionErrors = 0;
-		if (AActor* Actor = ActorDesc->GetActor())
+		if (AActor* Actor = ActorDescInstance->GetActor())
 		{
 			ActorConversionErrors += RemapDataLayersAssetsFromPreviousConversions(CommandletContext, Actor);
 			ActorConversionErrors += RemapActorDataLayers(CommandletContext, Actor);
@@ -507,13 +507,13 @@ bool UDataLayerToAssetCommandlet::RemapActorDataLayersToAssets(TStrongObjectPtr<
 		}
 		else
 		{
-			const TArray<FName>& ActDescDataLayers = ActorDesc->GetDataLayerInstanceNames();
+			const TArray<FName>& ActDescDataLayers = ActorDescInstance->GetDataLayerInstanceNames();
 			if (!ActDescDataLayers.IsEmpty())
 			{
 				FString DataLayerString = FString::JoinBy(ActDescDataLayers, TEXT(", "), [](const FName& DataLayerName) { return DataLayerName.ToString(); });
 
 				UE_LOG(LogDataLayerToAssetCommandlet, Error, TEXT("Actor %s failed to load. Its data layers %s will not be remapped to a data layer asset."),
-					*ActorDesc->GetActorName().ToString(), *DataLayerString);
+					*ActorDescInstance->GetActorName().ToString(), *DataLayerString);
 				if (!bIgnoreActorLoadingErrors)
 				{
 					ActorConversionErrors++;
diff --git a/Engine/Source/Editor/UnrealEd/Private/EditorActorFolders.cpp b/Engine/Source/Editor/UnrealEd/Private/EditorActorFolders.cpp
--- a/Engine/Source/Editor/UnrealEd/Private/EditorActorFolders.cpp
+++ b/Engine/Source/Editor/UnrealEd/Private/EditorActorFolders.cpp
@@ -23,7 +23,8 @@
 #include "Widgets/Text/STextBlock.h"
 #include "WorldPartition/WorldPartitionActorDesc.h"
 #include "WorldPartition/WorldPartitionHelpers.h"
-#include "WorldPartition/ActorDescContainer.h"
+#include "WorldPartition/ActorDescContainerInstance.h"
+#include "WorldPartition/WorldPartitionActorDescInstance.h"
 
 #define LOCTEXT_NAMESPACE "FActorFolders"
 
@@ -647,34 +648,34 @@ void FActorFolders::ForEachFolderWithRootObject(UWorld& InWorld, const FFolder::
 	GetOrCreateWorldFolders(InWorld).ForEachFolderWithRootObject(FolderRootObject, Operation);
 }
 
-FFolder FActorFolders::GetActorDescFolder(UWorld& InWorld, const FWorldPartitionActorDesc* InActorDesc)
+FFolder FActorFolders::GetActorDescInstanceFolder(UWorld& InWorld, const FWorldPartitionActorDescInstance* InActorDescInstance)
 {
-	if (InActorDesc)
+	if (InActorDescInstance)
 	{
-		UWorld& OuterWorld = InActorDesc->GetContainer() ? *InActorDesc->GetContainer()->GetTypedOuter<UWorld>() : InWorld;
+		UWorld& OuterWorld = InActorDescInstance->GetContainerInstance() ? *InActorDescInstance->GetContainerInstance()->GetTypedOuter<UWorld>() : InWorld;
 		{
 			ULevel* OuterLevel = OuterWorld.PersistentLevel;
 			if (OuterLevel->IsUsingActorFolders())
 			{
-				if (UActorFolder* ActorFolder = OuterLevel->GetActorFolder(InActorDesc->GetFolderGuid()))
+				if (UActorFolder* ActorFolder = OuterLevel->GetActorFolder(InActorDescInstance->GetFolderGuid()))
 				{
 					return ActorFolder->GetFolder();
 				}
 				return FFolder::GetWorldRootFolder(&OuterWorld).GetRootObject();
 			}
-			return FFolder(FFolder::GetWorldRootFolder(&OuterWorld).GetRootObject(), InActorDesc->GetFolderPath());
+			return FFolder(FFolder::GetWorldRootFolder(&OuterWorld).GetRootObject(), InActorDescInstance->GetFolderPath());
 		}
 	}
 	return FFolder::GetInvalidFolder();
 }
 
-void FActorFolders::ForEachActorDescInFolders(UWorld& InWorld, const TSet<FName>& InPaths, TFunctionRef<bool(const FWorldPartitionActorDesc*)> Operation, const FFolder::FRootObject& InFolderRootObject /*= FFolder::GetInvalidRootObject()*/)
+void FActorFolders::ForEachActorDescInstanceInFolders(UWorld& InWorld, const TSet<FName>& InPaths, TFunctionRef<bool(const FWorldPartitionActorDescInstance*)> Operation, const FFolder::FRootObject& InFolderRootObject /*= FFolder::GetInvalidRootObject()*/)
 {
 	if (UWorldPartition* WorldPartition = InWorld.GetWorldPartition())
 	{
-		FWorldPartitionHelpers::ForEachActorDesc(WorldPartition, [&](const FWorldPartitionActorDesc* ActorDesc)
+		FWorldPartitionHelpers::ForEachActorDescInstance(WorldPartition, [&](const FWorldPartitionActorDescInstance* ActorDescInstance)
 		{
-			FFolder ActorDescFolder = GetActorDescFolder(InWorld, ActorDesc);
+			FFolder ActorDescFolder = GetActorDescInstanceFolder(InWorld, ActorDescInstance);
 			if (ActorDescFolder == FFolder::GetInvalidFolder())
 			{
 				return true;
@@ -686,7 +687,7 @@ void FActorFolders::ForEachActorDescInFolders(UWorld& InWorld, const TSet<FName>
 				return true;
 			}
 
-			return Operation(ActorDesc);
+			return Operation(ActorDescInstance);
 		});
 	}
 }
diff --git a/Engine/Source/Editor/UnrealEd/Private/WorldPartition/WorldPartitionFoliageBuilder.cpp b/Engine/Source/Editor/UnrealEd/Private/WorldPartition/WorldPartitionFoliageBuilder.cpp
--- a/Engine/Source/Editor/UnrealEd/Private/WorldPartition/WorldPartitionFoliageBuilder.cpp
+++ b/Engine/Source/Editor/UnrealEd/Private/WorldPartition/WorldPartitionFoliageBuilder.cpp
@@ -2,11 +2,12 @@
 
 #include "WorldPartition/WorldPartitionFoliageBuilder.h"
 
-#include "WorldPartition/ActorDescContainer.h"
+#include "WorldPartition/ActorDescContainerInstance.h"
 #include "WorldPartition/ActorPartition/PartitionActorDesc.h"
 #include "WorldPartition/DataLayer/DataLayerEditorContext.h"
 #include "WorldPartition/WorldPartition.h"
 #include "WorldPartition/WorldPartitionHelpers.h"
+#include "WorldPartition/WorldPartitionActorDescInstance.h"
 #include "PackageSourceControlHelper.h"
 #include "SourceControlHelpers.h"
 #include "HAL/PlatformFileManager.h"
@@ -81,9 +82,9 @@ bool UWorldPartitionFoliageBuilder::RunInternal(UWorld* World, const FCellInfo&
 	TArray<FGuid> ExistingActors;
 			
 	// Snapshot existing ActorDescs
-	FWorldPartitionHelpers::ForEachActorDesc(WorldPartition, AInstancedFoliageActor::StaticClass(), [&ExistingActors](const FWorldPartitionActorDesc* ActorDesc)
+	FWorldPartitionHelpers::ForEachActorDescInstance(WorldPartition, AInstancedFoliageActor::StaticClass(), [&ExistingActors](const FWorldPartitionActorDescInstance* ActorDescInstance)
 	{
-		ExistingActors.Add(ActorDesc->GetGuid());
+		ExistingActors.Add(ActorDescInstance->GetGuid());
 		return true;
 	});
 
@@ -104,16 +105,16 @@ bool UWorldPartitionFoliageBuilder::RunInternal(UWorld* World, const FCellInfo&
 	ForEachActorWithLoadingParams.ActorClasses = { AInstancedFoliageActor::StaticClass() };
 	ForEachActorWithLoadingParams.ActorGuids.Append(ExistingActors);
 
-	FWorldPartitionHelpers::ForEachActorWithLoading(WorldPartition, [this, World, WorldPartition, ActorPartitionSubsystem, &SCCHelper, &NewActors, &NumInstances, &NumInstancesProcessed](const FWorldPartitionActorDesc* ActorDesc)
+	FWorldPartitionHelpers::ForEachActorWithLoading(WorldPartition, [this, World, WorldPartition, ActorPartitionSubsystem, &SCCHelper, &NewActors, &NumInstances, &NumInstancesProcessed](const FWorldPartitionActorDescInstance* ActorDescInstance)
 	{
 		TMap<UFoliageType*, TArray<FFoliageInstance>> FoliageToAdd;
 		FBox InstanceBounds(ForceInit);
 
-		FPartitionActorDesc* PartitionActorDesc = (FPartitionActorDesc*)ActorDesc;
-		AInstancedFoliageActor* IFA = Cast<AInstancedFoliageActor>(ActorDesc->GetActor());
+		const FPartitionActorDesc* PartitionActorDesc = (FPartitionActorDesc*)ActorDescInstance->GetActorDesc();
+		AInstancedFoliageActor* IFA = Cast<AInstancedFoliageActor>(ActorDescInstance->GetActor());
 		if (!IFA)
 		{
-			UE_LOG(LogWorldPartitionFoliageBuilder, Error, TEXT("Foliage actor failed to load: %s (%s)"), *ActorDesc->GetActorName().ToString(), *ActorDesc->GetActorPackage().ToString());
+			UE_LOG(LogWorldPartitionFoliageBuilder, Error, TEXT("Foliage actor failed to load: %s (%s)"), *PartitionActorDesc->GetActorName().ToString(), *PartitionActorDesc->GetActorPackage().ToString());
 			return false;
 		}
 
@@ -121,31 +122,31 @@ bool UWorldPartitionFoliageBuilder::RunInternal(UWorld* World, const FCellInfo&
 		{
 			FPackageSourceControlHelper PackageHelper;
 
-			const FDataLayerEditorContext DataLayerEditorContext(World, ActorDesc->GetDataLayerInstanceNames());
-			const FActorPartitionIdentifier ActorPartitionIdentifier(ActorDesc->GetActorNativeClass(), IFA->GetGridGuid(), DataLayerEditorContext.GetHash());
+			const FDataLayerEditorContext DataLayerEditorContext(World, ActorDescInstance->GetDataLayerInstanceNames());
+			const FActorPartitionIdentifier ActorPartitionIdentifier(PartitionActorDesc->GetActorNativeClass(), IFA->GetGridGuid(), DataLayerEditorContext.GetHash());
 			const bool bShouldIncludeGridSizeInName = ActorPartitionIdentifier.GetClass()->GetDefaultObject<APartitionActor>()->ShouldIncludeGridSizeInName(World, ActorPartitionIdentifier);
 			const FString ExpectedActorName = APartitionActor::GetActorName(World, ActorPartitionIdentifier, PartitionActorDesc->GridSize, 
 				PartitionActorDesc->GridIndexX, PartitionActorDesc->GridIndexY, PartitionActorDesc->GridIndexZ);
 
-			if (ActorDesc->GetActorName() != ExpectedActorName)
+			if (PartitionActorDesc->GetActorName() != ExpectedActorName)
 			{
 				UPackage* PackageToModify = nullptr;
 				UPackage* PackageToDelete = nullptr;
 				UPackage* PackageToAdd = nullptr;
 
-				UE_LOG(LogWorldPartitionFoliageBuilder, Display, TEXT("Repairing invalid foliage actor: %s->%s "), *ActorDesc->GetActorName().ToString(), *ExpectedActorName);
+				UE_LOG(LogWorldPartitionFoliageBuilder, Display, TEXT("Repairing invalid foliage actor: %s->%s "), *PartitionActorDesc->GetActorName().ToString(), *ExpectedActorName);
 
 				const FString ExpectedActorPath = FString::Printf(TEXT("%s.%s"), *World->PersistentLevel->GetPathName(), *ExpectedActorName);
 				
-				if (const FWorldPartitionActorDesc* DupActorDesc = WorldPartition->GetActorDescByPath(ExpectedActorPath))
+				if (const FWorldPartitionActorDescInstance* DupActorDescInstance = WorldPartition->GetActorDescInstanceByPath(ExpectedActorPath))
 				{
 					// Merge with existing
-					FWorldPartitionReference DupActorRef(WorldPartition, DupActorDesc->GetGuid());
+					FWorldPartitionReference DupActorRef(WorldPartition, DupActorDescInstance->GetGuid());
 
-					AInstancedFoliageActor* DupIFA = Cast<AInstancedFoliageActor>(DupActorDesc->GetActor());
+					AInstancedFoliageActor* DupIFA = Cast<AInstancedFoliageActor>(DupActorDescInstance->GetActor());
 					if (!DupIFA)
 					{
-						UE_LOG(LogWorldPartitionFoliageBuilder, Error, TEXT("Foliage actor failed to load: %s (%s)"), *DupActorDesc->GetActorName().ToString(), *DupActorDesc->GetActorPackage().ToString());
+						UE_LOG(LogWorldPartitionFoliageBuilder, Error, TEXT("Foliage actor failed to load: %s (%s)"), *DupActorDescInstance->GetActorName().ToString(), *DupActorDescInstance->GetActorPackage().ToString());
 						return false;
 					}
 
@@ -224,7 +225,7 @@ bool UWorldPartitionFoliageBuilder::RunInternal(UWorld* World, const FCellInfo&
 		{	
 			check(IFA->GetGridSize() != NewGridSize);
 			// Harvest Instances from existing IFA and build up instance bounds
-			UE_LOG(LogWorldPartitionFoliageBuilder, Display, TEXT("Processing existing foliage actor: %s (%s)"), *ActorDesc->GetActorName().ToString(), *ActorDesc->GetActorPackage().ToString());
+			UE_LOG(LogWorldPartitionFoliageBuilder, Display, TEXT("Processing existing foliage actor: %s (%s)"), *PartitionActorDesc->GetActorName().ToString(), *PartitionActorDesc->GetActorPackage().ToString());
 		
 			IFA->ForEachFoliageInfo([IFA, &FoliageToAdd, &InstanceBounds, &NumInstances](UFoliageType* FoliageType, FFoliageInfo& FoliageInfo)
 			{
@@ -262,7 +263,7 @@ bool UWorldPartitionFoliageBuilder::RunInternal(UWorld* World, const FCellInfo&
 					const FGuid& NewActorGuid = NewActors.FindChecked(InCellCoord);
 					FWorldPartitionReference Reference(WorldPartition, NewActorGuid);
 					ActorReferences.Add(Reference);
-					IntersectingActors.Add(InCellCoord, CastChecked<AInstancedFoliageActor>(Reference->GetActor()));
+					IntersectingActors.Add(InCellCoord, CastChecked<AInstancedFoliageActor>(Reference.GetActor()));
 				}
 
 				return true;
@@ -323,8 +324,8 @@ bool UWorldPartitionFoliageBuilder::RunInternal(UWorld* World, const FCellInfo&
 					}
 
 					FGuid NewActorGuid = IntersectingActor->GetActorGuid();
-					FWorldPartitionActorDesc* NewActorDesc = WorldPartition->GetActorDesc(NewActorGuid);
-					check(NewActorDesc);
+					FWorldPartitionActorDescInstance* NewActorDescInstance = WorldPartition->GetActorDescInstance(NewActorGuid);
+					check(NewActorDescInstance);
 					NewActors.Add(Pair.Key, NewActorGuid);
 
 					// Add ref so actor gets unloaded 
@@ -366,9 +367,9 @@ bool UWorldPartitionFoliageBuilder::RunInternal(UWorld* World, const FCellInfo&
 		// Delete old IFAs
 		for (const FGuid& ActorGuid : ExistingActors)
 		{
-			FWorldPartitionActorDesc* ActorDesc = WorldPartition->GetActorDesc(ActorGuid);
-			check(ActorDesc);
-			const FString PackageToDelete = ActorDesc->GetActorPackage().ToString();
+			FWorldPartitionActorDescInstance* ActorDescInstance = WorldPartition->GetActorDescInstance(ActorGuid);
+			check(ActorDescInstance);
+			const FString PackageToDelete = ActorDescInstance->GetActorPackage().ToString();
 			if (!SCCHelper.Delete(PackageToDelete))
 			{
 				UE_LOG(LogWorldPartitionFoliageBuilder, Error, TEXT("Error deleting package %s."), *PackageToDelete);
diff --git a/Engine/Source/Editor/UnrealEd/Private/WorldPartition/WorldPartitionHLODsBuilder.cpp b/Engine/Source/Editor/UnrealEd/Private/WorldPartition/WorldPartitionHLODsBuilder.cpp
--- a/Engine/Source/Editor/UnrealEd/Private/WorldPartition/WorldPartitionHLODsBuilder.cpp
+++ b/Engine/Source/Editor/UnrealEd/Private/WorldPartition/WorldPartitionHLODsBuilder.cpp
@@ -23,7 +23,8 @@
 #include "DerivedDataCacheInterface.h"
 
 #include "WorldPartition/DataLayer/DataLayerInstance.h"
-#include "WorldPartition/ActorDescContainer.h"
+#include "WorldPartition/ActorDescContainerInstance.h"
+#include "WorldPartition/WorldPartitionActorDescInstance.h"
 #include "WorldPartition/WorldPartition.h"
 #include "WorldPartition/WorldPartitionHelpers.h"
 #include "WorldPartition/HLOD/HLODActor.h"
@@ -390,10 +391,10 @@ bool UWorldPartitionHLODsBuilder::SetupHLODActors()
 		UE_LOG(LogWorldPartitionHLODsBuilder, Display, TEXT("#### World HLOD actors ####"));
 
 		int32 NumActors = 0;
-		for (FActorDescContainerCollection::TIterator<AWorldPartitionHLOD> HLODIterator(WorldPartition); HLODIterator; ++HLODIterator)
+		for (FActorDescContainerInstanceCollection::TIterator<AWorldPartitionHLOD> HLODIterator(WorldPartition); HLODIterator; ++HLODIterator)
 		{
-			FWorldPartitionActorDesc* HLODActorDesc = *HLODIterator;
-			FString PackageName = HLODActorDesc->GetActorPackage().ToString();
+			FWorldPartitionActorDescInstance* HLODActorDescInstance = *HLODIterator;
+			FString PackageName = HLODActorDescInstance->GetActorPackage().ToString();
 
 			UE_LOG(LogWorldPartitionHLODsBuilder, Display, TEXT("    [%d] %s"), NumActors, *PackageName);
 
@@ -521,9 +522,8 @@ bool UWorldPartitionHLODsBuilder::BuildHLODActors()
 			const FGuid& HLODActorGuid = HLODActorsToBuild[CurrentActor];
 
 			FWorldPartitionReference ActorRef(WorldPartition, HLODActorGuid);
-			FWorldPartitionActorDesc* ActorDesc = ActorRef.Get();
 
-			AWorldPartitionHLOD* HLODActor = CastChecked<AWorldPartitionHLOD>(ActorDesc->GetActor());
+			AWorldPartitionHLOD* HLODActor = CastChecked<AWorldPartitionHLOD>(ActorRef.GetActor());
 
 			UE_LOG(LogWorldPartitionHLODsBuilder, Display, TEXT("[%d / %d] Building HLOD actor %s..."), CurrentActor + 1, HLODActorsToBuild.Num(), *HLODActor->GetActorLabel());
 
@@ -625,11 +625,11 @@ bool UWorldPartitionHLODsBuilder::DeleteHLODActors()
 	};
 
 	TArray<FString> PackagesToDelete;
-	for (FActorDescContainerCollection::TIterator<> ActorDescIterator(WorldPartition); ActorDescIterator; ++ActorDescIterator)
+	for (FActorDescContainerInstanceCollection::TIterator<> Iterator(WorldPartition); Iterator; ++Iterator)
 	{
-		if (HLODActorClasses.FindByPredicate([ActorClass = ActorDescIterator->GetActorNativeClass()](const UClass* HLODClass) { return ActorClass->IsChildOf(HLODClass); }))
+		if (HLODActorClasses.FindByPredicate([ActorClass = Iterator->GetActorNativeClass()](const UClass* HLODClass) { return ActorClass->IsChildOf(HLODClass); }))
 		{
-			FString PackageName = ActorDescIterator->GetActorPackage().ToString();
+			FString PackageName = Iterator->GetActorPackage().ToString();
 			PackagesToDelete.Add(PackageName);
 		}
 	}
@@ -734,21 +734,22 @@ TArray<TArray<FGuid>> UWorldPartitionHLODsBuilder::GetHLODWorkloads(int32 NumWor
 
 	// Build a mapping of 1 HLOD[Level] -> N HLOD[Level - 1]
 	TMap<FGuid, TArray<FGuid>>	HLODParenting;
-	for (FActorDescContainerCollection::TIterator<AWorldPartitionHLOD> HLODIterator(WorldPartition); HLODIterator; ++HLODIterator)
+	for (FActorDescContainerInstanceCollection::TIterator<AWorldPartitionHLOD> HLODIterator(WorldPartition); HLODIterator; ++HLODIterator)
 	{
+		const FHLODActorDesc& HLODActorDesc = *(FHLODActorDesc*)HLODIterator->GetActorDesc();
 		// Filter by HLOD actor
-		if (!HLODActorToBuild.IsNone() && HLODIterator->GetActorLabel() != HLODActorToBuild)
+		if (!HLODActorToBuild.IsNone() && HLODActorDesc.GetActorLabel() != HLODActorToBuild)
 		{
 			continue;
 		}
 
 		// Filter by HLOD layer
-		if (!HLODLayerToBuild.IsNone() && HLODIterator->GetSourceHLODLayer().GetAssetName() != HLODLayerToBuild)
+		if (!HLODLayerToBuild.IsNone() && HLODActorDesc.GetSourceHLODLayer().GetAssetName() != HLODLayerToBuild)
 		{
 			continue;
 		}
 
-		HLODParenting.Add(HLODIterator->GetGuid(), HLODIterator->GetChildHLODActors());
+		HLODParenting.Add(HLODIterator->GetGuid(), HLODActorDesc.GetChildHLODActors());
 	}
 
 	// All child HLODs must be built before their parent HLOD
@@ -820,20 +821,20 @@ bool UWorldPartitionHLODsBuilder::ValidateWorkload(const TArray<FGuid>& Workload
 	// For each HLOD entry in the workload, validate that its children are found before itself
 	for (const FGuid& HLODActorGuid : Workload)
 	{
-		const FWorldPartitionActorDesc* ActorDesc = WorldPartition->GetActorDesc(HLODActorGuid);
-		if(!ActorDesc)
+		const FWorldPartitionActorDescInstance* ActorDescInstance = WorldPartition->GetActorDescInstance(HLODActorGuid);
+		if(!ActorDescInstance)
 		{
 			UE_LOG(LogWorldPartitionHLODsBuilder, Error, TEXT("Unknown actor guid found, your HLOD actors are probably out of date. Run with -SetupHLODs to fix this. Exiting..."));
 			return false;
 		}
 
-		if (!ActorDesc->GetActorNativeClass()->IsChildOf<AWorldPartitionHLOD>())
+		if (!ActorDescInstance->GetActorNativeClass()->IsChildOf<AWorldPartitionHLOD>())
 		{
 			UE_LOG(LogWorldPartitionHLODsBuilder, Error, TEXT("Unexpected actor guid found in HLOD workload, exiting..."));
 			return false;
 		}
 
-		const FHLODActorDesc* HLODActorDesc = static_cast<const FHLODActorDesc*>(ActorDesc);
+		const FHLODActorDesc* HLODActorDesc = static_cast<const FHLODActorDesc*>(ActorDescInstance->GetActorDesc());
 
 		for (const FGuid& ChildHLODActorGuid : HLODActorDesc->GetChildHLODActors())
 		{
@@ -881,13 +882,13 @@ bool UWorldPartitionHLODsBuilder::GenerateBuildManifest(TMap<FString, int32>& Fi
 			if (WorldPartition)
 			{
 				// Track which builder is responsible to handle each actor
-				const FWorldPartitionActorDesc* ActorDesc = WorldPartition->GetActorDesc(ActorGuid);
-				if (!ActorDesc)
+				const FWorldPartitionActorDescInstance* ActorDescInstance = WorldPartition->GetActorDescInstance(ActorGuid);
+				if (!ActorDescInstance)
 				{
 					UE_LOG(LogWorldPartitionHLODsBuilder, Error, TEXT("Invalid actor GUID found while generating the HLOD build manifest, exiting..."));
 					return false;
 				}
-				FString ActorPackageFilename = USourceControlHelpers::PackageFilename(ActorDesc->GetActorPackage().ToString());
+				FString ActorPackageFilename = USourceControlHelpers::PackageFilename(ActorDescInstance->GetActorPackage().ToString());
 				FilesToBuilderMap.Emplace(ActorPackageFilename, BuilderIndex);
 			}
 		}
diff --git a/Engine/Source/Editor/UnrealEd/Private/WorldPartition/WorldPartitionLandscapeSplineMeshesBuilder.cpp b/Engine/Source/Editor/UnrealEd/Private/WorldPartition/WorldPartitionLandscapeSplineMeshesBuilder.cpp
--- a/Engine/Source/Editor/UnrealEd/Private/WorldPartition/WorldPartitionLandscapeSplineMeshesBuilder.cpp
+++ b/Engine/Source/Editor/UnrealEd/Private/WorldPartition/WorldPartitionLandscapeSplineMeshesBuilder.cpp
@@ -29,6 +29,7 @@
 #include "WorldPartition/WorldPartition.h"
 #include "WorldPartition/WorldPartitionActorDesc.h"
 #include "WorldPartition/WorldPartitionHelpers.h"
+#include "WorldPartition/WorldPartitionActorDescInstance.h"
 #include "ActorPartition/ActorPartitionSubsystem.h"
 #include "Rendering/ColorVertexBuffer.h"
 #include "EngineUtils.h"
@@ -349,10 +350,10 @@ bool UWorldPartitionLandscapeSplineMeshesBuilder::RunInternal(UWorld* InWorld, c
 	TSet<ALandscapeSplineMeshesActor*> PreviousGeneratedActors;
 
 	// Preload all ALandscapeSplineMeshesActor and ALandscapeSplineActor
-	FWorldPartitionHelpers::ForEachActorWithLoading(WorldPartition, [&ActorReferences, &PreviousGeneratedActors, WorldPartition](const FWorldPartitionActorDesc* ActorDesc)
+	FWorldPartitionHelpers::ForEachActorWithLoading(WorldPartition, [&ActorReferences, &PreviousGeneratedActors, WorldPartition](const FWorldPartitionActorDescInstance* ActorDescInstance)
 	{
-		ActorReferences.Emplace(WorldPartition, ActorDesc->GetGuid());
-		AActor* Actor = ActorDesc->GetActor();
+		ActorReferences.Emplace(WorldPartition, ActorDescInstance->GetGuid());
+		AActor* Actor = ActorDescInstance->GetActor();
 		if (IsValid(Actor))
 		{
 			if (ALandscapeSplineMeshesActor* LandscapeSplineMeshesActor = Cast<ALandscapeSplineMeshesActor>(Actor))
diff --git a/Engine/Source/Editor/UnrealEd/Private/WorldPartition/WorldPartitionRenameDuplicateBuilder.cpp b/Engine/Source/Editor/UnrealEd/Private/WorldPartition/WorldPartitionRenameDuplicateBuilder.cpp
--- a/Engine/Source/Editor/UnrealEd/Private/WorldPartition/WorldPartitionRenameDuplicateBuilder.cpp
+++ b/Engine/Source/Editor/UnrealEd/Private/WorldPartition/WorldPartitionRenameDuplicateBuilder.cpp
@@ -4,7 +4,7 @@
 #include "WorldPartition/WorldPartitionStreamingGeneration.h"
 #include "WorldPartition/WorldPartition.h"
 #include "WorldPartition/WorldPartitionHelpers.h"
-#include "WorldPartition/WorldPartitionActorDescView.h"
+#include "WorldPartition/WorldPartitionActorDescInstance.h"
 #include "ReferenceCluster.h"
 #include "PackageSourceControlHelper.h"
 #include "SourceControlHelpers.h"
@@ -19,6 +19,8 @@
 #include "UObject/ObjectRedirector.h"
 #include "ProfilingDebugging/ScopedTimers.h"
 #include "HAL/FileManager.h"
+#include "WorldPartition/ActorDescContainerInstanceCollection.h"
+#include "WorldPartition/ActorDescContainerInstance.h"
 #include "WorldPartition/ActorDescContainer.h"
 
 DEFINE_LOG_CATEGORY_STATIC(LogWorldPartitionRenameDuplicateBuilder, All, All);
@@ -146,14 +148,14 @@ bool UWorldPartitionRenameDuplicateBuilder::RunInternal(UWorld* World, const FCe
 
 	if (WorldPartition->GetActorDescContainerCount() > 1)
 	{
-		UE_LOG(LogWorldPartitionRenameDuplicateBuilder, Warning, TEXT("Actors coming from containers outside of %s will not be duplicated"), *WorldPartition->GetActorDescContainer()->GetExternalActorPath());
+		UE_LOG(LogWorldPartitionRenameDuplicateBuilder, Warning, TEXT("Actors coming from containers outside of %s will not be duplicated"), *WorldPartition->GetActorDescContainerInstance()->GetExternalActorPath());
 	}
 
 	// Build actor clusters
 	TArray<TPair<FGuid, TArray<FGuid>>> ActorsWithRefs;
-	for (FActorDescContainerCollection::TIterator<> ActorDescIterator(WorldPartition); ActorDescIterator; ++ActorDescIterator)
+	for (FActorDescContainerInstanceCollection::TIterator<> Iterator(WorldPartition); Iterator; ++Iterator)
 	{
-		ActorsWithRefs.Emplace(ActorDescIterator->GetGuid(), ActorDescIterator->GetReferences());
+		ActorsWithRefs.Emplace(Iterator->GetGuid(), Iterator->GetReferences());
 	}
 
 	TArray<TArray<FGuid>> ActorClusters = GenerateObjectsClusters(ActorsWithRefs);
@@ -259,7 +261,7 @@ bool UWorldPartitionRenameDuplicateBuilder::RunInternal(UWorld* World, const FCe
 				
 				for (const FWorldPartitionReference& ActorReference : ActorReferences)
 				{
-					AActor* Actor = ActorReference->GetActor();
+					AActor* Actor = ActorReference.GetActor();
 					UPackage* PreviousActorPackage = Actor->GetExternalPackage();
 
 					// Rename Actor first so new package gets created
@@ -311,7 +313,7 @@ bool UWorldPartitionRenameDuplicateBuilder::RunInternal(UWorld* World, const FCe
 				// This is to prevent failures when some non-serialized references get taken by loaded actors this makes sure those references will resolve. (example that no longer exists: Landscape SplineHandles)
 				for (FWorldPartitionReference ActorReference : ActorReferences)
 				{
-					AActor* Actor = ActorReference->GetActor();
+					AActor* Actor = ActorReference.GetActor();
 					UPackage* NewActorPackage = Actor->GetExternalPackage();
 					check(Actor);
 					Actor->Rename(nullptr, World->PersistentLevel, REN_NonTransactional | REN_DontCreateRedirectors | REN_ForceNoResetLoaders | REN_DoNotDirty);
@@ -347,8 +349,8 @@ bool UWorldPartitionRenameDuplicateBuilder::RunInternal(UWorld* World, const FCe
 					if (!DuplicatedActorGuids.Contains(ActorGuid))
 					{
 						FWorldPartitionReference ActorReference(WorldPartition, ActorGuid);
-						check(*ActorReference);
-						AActor* Actor = ActorReference->GetActor();
+						check(ActorReference.IsValid());
+						AActor* Actor = ActorReference.GetActor();
 						check(Actor);
 						ActorReferences.Add(ActorReference);
 					}
@@ -356,10 +358,10 @@ bool UWorldPartitionRenameDuplicateBuilder::RunInternal(UWorld* World, const FCe
 					// If we are renaming add package to delete
 					if (bRename)
 					{
-						const FWorldPartitionActorDesc* ActorDesc = WorldPartition->GetActorDesc(ActorGuid);
-						check(ActorDesc);
+						const FWorldPartitionActorDescInstance* ActorDescInstance = WorldPartition->GetActorDescInstance(ActorGuid);
+						check(ActorDescInstance);
 
-						PackagesToDelete.Add(ActorDesc->GetActorPackage().ToString());
+						PackagesToDelete.Add(ActorDescInstance->GetActorPackage().ToString());
 					}
 				}
 						
@@ -393,29 +395,29 @@ bool UWorldPartitionRenameDuplicateBuilder::RunInternal(UWorld* World, const FCe
 		{
 			// Validate results
 			UE_SCOPED_TIMER(TEXT("Validating actors"), LogWorldPartitionRenameDuplicateBuilder, Display);
-			for (FActorDescContainerCollection::TConstIterator<> ActorDescIterator(WorldPartition); ActorDescIterator; ++ActorDescIterator)
+			for (FActorDescContainerInstanceCollection::TConstIterator<> Iterator(WorldPartition); Iterator; ++Iterator)
 			{
-				const FWorldPartitionActorDesc* SourceActorDesc = *ActorDescIterator;
-				FGuid* DuplicatedGuid = DuplicatedActorGuids.Find(SourceActorDesc->GetGuid());
-				const FWorldPartitionActorDesc* NewActorDesc = NewWorld->GetWorldPartition()->GetActorDescContainer()->GetActorDesc(DuplicatedGuid ? *DuplicatedGuid : SourceActorDesc->GetGuid());
-				if (!NewActorDesc)
+				const FWorldPartitionActorDescInstance* SourceActorDescInstance = *Iterator;
+				FGuid* DuplicatedGuid = DuplicatedActorGuids.Find(SourceActorDescInstance->GetGuid());
+				const FWorldPartitionActorDescInstance* NewActorDescInstance = NewWorld->GetWorldPartition()->GetActorDescInstance(DuplicatedGuid ? *DuplicatedGuid : SourceActorDescInstance->GetGuid());
+				if (!NewActorDescInstance)
 				{
-					UE_LOG(LogWorldPartitionRenameDuplicateBuilder, Warning, TEXT("Failed to find source actor for Actor: %s"), *SourceActorDesc->GetActorSoftPath().ToString());
+					UE_LOG(LogWorldPartitionRenameDuplicateBuilder, Warning, TEXT("Failed to find source actor for Actor: %s"), *SourceActorDescInstance->GetActorSoftPath().ToString());
 				}
 				else
 				{
-					if (NewActorDesc->GetReferences().Num() != SourceActorDesc->GetReferences().Num())
+					if (NewActorDescInstance->GetReferences().Num() != SourceActorDescInstance->GetReferences().Num())
 					{
-						UE_LOG(LogWorldPartitionRenameDuplicateBuilder, Warning, TEXT("Actor: %s and Source Actor: %s have mismatching reference count"), *NewActorDesc->GetActorSoftPath().ToString(), *SourceActorDesc->GetActorSoftPath().ToString());
+						UE_LOG(LogWorldPartitionRenameDuplicateBuilder, Warning, TEXT("Actor: %s and Source Actor: %s have mismatching reference count"), *NewActorDescInstance->GetActorSoftPath().ToString(), *SourceActorDescInstance->GetActorSoftPath().ToString());
 					}
 					else
 					{
-						for (const FGuid& ReferenceGuid : SourceActorDesc->GetReferences())
+						for (const FGuid& ReferenceGuid : SourceActorDescInstance->GetReferences())
 						{
 							FGuid* DuplicateReferenceGuid = DuplicatedActorGuids.Find(ReferenceGuid);
-							if (!NewActorDesc->GetReferences().Contains(DuplicateReferenceGuid ? *DuplicateReferenceGuid : ReferenceGuid))
+							if (!NewActorDescInstance->GetReferences().Contains(DuplicateReferenceGuid ? *DuplicateReferenceGuid : ReferenceGuid))
 							{
-								UE_LOG(LogWorldPartitionRenameDuplicateBuilder, Warning, TEXT("Actor: %s and Source Actor: %s have mismatching reference"), *NewActorDesc->GetActorSoftPath().ToString(), *SourceActorDesc->GetActorSoftPath().ToString());
+								UE_LOG(LogWorldPartitionRenameDuplicateBuilder, Warning, TEXT("Actor: %s and Source Actor: %s have mismatching reference"), *NewActorDescInstance->GetActorSoftPath().ToString(), *SourceActorDescInstance->GetActorSoftPath().ToString());
 							}
 						}
 					}
diff --git a/Engine/Source/Editor/UnrealEd/Private/WorldPartition/WorldPartitionResaveActorsBuilder.cpp b/Engine/Source/Editor/UnrealEd/Private/WorldPartition/WorldPartitionResaveActorsBuilder.cpp
--- a/Engine/Source/Editor/UnrealEd/Private/WorldPartition/WorldPartitionResaveActorsBuilder.cpp
+++ b/Engine/Source/Editor/UnrealEd/Private/WorldPartition/WorldPartitionResaveActorsBuilder.cpp
@@ -4,6 +4,7 @@
 #include "WorldPartition/WorldPartition.h"
 #include "WorldPartition/WorldPartitionHelpers.h"
 #include "WorldPartition/WorldPartitionClassDescRegistry.h"
+#include "WorldPartition/WorldPartitionActorDescInstance.h"
 #include "PackageSourceControlHelper.h"
 #include "SourceControlHelpers.h"
 #include "Engine/Level.h"
@@ -241,9 +242,9 @@ bool UWorldPartitionResaveActorsBuilder::RunInternal(UWorld* World, const FCellI
 
 		// Build clusters
 		TArray<TPair<FGuid, TArray<FGuid>>> ActorsWithRefs;
-		for (FActorDescContainerCollection::TIterator<> ActorDescIterator(WorldPartition); ActorDescIterator; ++ActorDescIterator)
+		for (FActorDescContainerInstanceCollection::TIterator<> Iterator(WorldPartition); Iterator; ++Iterator)
 		{
-			ActorsWithRefs.Emplace(ActorDescIterator->GetGuid(), ActorDescIterator->GetReferences());
+			ActorsWithRefs.Emplace(Iterator->GetGuid(), Iterator->GetReferences());
 		}
 
 		TArray<TArray<FGuid>> Clusters = GenerateObjectsClusters(ActorsWithRefs);
@@ -264,12 +265,12 @@ bool UWorldPartitionResaveActorsBuilder::RunInternal(UWorld* World, const FCellI
 			// Change packaging of all actors in the current cluster
 			for (FWorldPartitionReference& ActorReference : ActorReferences)
 			{
-				const FWorldPartitionActorDesc* ActorDesc = ActorReference.Get();
-				AActor* Actor = ActorDesc->GetActor();
+				const FWorldPartitionActorDescInstance* ActorDescInstance = ActorReference.GetInstance();
+				AActor* Actor = ActorReference.GetActor();
 
 				if (!Actor)
 				{
-					PackagesToDelete.Add(ActorDesc->GetActorPackage().ToString());
+					PackagesToDelete.Add(ActorDescInstance->GetActorPackage().ToString());
 					continue;
 				}
 				else
@@ -280,7 +281,7 @@ bool UWorldPartitionResaveActorsBuilder::RunInternal(UWorld* World, const FCellI
 					if (!bReportOnly)
 					{
 						// Always mark this package for deletion, as it will contain a temporary redirector to fixup references
-						PackagesToDelete.Add(ActorDesc->GetActorPackage().ToString());
+						PackagesToDelete.Add(ActorDescInstance->GetActorPackage().ToString());
 
 						// Move actor back into world's package
 						Actor->SetPackageExternal(false);
@@ -430,13 +431,13 @@ bool UWorldPartitionResaveActorsBuilder::RunInternal(UWorld* World, const FCellI
 			PackagesToSave.Empty();
 		};
 
-		FWorldPartitionHelpers::ForEachActorWithLoading(WorldPartition, [this, &PackagesToDelete, &PackagesToSave, &DirtyActorDescsOld, &DirtyActorDescsNew](const FWorldPartitionActorDesc* ActorDesc)
+		FWorldPartitionHelpers::ForEachActorWithLoading(WorldPartition, [this, &PackagesToDelete, &PackagesToSave, &DirtyActorDescsOld, &DirtyActorDescsNew](const FWorldPartitionActorDescInstance* ActorDescInstance)
 		{
-			AActor* Actor = ActorDesc->GetActor();
-
+			AActor* Actor = ActorDescInstance->GetActor();
+			
 			if (!Actor)
 			{
-				PackagesToDelete.Add(ActorDesc->GetActorPackage().ToString());
+				PackagesToDelete.Add(ActorDescInstance->GetActorPackage().ToString());
 			}
 			else
 			{
@@ -454,6 +455,7 @@ bool UWorldPartitionResaveActorsBuilder::RunInternal(UWorld* World, const FCellI
 
 				if (bResaveDirtyActorDescsOnly)
 				{
+					const FWorldPartitionActorDesc* ActorDesc = ActorDescInstance->GetActorDesc();
 					TUniquePtr<FWorldPartitionActorDesc> NewActorDesc = Actor->CreateActorDesc();
 
 					if (!ActorDesc->IsResaveNeeded() && !ActorDesc->ShouldResave(NewActorDesc.Get()))
diff --git a/Engine/Source/Editor/UnrealEd/Private/WorldPartition/WorldPartitionRuntimeVirtualTextureBuilder.cpp b/Engine/Source/Editor/UnrealEd/Private/WorldPartition/WorldPartitionRuntimeVirtualTextureBuilder.cpp
--- a/Engine/Source/Editor/UnrealEd/Private/WorldPartition/WorldPartitionRuntimeVirtualTextureBuilder.cpp
+++ b/Engine/Source/Editor/UnrealEd/Private/WorldPartition/WorldPartitionRuntimeVirtualTextureBuilder.cpp
@@ -5,6 +5,7 @@
 #include "WorldPartition/WorldPartitionActorDesc.h"
 #include "WorldPartition/WorldPartitionHelpers.h"
 #include "WorldPartition/WorldPartitionHandle.h"
+#include "WorldPartition/WorldPartitionActorDescInstance.h"
 #include "Components/RuntimeVirtualTextureComponent.h"
 #include "Components/PrimitiveComponent.h"
 #include "VT/VirtualTextureBuilder.h"
@@ -33,7 +34,7 @@ void UWorldPartitionRuntimeVirtualTextureBuilder::LoadRuntimeVirtualTextureActor
 	ForEachActorWithLoadingParams.FilterActorDesc = [](const FWorldPartitionActorDesc* ActorDesc) -> bool { return ActorDesc->HasProperty(UPrimitiveComponent::RVTActorDescProperty); };
 		
 	// @todo: in order to scale the RVTs should be generated with tiling so that we don't need to load all actors writing to RVTs at once.
-	FWorldPartitionHelpers::ForEachActorWithLoading(WorldPartition, [](const FWorldPartitionActorDesc*) { return true; }, ForEachActorWithLoadingParams, Result);
+	FWorldPartitionHelpers::ForEachActorWithLoading(WorldPartition, [](const FWorldPartitionActorDescInstance*) { return true; }, ForEachActorWithLoadingParams, Result);
 
 	// Make sure all assets are finished compiling
 	FAssetCompilingManager::Get().FinishAllCompilation();
diff --git a/Engine/Source/Editor/UnrealEd/Public/EditorActorFolders.h b/Engine/Source/Editor/UnrealEd/Public/EditorActorFolders.h
--- a/Engine/Source/Editor/UnrealEd/Public/EditorActorFolders.h
+++ b/Engine/Source/Editor/UnrealEd/Public/EditorActorFolders.h
@@ -14,6 +14,7 @@ class FObjectPostSaveContext;
 class AActor;
 class UActorFolder;
 class FWorldPartitionActorDesc;
+class FWorldPartitionActorDescInstance;
 
 /** Multicast delegates for broadcasting various folder events */
 
@@ -88,12 +89,18 @@ struct FActorFolders : public FGCObject, public IActorEditorContextClient
 	UE_DEPRECATED(5.0, "RenameFolderInWorld using FName  has been deprecated. Please use new interface using FFolder.")
 	UNREALED_API bool RenameFolderInWorld(UWorld& World, FName OldPath, FName NewPath);
 
+	UE_DEPRECATED(5.4, "Please use ForEachActorDescInstanceInFolders")
+	static UNREALED_API void ForEachActorDescInFolders(UWorld& InWorld, const TSet<FName>& InPaths, TFunctionRef<bool(const FWorldPartitionActorDesc*)> Operation, const FFolder::FRootObject& InFolderRootObject = FFolder::GetInvalidRootObject()) {}
+	
+	UE_DEPRECATED(5.4, "Please use GetActorDescInstanceFolder")
+	static UNREALED_API FFolder GetActorDescFolder(UWorld& InWorld, const FWorldPartitionActorDesc* InActorDesc) { return FFolder::GetInvalidFolder(); }
 	//~ End Deprecated
-	static UNREALED_API FFolder GetActorDescFolder(UWorld& InWorld, const FWorldPartitionActorDesc* InActorDesc);
-
+	
+	static UNREALED_API FFolder GetActorDescInstanceFolder(UWorld& InWorld, const FWorldPartitionActorDescInstance* InActorDescInstance);
 	/** Apply an operation to each actor desc in the given list of folders. */
-	static UNREALED_API void ForEachActorDescInFolders(UWorld& InWorld, const TSet<FName>& InPaths, TFunctionRef<bool(const FWorldPartitionActorDesc*)> Operation, const FFolder::FRootObject& InFolderRootObject = FFolder::GetInvalidRootObject());
-		
+			
+	static UNREALED_API void ForEachActorDescInstanceInFolders(UWorld& InWorld, const TSet<FName>& InPaths, TFunctionRef<bool(const FWorldPartitionActorDescInstance*)> Operation, const FFolder::FRootObject& InFolderRootObject = FFolder::GetInvalidRootObject());
+
 	/** Apply an operation to each actor in the given list of folders. Will stop when operation returns false. */
 	static UNREALED_API void ForEachActorInFolders(UWorld& InWorld, const TArray<FName>& InPaths, TFunctionRef<bool(AActor*)> Operation, const FFolder::FRootObject& InFolderRootObject = FFolder::GetInvalidRootObject());
 	static UNREALED_API void ForEachActorInFolders(UWorld& InWorld, const TSet<FName>& InPaths, TFunctionRef<bool(AActor*)> Operation, const FFolder::FRootObject& InFolderRootObject = FFolder::GetInvalidRootObject());
diff --git a/Engine/Source/Editor/WorldPartitionEditor/Private/WorldPartition/HLOD/HLODEditorData.cpp b/Engine/Source/Editor/WorldPartitionEditor/Private/WorldPartition/HLOD/HLODEditorData.cpp
--- a/Engine/Source/Editor/WorldPartitionEditor/Private/WorldPartition/HLOD/HLODEditorData.cpp
+++ b/Engine/Source/Editor/WorldPartitionEditor/Private/WorldPartition/HLOD/HLODEditorData.cpp
@@ -6,20 +6,20 @@
 #include "WorldPartition/HLOD/HLODActorDesc.h"
 #include "WorldPartition/HLOD/HLODLoaderAdapter.h"
 
-#include "WorldPartition/ActorDescContainerCollection.h"
+#include "WorldPartition/ActorDescContainerInstanceCollection.h"
+#include "WorldPartition/WorldPartitionActorDescInstance.h"
 #include "WorldPartition/WorldPartition.h"
 #include "WorldPartition/WorldPartitionHandle.h"
 #include "WorldPartition/WorldPartitionHelpers.h"
 
-
 FWorldPartitionHLODEditorData::FWorldPartitionHLODEditorData(UWorldPartition* InWorldPartition)
 	: WorldPartition(InWorldPartition)
 	, LastStateUpdate(INDEX_NONE)
 {
 	// First pass, build mapping of GUID to HLOD scene node
-	for (FActorDescContainerCollection::TIterator<AWorldPartitionHLOD> HLODIterator(WorldPartition); HLODIterator; ++HLODIterator)
+	for (FActorDescContainerInstanceCollection::TIterator<AWorldPartitionHLOD> HLODIterator(WorldPartition); HLODIterator; ++HLODIterator)
 	{
-		const FHLODActorDesc& HLODActorDesc = **HLODIterator;
+		const FHLODActorDesc& HLODActorDesc = *(FHLODActorDesc*)HLODIterator->GetActorDesc();
 		const FGuid& HLODActorGuid = HLODActorDesc.GetGuid();
 
 		TUniquePtr<FHLODSceneNode>& HLODSceneNode = HLODActorNodes.Emplace(HLODActorGuid, new FHLODSceneNode());
@@ -28,9 +28,9 @@ FWorldPartitionHLODEditorData::FWorldPartitionHLODEditorData(UWorldPartition* In
 	}
 
 	// Second pass, build a hierarchy now that nodes were all created
-	for (FActorDescContainerCollection::TIterator<AWorldPartitionHLOD> HLODIterator(WorldPartition); HLODIterator; ++HLODIterator)
+	for (FActorDescContainerInstanceCollection::TIterator<AWorldPartitionHLOD> HLODIterator(WorldPartition); HLODIterator; ++HLODIterator)
 	{
-		const FHLODActorDesc& HLODActorDesc = **HLODIterator;
+		const FHLODActorDesc& HLODActorDesc = *(FHLODActorDesc*)HLODIterator->GetActorDesc();
 		const FGuid& HLODActorGuid = HLODActorDesc.GetGuid();
 		TUniquePtr<FHLODSceneNode>& HLODSceneNode = HLODActorNodes.FindChecked(HLODActorGuid);
 
@@ -69,36 +69,36 @@ struct FBoundsWithVolume
 };
 
 // Gather Pinned Actors bounds - also include references & contained actors if the pinned actor is a container.
-static void GatherLoadedActorsBounds(TArray<FBoundsWithVolume>& OutLoadedBounds, const FWorldPartitionActorDesc* InActorDesc, const UActorDescContainer* InContainer, const TOptional<FTransform>& InContainerTransform = TOptional<FTransform>())
+static void GatherLoadedActorsBounds(TArray<FBoundsWithVolume>& OutLoadedBounds, const FWorldPartitionActorDescInstance* InActorDescInstance, const UActorDescContainerInstance* InContainerInstance, const TOptional<FTransform>& InContainerTransform = TOptional<FTransform>())
 {
-	if (InActorDesc)
+	if (InActorDescInstance)
 	{
 		// The actor itself - Include the actor bounds only if HLOD relevant
-		if (InActorDesc->GetIsSpatiallyLoaded() && InActorDesc->IsEditorRelevant() && InActorDesc->GetActorIsHLODRelevant())
+		if (InActorDescInstance->GetIsSpatiallyLoaded() && InActorDescInstance->GetActorIsHLODRelevant() && InActorDescInstance->IsEditorRelevant())
 		{
-			const FBox Box = InContainerTransform.IsSet() ? InActorDesc->GetEditorBounds().TransformBy(InContainerTransform.GetValue()) : InActorDesc->GetEditorBounds();
+			const FBox Box = InContainerTransform.IsSet() ? InActorDescInstance->GetEditorBounds().TransformBy(InContainerTransform.GetValue()) : InActorDescInstance->GetEditorBounds();
 			OutLoadedBounds.Emplace(Box);
 		}
 
 		// Test its references
-		for (const FGuid& ReferenceGuid : InActorDesc->GetReferences())
+		for (const FGuid& ReferenceGuid : InActorDescInstance->GetReferences())
 		{
-			if (const FWorldPartitionActorDesc* ReferenceActorDesc = InContainer->GetActorDesc(ReferenceGuid))
+			if (const FWorldPartitionActorDescInstance* ReferenceActorDescInstance = InContainerInstance->GetActorDescInstance(ReferenceGuid))
 			{
-				GatherLoadedActorsBounds(OutLoadedBounds, ReferenceActorDesc, InContainer, InContainerTransform);
+				GatherLoadedActorsBounds(OutLoadedBounds, ReferenceActorDescInstance, InContainerInstance, InContainerTransform);
 			}
 		}
 
 		// If it's a container, test the contained actors
-		if (InActorDesc->IsContainerInstance())
+		if (InActorDescInstance->IsChildContainerInstance())
 		{
 			FWorldPartitionActorDesc::FContainerInstance ContainerInstance;
-			if (InActorDesc->GetContainerInstance(ContainerInstance))
+			if (InActorDescInstance->GetChildContainerInstance(ContainerInstance))
 			{
 				FTransform ContainerWorldSpaceTransform = InContainerTransform.IsSet() ? ContainerInstance.Transform * InContainerTransform.GetValue() : ContainerInstance.Transform;
-				for (FActorDescList::TConstIterator<> ActorDescIt(ContainerInstance.Container); ActorDescIt; ++ActorDescIt)
+				for (UActorDescContainerInstance::TConstIterator<> Iterator(ContainerInstance.ContainerInstance); Iterator; ++Iterator)
 				{
-					GatherLoadedActorsBounds(OutLoadedBounds, *ActorDescIt, ContainerInstance.Container, ContainerWorldSpaceTransform);
+					GatherLoadedActorsBounds(OutLoadedBounds, *Iterator, ContainerInstance.ContainerInstance, ContainerWorldSpaceTransform);
 				}
 			}
 		}
@@ -133,17 +133,17 @@ void FWorldPartitionHLODEditorData::UpdateLoadedActorsState()
 	}
 
 	// Gather ActorLoaderInterface
-	FWorldPartitionHelpers::ForEachActorDesc(WorldPartition, [&LoadedBounds](const FWorldPartitionActorDesc* ActorDesc)
+	FWorldPartitionHelpers::ForEachActorDescInstance(WorldPartition, [&LoadedBounds](const FWorldPartitionActorDescInstance* ActorDescInstance)
 	{
-		if (ActorDesc->GetActorNativeClass()->ImplementsInterface(UWorldPartitionActorLoaderInterface::StaticClass()))
+		if (ActorDescInstance->GetActorNativeClass()->ImplementsInterface(UWorldPartitionActorLoaderInterface::StaticClass()))
 		{
-			if (AActor* Actor = ActorDesc->GetActor(false))
+			if (AActor* Actor = ActorDescInstance->GetActor(false))
 			{
 				if (IWorldPartitionActorLoaderInterface::ILoaderAdapter* LoaderAdapter = Cast<IWorldPartitionActorLoaderInterface>(Actor)->GetLoaderAdapter())
 				{
 					if (LoaderAdapter->IsLoaded() && LoaderAdapter->GetBoundingBox().IsSet())
 					{
-						LoadedBounds.Emplace(ActorDesc->GetEditorBounds());
+						LoadedBounds.Emplace(ActorDescInstance->GetEditorBounds());
 					}
 				}
 			}
@@ -159,7 +159,7 @@ void FWorldPartitionHLODEditorData::UpdateLoadedActorsState()
 		{
 			if (PinnedActor.IsValid())
 			{
-				GatherLoadedActorsBounds(LoadedBounds, *PinnedActor, WorldPartition->GetActorDescContainer());
+				GatherLoadedActorsBounds(LoadedBounds, PinnedActor.GetInstance(), WorldPartition->GetActorDescContainerInstance());
 			}
 		}
 	}
@@ -222,7 +222,7 @@ void FHLODSceneNode::UpdateVisibility(const FVector& InCameraLocation, double In
 {
 	bool bNodeShouldBeVisible = !bInForceHidden;
 
-	if (AWorldPartitionHLOD* HLODActor = Cast<AWorldPartitionHLOD>(HLODActorHandle->GetActor()))
+	if (AWorldPartitionHLOD* HLODActor = Cast<AWorldPartitionHLOD>(HLODActorHandle.GetActor()))
 	{
 		// Do not show node if there are loaded actors under it
 		if (bNodeShouldBeVisible)
diff --git a/Engine/Source/Editor/WorldPartitionEditor/Private/WorldPartition/HLOD/HLODLoaderAdapter.cpp b/Engine/Source/Editor/WorldPartitionEditor/Private/WorldPartition/HLOD/HLODLoaderAdapter.cpp
--- a/Engine/Source/Editor/WorldPartitionEditor/Private/WorldPartition/HLOD/HLODLoaderAdapter.cpp
+++ b/Engine/Source/Editor/WorldPartitionEditor/Private/WorldPartition/HLOD/HLODLoaderAdapter.cpp
@@ -8,10 +8,10 @@
 #include "WorldPartition/HLOD/HLODActorDesc.h"
 #include "WorldPartition/HLOD/HLODLayer.h"
 
-#include "WorldPartition/ActorDescContainerCollection.h"
+#include "WorldPartition/ActorDescContainerInstanceCollection.h"
 #include "WorldPartition/WorldPartition.h"
 #include "WorldPartition/WorldPartitionHandle.h"
-
+#include "WorldPartition/WorldPartitionActorDescInstance.h"
 
 
 FLoaderAdapterHLOD::FLoaderAdapterHLOD(UWorld* InWorld)
@@ -19,9 +19,9 @@ FLoaderAdapterHLOD::FLoaderAdapterHLOD(UWorld* InWorld)
 {
 	UWorldPartition* WorldPartition = InWorld->GetWorldPartition();
 
-	for (FActorDescContainerCollection::TIterator<AWorldPartitionHLOD> HLODIterator(WorldPartition); HLODIterator; ++HLODIterator)
+	for (FActorDescContainerInstanceCollection::TIterator<AWorldPartitionHLOD> HLODIterator(WorldPartition); HLODIterator; ++HLODIterator)
 	{
-		const FHLODActorDesc& HLODActorDesc = **HLODIterator;
+		const FHLODActorDesc& HLODActorDesc = *(FHLODActorDesc*)HLODIterator->GetActorDesc();
 		if (ShouldLoadHLOD(HLODActorDesc))
 		{
 			Actors.Add(FWorldPartitionHandle(WorldPartition, HLODActorDesc.GetGuid()));
