Scene Outliner: Add an option to mark a folder as the Current Folder on double click instead of expanding it

diff --git a/Engine/Source/Editor/SceneOutliner/Private/ActorBrowsingMode.cpp b/Engine/Source/Editor/SceneOutliner/Private/ActorBrowsingMode.cpp
--- a/Engine/Source/Editor/SceneOutliner/Private/ActorBrowsingMode.cpp
+++ b/Engine/Source/Editor/SceneOutliner/Private/ActorBrowsingMode.cpp
@@ -527,6 +527,23 @@ void FActorBrowsingMode::InitializeViewMenuExtender(TSharedPtr<FExtender> Extend
 			EUserInterfaceActionType::ToggleButton
 		);
 
+		MenuBuilder.AddMenuEntry(
+			LOCTEXT("FolderDoubleClickToggleCurrentFolderLabel", "Double Click toggles Current Folder"),
+			LOCTEXT(
+				"FolderDoubleClickToggleCurrentFolderTooltip",
+				"When enabled, double clicking on a folder will result in it toggling its Current Folder state"
+				"instead of its expansion."
+			),
+			FSlateIcon(),
+			FUIAction(
+				FExecuteAction::CreateRaw(this, &FActorBrowsingMode::OnToggleFolderDoubleClickMarkCurrentFolder),
+				FCanExecuteAction(),
+				FIsActionChecked::CreateRaw(this, &FActorBrowsingMode::DoesFolderDoubleClickMarkCurrentFolder)
+			),
+			NAME_None,
+			EUserInterfaceActionType::ToggleButton
+		);
+
 		MenuBuilder.EndSection();
 
 		MenuBuilder.BeginSection("World", LOCTEXT("ShowWorldHeading", "World"));
@@ -1137,6 +1154,60 @@ void FActorBrowsingMode::OnItemDoubleClick(FSceneOutlinerTreeItemPtr Item)
 	{
 		ActorDescItem->FocusActorBounds();
 	}
+	else if (const FActorFolderTreeItem* FolderItem = Item->CastTo<FActorFolderTreeItem>())
+	{
+		if (DoesFolderDoubleClickMarkCurrentFolder())
+		{
+			if (UWorld* World = FolderItem->World.Get())
+			{
+				const FScopedTransaction Transaction(
+					LOCTEXT("ToggleCurrentActorFolder", "Toggle Current Actor Folder")
+				);
+
+				const FFolder CurrentContextFolder = FActorFolders::Get().GetActorEditorContextFolder(*World);
+				if (CurrentContextFolder == FolderItem->GetFolder())
+				{
+					FActorFolders::Get().SetActorEditorContextFolder(*World, FFolder::GetWorldRootFolder(World));
+				}
+				else
+				{
+					FActorFolders::Get().SetActorEditorContextFolder(*World, FolderItem->GetFolder());
+				}
+			}
+		}
+	}
+}
+
+bool FActorBrowsingMode::HasCustomFolderDoubleClick() const
+{
+	return DoesFolderDoubleClickMarkCurrentFolder();
+}
+
+void FActorBrowsingMode::OnToggleFolderDoubleClickMarkCurrentFolder()
+{
+	if (FActorBrowsingModeConfig* Settings = GetMutableConfig())
+	{
+		if (Settings->FolderDoubleClickMethod == EActorBrowsingFolderDoubleClickMethod::ToggleCurrentFolder)
+		{
+			Settings->FolderDoubleClickMethod = EActorBrowsingFolderDoubleClickMethod::ToggleExpansion;
+		}
+		else
+		{
+			Settings->FolderDoubleClickMethod = EActorBrowsingFolderDoubleClickMethod::ToggleCurrentFolder;
+		}
+
+		SaveConfig();
+	}
+}
+
+bool FActorBrowsingMode::DoesFolderDoubleClickMarkCurrentFolder() const
+{
+	if (const FActorBrowsingModeConfig* Settings = GetConstConfig())
+	{
+		return Settings->FolderDoubleClickMethod == EActorBrowsingFolderDoubleClickMethod::ToggleCurrentFolder;
+	}
+
+	return false;
 }
 
 void FActorBrowsingMode::OnFilterTextCommited(FSceneOutlinerItemSelection& Selection, ETextCommit::Type CommitType)
diff --git a/Engine/Source/Editor/SceneOutliner/Private/ActorBrowsingModeSettings.h b/Engine/Source/Editor/SceneOutliner/Private/ActorBrowsingModeSettings.h
--- a/Engine/Source/Editor/SceneOutliner/Private/ActorBrowsingModeSettings.h
+++ b/Engine/Source/Editor/SceneOutliner/Private/ActorBrowsingModeSettings.h
@@ -12,6 +12,16 @@
 class FName;
 class UObject;
 
+UENUM()
+enum class EActorBrowsingFolderDoubleClickMethod : uint8
+{
+	/** Expands / Collapses the folder */
+	ToggleExpansion,
+
+	/** Marks / unmarks the folder as the Current Folder */
+	ToggleCurrentFolder,
+};
+
 USTRUCT()
 struct FActorBrowsingModeConfig
 {
@@ -58,6 +68,10 @@ public:
 	/** True when the Scene Outliner updates when an actor is selected in the viewport */
 	UPROPERTY()
 	bool bAlwaysFrameSelection = true;
+
+	/** Specifies the behavior of double click on a folder */
+	UPROPERTY()
+	EActorBrowsingFolderDoubleClickMethod FolderDoubleClickMethod = EActorBrowsingFolderDoubleClickMethod::ToggleExpansion;
 };
 
 UCLASS(EditorConfig="ActorBrowsingMode")
diff --git a/Engine/Source/Editor/SceneOutliner/Private/SSceneOutliner.cpp b/Engine/Source/Editor/SceneOutliner/Private/SSceneOutliner.cpp
--- a/Engine/Source/Editor/SceneOutliner/Private/SSceneOutliner.cpp
+++ b/Engine/Source/Editor/SceneOutliner/Private/SSceneOutliner.cpp
@@ -1945,9 +1945,12 @@ void SSceneOutliner::OnOutlinerTreeSelectionChanged( FSceneOutlinerTreeItemPtr T
 
 void SSceneOutliner::OnOutlinerTreeDoubleClick( FSceneOutlinerTreeItemPtr TreeItem )
 {
-	if (TreeItem->IsA<FFolderTreeItem>())
+	if (!Mode->HasCustomFolderDoubleClick())
 	{
-		SetItemExpansion(TreeItem, !IsItemExpanded(TreeItem));
+		if (TreeItem->IsA<FFolderTreeItem>())
+		{
+			SetItemExpansion(TreeItem, !IsItemExpanded(TreeItem));
+		}
 	}
 
 	Mode->OnItemDoubleClick(TreeItem);
diff --git a/Engine/Source/Editor/SceneOutliner/Public/ActorBrowsingMode.h b/Engine/Source/Editor/SceneOutliner/Public/ActorBrowsingMode.h
--- a/Engine/Source/Editor/SceneOutliner/Public/ActorBrowsingMode.h
+++ b/Engine/Source/Editor/SceneOutliner/Public/ActorBrowsingMode.h
@@ -28,6 +28,7 @@ public:
 	SCENEOUTLINER_API virtual void OnItemPassesFilters(const ISceneOutlinerTreeItem& Item) override;
 	SCENEOUTLINER_API virtual FReply OnKeyDown(const FKeyEvent& InKeyEvent) override;
 	SCENEOUTLINER_API virtual void OnDuplicateSelected() override;
+	SCENEOUTLINER_API virtual bool HasCustomFolderDoubleClick() const override;
 	SCENEOUTLINER_API virtual bool CanRenameItem(const ISceneOutlinerTreeItem& Item) const override;
 	SCENEOUTLINER_API virtual FText GetStatusText() const override;
 	SCENEOUTLINER_API virtual FSlateColor GetStatusTextColor() const override;
@@ -140,6 +141,9 @@ private:
 	void OnToggleHideLevelInstanceHierarchy();
 	bool ShouldHideLevelInstanceHierarchy() const;
 
+	void OnToggleFolderDoubleClickMarkCurrentFolder();
+	bool DoesFolderDoubleClickMarkCurrentFolder() const;
+
 	/**
 	 * Get a mutable version of the ActorBrowser config for setting values.
 	 * @returns		The config for this ActorBrowser.
diff --git a/Engine/Source/Editor/SceneOutliner/Public/ISceneOutlinerMode.h b/Engine/Source/Editor/SceneOutliner/Public/ISceneOutlinerMode.h
--- a/Engine/Source/Editor/SceneOutliner/Public/ISceneOutlinerMode.h
+++ b/Engine/Source/Editor/SceneOutliner/Public/ISceneOutlinerMode.h
@@ -159,6 +159,8 @@ public:
 	virtual bool ReparentItemToFolder(const FFolder& FolderPath, const FSceneOutlinerTreeItemPtr& Item) { return false; }
 	/** Select all descendants of a folder. Optionally select only immediate descendants. */
 	virtual void SelectFoldersDescendants(const TArray<FFolderTreeItem*>& FolderItems, bool bSelectImmediateChildrenOnly) {}
+	/** Returns true if this mode overrides the default folder double click behavior (expanding the subtree) */
+	virtual bool HasCustomFolderDoubleClick() const { return false; }
 public:
 	/** Pins an item list in the outliner */
 	virtual void PinItems(const TArray<FSceneOutlinerTreeItemPtr>& InItems) {}
