World Partition:
- Show unloaded actors under level instances in the World Outliner.
- Fix data layer manager to to call UDataLayerManager::OnActorDescContainerInstanceInitialized on all registered actor descriptor container instances recursively.

#changelist validated
#virtualized

diff --git a/Engine/Plugins/Editor/ObjectMixer/ObjectMixer/Source/ObjectMixer/Private/Views/List/SObjectMixerEditorList.cpp b/Engine/Plugins/Editor/ObjectMixer/ObjectMixer/Source/ObjectMixer/Private/Views/List/SObjectMixerEditorList.cpp
--- a/Engine/Plugins/Editor/ObjectMixer/ObjectMixer/Source/ObjectMixer/Private/Views/List/SObjectMixerEditorList.cpp
+++ b/Engine/Plugins/Editor/ObjectMixer/ObjectMixer/Source/ObjectMixer/Private/Views/List/SObjectMixerEditorList.cpp
@@ -907,7 +907,7 @@ void SObjectMixerEditorList::CreateActorTextInfoColumns(FSceneOutlinerInitializa
 				{
 					if (const UActorDescContainerInstance* ActorDescContainerInstance = ActorDescInstance->GetContainerInstance())
 					{
-						const UWorld* OwningWorld = ActorDescContainerInstance->GetWorldPartition()->GetWorld();
+						const UWorld* OwningWorld = ActorDescContainerInstance->GetOuterWorldPartition()->GetWorld();
 						if (const UDataLayerManager* DataLayerManager = UDataLayerManager::GetDataLayerManager(OwningWorld))
 						{
 							TSet<const UDataLayerInstance*> DataLayerInstances;
diff --git a/Engine/Source/Editor/DataLayerEditor/Private/DataLayer/DataLayersActorDescTreeItem.h b/Engine/Source/Editor/DataLayerEditor/Private/DataLayer/DataLayersActorDescTreeItem.h
--- a/Engine/Source/Editor/DataLayerEditor/Private/DataLayer/DataLayersActorDescTreeItem.h
+++ b/Engine/Source/Editor/DataLayerEditor/Private/DataLayer/DataLayersActorDescTreeItem.h
@@ -50,7 +50,7 @@ public:
 		if (ActorDescHandle.IsValid())
 		{
 			UActorDescContainerInstance* ContainerInstance = ActorDescHandle.GetContainerInstance();
-			UWorld* OwningWorld = ContainerInstance->GetWorldPartition()->GetWorld();
+			UWorld* OwningWorld = ContainerInstance->GetOuterWorldPartition()->GetWorld();
 			ULevelInstanceSubsystem* LevelInstanceSubsystem = UWorld::GetSubsystem<ULevelInstanceSubsystem>(OwningWorld);
 			ULevel* Level = ContainerInstance->GetTypedOuter<UWorld>()->PersistentLevel;
 			if (LevelInstanceSubsystem && Level && (Level != OwningWorld->GetCurrentLevel()))
diff --git a/Engine/Source/Editor/SceneOutliner/Private/ActorDescTreeItem.cpp b/Engine/Source/Editor/SceneOutliner/Private/ActorDescTreeItem.cpp
--- a/Engine/Source/Editor/SceneOutliner/Private/ActorDescTreeItem.cpp
+++ b/Engine/Source/Editor/SceneOutliner/Private/ActorDescTreeItem.cpp
@@ -245,6 +245,22 @@ FActorDescTreeItem::FActorDescTreeItem(const FGuid& InActorGuid, UActorDescConta
 	}
 }
 
+FActorDescTreeItem::FActorDescTreeItem(const FWorldPartitionActorDescInstance* InActorDescInstance)
+	: IActorBaseTreeItem(Type)
+	, ActorDescHandle(InActorDescInstance->GetContainerInstance(), InActorDescInstance->GetGuid())
+	, ID(ComputeTreeItemID(InActorDescInstance->GetGuid(), InActorDescInstance->GetContainerInstance()))
+	, ActorGuid(InActorDescInstance->GetGuid())
+{
+	if (const FWorldPartitionActorDescInstance* const ActorDescInstance = ActorDescHandle.GetInstance())
+	{
+		DisplayString = ActorDescInstance->GetActorLabel().ToString();
+	}
+	else
+	{
+		DisplayString = LOCTEXT("ActorLabelForMissingActor", "(Deleted Actor)").ToString();
+	}
+}
+
 FSceneOutlinerTreeItemID FActorDescTreeItem::ComputeTreeItemID(FGuid InActorGuid, UActorDescContainerInstance* InContainerInstance)
 {
 	FArchiveMD5 Ar;
@@ -328,7 +344,7 @@ bool FActorDescTreeItem::GetPinnedState() const
 {
 	if (ActorDescHandle.IsValid() && ActorDescHandle.GetContainerInstance())
 	{
-		UWorldPartition* WorldPartition = ActorDescHandle.GetContainerInstance()->GetWorldPartition();
+		UWorldPartition* WorldPartition = ActorDescHandle.GetContainerInstance()->GetOuterWorldPartition();
 		return WorldPartition ? WorldPartition->IsActorPinned(GetGuid()) : false;
 	}
 	return false;
diff --git a/Engine/Source/Editor/SceneOutliner/Private/ActorHierarchy.cpp b/Engine/Source/Editor/SceneOutliner/Private/ActorHierarchy.cpp
--- a/Engine/Source/Editor/SceneOutliner/Private/ActorHierarchy.cpp
+++ b/Engine/Source/Editor/SceneOutliner/Private/ActorHierarchy.cpp
@@ -342,6 +342,21 @@ FSceneOutlinerTreeItemPtr FActorHierarchy::FindOrCreateParentItem(const ISceneOu
 						return bCreate ? CreateItemForActor(ParentActor, true) : nullptr;
 					}
 				}
+				// If parent actor is not loaded
+				else if (UActorDescContainerInstance* ParentContainerInstance = Cast<UActorDescContainerInstance>(ContainerInstance->GetOuter()))
+				{
+					if (const FWorldPartitionActorDescInstance* ParentActorDescInstance = ParentContainerInstance->GetActorDescInstance(ContainerInstance->GetContainerActorGuid()))
+					{
+						if (const FSceneOutlinerTreeItemPtr* ParentItem = Items.Find(FActorDescTreeItem::ComputeTreeItemID(ParentActorDescInstance->GetGuid(), ParentContainerInstance)))
+						{
+							return *ParentItem;
+						}
+						else
+						{
+							return bCreate ? CreateItemForActorDescInstance(ParentActorDescInstance, true) : nullptr;
+						}
+					}
+				}
 			}
 		}
 	}
@@ -413,6 +428,11 @@ FSceneOutlinerTreeItemPtr FActorHierarchy::CreateItemForActor(AActor* InActor, b
 	return Mode->CreateItemFor<FActorTreeItem>(InActor, bForce);
 }
 
+FSceneOutlinerTreeItemPtr FActorHierarchy::CreateItemForActorDescInstance(const FWorldPartitionActorDescInstance* InActorDescInstance, bool bForce) const
+{
+	return Mode->CreateItemFor<FActorDescTreeItem>(InActorDescInstance, bForce);
+}
+
 bool FActorHierarchy::CheckLevelInstanceEditing(UWorld* World, AActor* Actor) const
 {
 	const ULevelInstanceSubsystem* LevelInstanceSubsystem = World->GetSubsystem<ULevelInstanceSubsystem>();
@@ -457,6 +477,28 @@ void FActorHierarchy::InsertActorItemAndCreateComponents(AActor* InActor, FScene
 
 void FActorHierarchy::CreateUnloadedItems(UWorld* World, TArray<FSceneOutlinerTreeItemPtr>& OutItems) const
 {
+	TFunction<void(const FWorldPartitionActorDescInstance*)> AddChildContainer = [this, &OutItems, &AddChildContainer](const FWorldPartitionActorDescInstance* ActorDescInstance)
+	{
+		FWorldPartitionActorDesc::FContainerInstance SubContainerInstance;
+		if (ActorDescInstance->GetChildContainerInstance(SubContainerInstance))
+		{
+			for (FActorDescInstanceList::TIterator<> It(SubContainerInstance.ContainerInstance); It; ++It)
+			{
+				const FWorldPartitionActorDescInstance* SubActorDescInstance = *It;
+
+				if (SubActorDescInstance && FActorDescTreeItem::ShouldDisplayInOutliner(SubActorDescInstance))
+				{
+					if (FSceneOutlinerTreeItemPtr SubActorDescItem = Mode->CreateItemFor<FActorDescTreeItem>(SubActorDescInstance))
+					{
+						OutItems.Add(SubActorDescItem);
+					}
+
+					AddChildContainer(SubActorDescInstance);
+				}
+			}
+		}
+	};
+
 	if (IsShowingUnloadedActors())
 	{
 		if (UWorldPartitionSubsystem* WorldPartitionSubsystem = UWorld::GetSubsystem<UWorldPartitionSubsystem>(World))
@@ -465,7 +507,7 @@ void FActorHierarchy::CreateUnloadedItems(UWorld* World, TArray<FSceneOutlinerTr
 			
 			const ULevelInstanceSubsystem* LevelInstanceSubsystem = World->GetSubsystem<ULevelInstanceSubsystem>();
 
-			WorldPartitionSubsystem->ForEachWorldPartition([this, LevelInstanceSubsystem, &LoadedActorCache, &OutItems](UWorldPartition* WorldPartition)
+			WorldPartitionSubsystem->ForEachWorldPartition([this, LevelInstanceSubsystem, &LoadedActorCache, &OutItems, &AddChildContainer](UWorldPartition* WorldPartition)
 			{
 				UWorld* OuterWorld = WorldPartition->GetTypedOuter<UWorld>();
 				ULevel* OuterLevel = OuterWorld ? OuterWorld->PersistentLevel : nullptr;
@@ -482,14 +524,26 @@ void FActorHierarchy::CreateUnloadedItems(UWorld* World, TArray<FSceneOutlinerTr
 
 				const TSet<FGuid>& LoadedActors = LoadedActorCache.GetLoadedActorsForLevel(OuterLevel);
 				
-				FWorldPartitionHelpers::ForEachActorDescInstance(WorldPartition, [this, &LoadedActors, &OutItems](const FWorldPartitionActorDescInstance* ActorDescInstance)
+				FWorldPartitionHelpers::ForEachActorDescInstance(WorldPartition, [this, &LoadedActors, &OutItems, &AddChildContainer, LevelInstanceSubsystem](const FWorldPartitionActorDescInstance* ActorDescInstance)
 				{
 					if (ActorDescInstance != nullptr && !LoadedActors.Contains(ActorDescInstance->GetGuid()) && FActorDescTreeItem::ShouldDisplayInOutliner(ActorDescInstance))
 					{
-						if (const FSceneOutlinerTreeItemPtr ActorDescItem = Mode->CreateItemFor<FActorDescTreeItem>(FActorDescTreeItem(ActorDescInstance->GetGuid(), ActorDescInstance->GetContainerInstance())))
+						UActorDescContainerInstance* ContainerInstance = ActorDescInstance->GetContainerInstance();
+
+						if (const FSceneOutlinerTreeItemPtr ActorDescItem = Mode->CreateItemFor<FActorDescTreeItem>(FActorDescTreeItem(ActorDescInstance->GetGuid(), ContainerInstance)))
 						{
 							OutItems.Add(ActorDescItem);
 						}
+
+						if (bShowingLevelInstances)
+						{
+							// If parent actor is not loaded, recurse
+							UWorld* OuterWorld = ContainerInstance->GetTypedOuter<UWorld>();
+							if (AActor* ParentActor = OuterWorld ? Cast<AActor>(LevelInstanceSubsystem->GetOwningLevelInstance(OuterWorld->PersistentLevel)) : nullptr; !ParentActor)
+							{
+								AddChildContainer(ActorDescInstance);
+							}
+						}
 					}
 					return true;
 				});
@@ -897,4 +951,4 @@ void FActorHierarchy::OnLevelActorFolderChanged(const AActor* InActor, FName Old
 		EventData.ItemIDs.Add(FSceneOutlinerTreeItemID(InActor));
 		HierarchyChangedEvent.Broadcast(EventData);
 	}
-}
\ No newline at end of file
+}
diff --git a/Engine/Source/Editor/SceneOutliner/Public/ActorDescTreeItem.h b/Engine/Source/Editor/SceneOutliner/Public/ActorDescTreeItem.h
--- a/Engine/Source/Editor/SceneOutliner/Public/ActorDescTreeItem.h
+++ b/Engine/Source/Editor/SceneOutliner/Public/ActorDescTreeItem.h
@@ -49,6 +49,9 @@ public:
 	/** Construct this item from an actor desc */
 	FActorDescTreeItem(const FGuid& InActorGuid, UActorDescContainerInstance* InContainerInstance);
 
+	/** Construct this item from an actor desc instance */
+	FActorDescTreeItem(const FWorldPartitionActorDescInstance* InActorDescInstance);
+
 	UE_DEPRECATED(5.4, "Use FWorldPartitionActorDescInstance version instead")
 	FActorDescTreeItem(const FGuid& InActorGuid, class UActorDescContainer* InContainer) : IActorBaseTreeItem(Type) {}
 
diff --git a/Engine/Source/Editor/SceneOutliner/Public/ActorHierarchy.h b/Engine/Source/Editor/SceneOutliner/Public/ActorHierarchy.h
--- a/Engine/Source/Editor/SceneOutliner/Public/ActorHierarchy.h
+++ b/Engine/Source/Editor/SceneOutliner/Public/ActorHierarchy.h
@@ -107,7 +107,11 @@ protected:
 	SCENEOUTLINER_API bool CheckLevelInstanceEditing(UWorld* World, AActor* Actor) const;
 
 	// Create a SceneOutlinerTreeItem for the given actor
-	virtual FSceneOutlinerTreeItemPtr CreateItemForActor(AActor* InActor, bool bForce = false) const;
+	SCENEOUTLINER_API virtual FSceneOutlinerTreeItemPtr CreateItemForActor(AActor* InActor, bool bForce = false) const;
+
+	// Create a SceneOutlinerTreeItem for the given actor descriptor instance
+	SCENEOUTLINER_API virtual FSceneOutlinerTreeItemPtr CreateItemForActorDescInstance(const FWorldPartitionActorDescInstance* InActorDescInstance, bool bForce = false) const;
+
 protected:
 	/** Send a an event indicating a full refresh of the hierarchy is required */
 	void FullRefreshEvent();
diff --git a/Engine/Source/Runtime/Engine/Private/LevelInstance/LevelInstanceSubsystem.cpp b/Engine/Source/Runtime/Engine/Private/LevelInstance/LevelInstanceSubsystem.cpp
--- a/Engine/Source/Runtime/Engine/Private/LevelInstance/LevelInstanceSubsystem.cpp
+++ b/Engine/Source/Runtime/Engine/Private/LevelInstance/LevelInstanceSubsystem.cpp
@@ -2583,7 +2583,7 @@ bool ULevelInstanceSubsystem::CheckForLoop(const ILevelInstanceInterface* LevelI
 
 bool ULevelInstanceSubsystem::PassLevelInstanceFilter(UWorld* World, const FWorldPartitionHandle& ActorHandle) const
 {
-	UWorld* ContainerOuterWorld = ActorHandle.GetContainerInstance()->GetWorldPartition()->GetTypedOuter<UWorld>();
+	UWorld* ContainerOuterWorld = ActorHandle.GetContainerInstance()->GetOuterWorldPartition()->GetTypedOuter<UWorld>();
 	check(ContainerOuterWorld);
 	if (const ILevelInstanceInterface* TopAncestor = GetOwningLevelInstance(ContainerOuterWorld->PersistentLevel))
 	{
diff --git a/Engine/Source/Runtime/Engine/Private/WorldPartition/ActorDescContainerInstance.cpp b/Engine/Source/Runtime/Engine/Private/WorldPartition/ActorDescContainerInstance.cpp
--- a/Engine/Source/Runtime/Engine/Private/WorldPartition/ActorDescContainerInstance.cpp
+++ b/Engine/Source/Runtime/Engine/Private/WorldPartition/ActorDescContainerInstance.cpp
@@ -89,6 +89,8 @@ void UActorDescContainerInstance::Initialize(const FInitializeParams& InParams)
 	FName OuterWorldContainerPackageName;
 
 	Transform = InParams.Transform;
+	ContainerActorGuid = InParams.ContainerActorGuid;
+
 	if (InParams.ContainerActorGuid.IsValid())
 	{
 		// It is possible to not have a ParentContainerInstance if we are in a non-WP main world
@@ -97,7 +99,7 @@ void UActorDescContainerInstance::Initialize(const FInitializeParams& InParams)
 	}
 				
 	// Only consider world if we are outered to a WorldPartition directly
-	UWorld* OuterWorld = HasWorldPartition() ? GetTypedOuter<UWorld>() : nullptr;
+	UWorld* OuterWorld = GetOuterWorldPartition() ? GetTypedOuter<UWorld>() : nullptr;
 	bool bIsInstanced = false;
 	FString SourceWorldPath, RemappedWorldPath;
 
@@ -223,7 +225,7 @@ void UActorDescContainerInstance::Uninitialize()
 bool UActorDescContainerInstance::ShouldRegisterDelegates() const
 {
 	// No World Partition means we are a ChildContainerInstance created for StreamingGeneration and need to listen to some events for updates
-	const UWorldPartition* WorldPartition = GetWorldPartition();
+	const UWorldPartition* WorldPartition = GetOuterWorldPartition();
 	const UWorld* OwningWorld = WorldPartition ? WorldPartition->GetWorld() : nullptr;
 
 	return !OwningWorld || !OwningWorld->IsGameWorld();
@@ -236,7 +238,7 @@ void UActorDescContainerInstance::RegisterDelegates()
 		check(Container);
 
 		// Only listen to Object replaced events on ContainerInstance that have a direct World Partition outer (Loaded Container Instances: Main World or Loaded Level Instances)
-		if (HasWorldPartition())
+		if (GetOuterWorldPartition())
 		{
 			FCoreUObjectDelegates::OnObjectsReplaced.AddUObject(this, &UActorDescContainerInstance::OnObjectsReplaced);
 		}
@@ -431,7 +433,7 @@ void UActorDescContainerInstance::RemoveActor(const FGuid& InActorGuid)
 {
 	if (TUniquePtr<FWorldPartitionActorDescInstance>* ActorDescInstance = GetActorDescriptor(InActorGuid))
 	{
-		if (UWorldPartition* WorldPartition = GetWorldPartition())
+		if (UWorldPartition* WorldPartition = GetOuterWorldPartition())
 		{
 			WorldPartition->OnActorDescInstanceRemoved(ActorDescInstance->Get());
 		}
@@ -468,7 +470,7 @@ const FLinkerInstancingContext* UActorDescContainerInstance::GetInstancingContex
 const FTransform& UActorDescContainerInstance::GetTransform() const
 {
 	// UActorDescContainerInstance outered to a UWorldPartition
-	if (UWorldPartition* WorldPartition = GetWorldPartition())
+	if (UWorldPartition* WorldPartition = GetOuterWorldPartition())
 	{
 		// GameWorld: Container instance is necessarly used for Streaming generation. In which case we want to return Identity and Transform of the World Partition will be applied on the LevelStreaming objects.
 		if (GetWorld()->IsGameWorld())
@@ -477,27 +479,21 @@ const FTransform& UActorDescContainerInstance::GetTransform() const
 			return FTransform::Identity;
 		}
 
-		return GetWorldPartition()->GetInstanceTransform();
+		return WorldPartition->GetInstanceTransform();
 	}
 	
 	// Transform is set when this Container Instance is a Child Container Instance created for Streaming Generation
 	return Transform.IsSet() ? Transform.GetValue() : FTransform::Identity;
 }
 
-UWorldPartition* UActorDescContainerInstance::GetWorldPartition() const
+UWorldPartition* UActorDescContainerInstance::GetTopWorldPartition() const
 {
-	// Only return a World Partition if we are directly outered to it
-	if (UWorldPartition* OuterWorldPartition = Cast<UWorldPartition>(GetOuter()))
-	{
-		return OuterWorldPartition;
-	}
-
-	return nullptr;
+	return GetTypedOuter<UWorldPartition>();
 }
 
-bool UActorDescContainerInstance::HasWorldPartition() const
+UWorldPartition* UActorDescContainerInstance::GetOuterWorldPartition() const
 {
-	return !!GetWorldPartition();
+	return Cast<UWorldPartition>(GetOuter());
 }
 
 void UActorDescContainerInstance::LoadAllActors(TArray<FWorldPartitionReference>& OutReferences)
@@ -538,7 +534,7 @@ void UActorDescContainerInstance::OnActorDescAdded(FWorldPartitionActorDesc* InA
 		NewActorDescInstance->RegisterChildContainerInstance();
 	}
 
-	if (UWorldPartition* WorldPartition = GetWorldPartition())
+	if (UWorldPartition* WorldPartition = GetOuterWorldPartition())
 	{
 		WorldPartition->OnActorDescInstanceAdded(NewActorDescInstance);
 	}
@@ -556,7 +552,7 @@ void UActorDescContainerInstance::OnActorDescUpdating(FWorldPartitionActorDesc*
 	TUniquePtr<FWorldPartitionActorDescInstance>* ActorDescInstance = GetActorDescInstancePtr(InActorDesc->GetGuid());
 	check(ActorDescInstance && ActorDescInstance->IsValid());
 		
-	if (UWorldPartition* WorldPartition = GetWorldPartition())
+	if (UWorldPartition* WorldPartition = GetOuterWorldPartition())
 	{
 		WorldPartition->OnActorDescInstanceUpdating(ActorDescInstance->Get());
 	}
@@ -578,7 +574,7 @@ void UActorDescContainerInstance::OnActorDescUpdated(FWorldPartitionActorDesc* I
 		ActorDescInstance->Get()->UpdateChildContainerInstance();
 	}
 
-	if (UWorldPartition* WorldPartition = GetWorldPartition())
+	if (UWorldPartition* WorldPartition = GetOuterWorldPartition())
 	{
 		WorldPartition->OnActorDescInstanceUpdated(ActorDescInstance->Get());
 	}
diff --git a/Engine/Source/Runtime/Engine/Private/WorldPartition/DataLayer/DataLayerManager.cpp b/Engine/Source/Runtime/Engine/Private/WorldPartition/DataLayer/DataLayerManager.cpp
--- a/Engine/Source/Runtime/Engine/Private/WorldPartition/DataLayer/DataLayerManager.cpp
+++ b/Engine/Source/Runtime/Engine/Private/WorldPartition/DataLayer/DataLayerManager.cpp
@@ -218,11 +218,11 @@ void UDataLayerManager::Initialize()
 	{
 		UActorDescContainerInstance::OnActorDescContainerInstanceInitialized.AddUObject(this, &UDataLayerManager::OnActorDescContainerInstanceInitialized);
 
-		// Manually call OnActorDescContainerInstanceInitialized on already initialized outer world partition container instance
-		if (ActorDescContainerInstance)
+		// Manually call OnActorDescContainerInstanceInitialized on already initialized outer world partition container instances
+		OuterWorldPartition->ForEachActorDescContainerInstance([this](UActorDescContainerInstance* ActorDescContainerInstance)
 		{
 			OnActorDescContainerInstanceInitialized(ActorDescContainerInstance);
-		}
+		}, true);
 	}
 
 	// SaveAs of a partition world will duplicate actors which will trigger AActor::FixupDataLayer and DataLayerManager is not yet created.
@@ -851,13 +851,7 @@ void UDataLayerManager::ResolveActorDescContainerInstanceDataLayersInternal(UAct
 	check(InActorDescContainerInstance);
 	check(!InActorDescInstance || (InActorDescInstance->GetContainerInstance() == InActorDescContainerInstance));
 		
-	const UWorldPartition* ContainerOuterWorldPartition = InActorDescContainerInstance->GetWorldPartition();
-	// Skip resolving for template containers (will be done on ActorDescViews)
-	if (!ContainerOuterWorldPartition)
-	{
-		return;
-	}
-
+	const UWorldPartition* ContainerOuterWorldPartition = InActorDescContainerInstance->GetTopWorldPartition();
 	const ULevelStreaming* ContainerLevelStreaming = FLevelUtils::FindStreamingLevel(ContainerOuterWorldPartition->GetTypedOuter<UWorld>()->PersistentLevel);
 	const UWorld* ContainerLevelStreamingWorld = ContainerLevelStreaming ? ContainerLevelStreaming->GetWorld() : nullptr;
 	const UWorldPartition* ContainerOwningWorldPartition = ContainerLevelStreamingWorld && !ContainerLevelStreamingWorld->IsGameWorld() ? ContainerLevelStreamingWorld->GetWorldPartition() : ContainerOuterWorldPartition;
diff --git a/Engine/Source/Runtime/Engine/Private/WorldPartition/DataLayer/ExternalDataLayerManager.cpp b/Engine/Source/Runtime/Engine/Private/WorldPartition/DataLayer/ExternalDataLayerManager.cpp
--- a/Engine/Source/Runtime/Engine/Private/WorldPartition/DataLayer/ExternalDataLayerManager.cpp
+++ b/Engine/Source/Runtime/Engine/Private/WorldPartition/DataLayer/ExternalDataLayerManager.cpp
@@ -743,7 +743,7 @@ bool UExternalDataLayerManager::UnregisterExternalDataLayerActorDescContainer(co
 	if (ActorDescContainerInstance)
 	{
 		// Unregister the ActorDescContainerInstance
-		UWorldPartition* WorldPartition = ActorDescContainerInstance->GetWorldPartition();
+		UWorldPartition* WorldPartition = ActorDescContainerInstance->GetOuterWorldPartition();
 		if (!WorldPartition || !WorldPartition->UnregisterActorDescContainerInstance(ActorDescContainerInstance))
 		{
 			UE_LOG(LogWorldPartition, Error, TEXT("[EDL: %s] Failed to unregister the ActorDescContainer of this External Data Layer."), *InExternalDataLayerAsset->GetName());
diff --git a/Engine/Source/Runtime/Engine/Private/WorldPartition/LoaderAdapter/LoaderAdapterPinnedActors.cpp b/Engine/Source/Runtime/Engine/Private/WorldPartition/LoaderAdapter/LoaderAdapterPinnedActors.cpp
--- a/Engine/Source/Runtime/Engine/Private/WorldPartition/LoaderAdapter/LoaderAdapterPinnedActors.cpp
+++ b/Engine/Source/Runtime/Engine/Private/WorldPartition/LoaderAdapter/LoaderAdapterPinnedActors.cpp
@@ -34,7 +34,7 @@ namespace LoaderAdapterPinnedActorsUtils
 
 		if (UActorDescContainerInstance* ContainerInstance = InActorDescInstance->GetContainerInstance())
 		{
-			const UWorldPartition* ContainerWorldPartition = ContainerInstance->GetWorldPartition();
+			const UWorldPartition* ContainerWorldPartition = ContainerInstance->GetOuterWorldPartition();
 			return ContainerWorldPartition && (ContainerWorldPartition->IsMainWorldPartition() || !bCheckIsMainWorldPartition);
 		}
 
diff --git a/Engine/Source/Runtime/Engine/Public/WorldPartition/ActorDescContainerInstance.h b/Engine/Source/Runtime/Engine/Public/WorldPartition/ActorDescContainerInstance.h
--- a/Engine/Source/Runtime/Engine/Public/WorldPartition/ActorDescContainerInstance.h
+++ b/Engine/Source/Runtime/Engine/Public/WorldPartition/ActorDescContainerInstance.h
@@ -101,11 +101,12 @@ public:
 	ENGINE_API bool IsInitialized() const { return bIsInitialized; }
 	ENGINE_API virtual void Uninitialize();
 
-	ENGINE_API UWorldPartition* GetWorldPartition() const;
-	ENGINE_API bool HasWorldPartition() const;
+	ENGINE_API UWorldPartition* GetTopWorldPartition() const;
+	ENGINE_API UWorldPartition* GetOuterWorldPartition() const;
 	ENGINE_API const FTransform& GetTransform() const;
 	ENGINE_API const FLinkerInstancingContext* GetInstancingContext() const;
 	const FActorContainerID& GetContainerID() const { return ContainerID; }
+	FGuid GetContainerActorGuid() const { return ContainerActorGuid; }
 
 	ENGINE_API static FName GetContainerPackageNameFromWorld(UWorld* InWorld);
 	ENGINE_API FName GetContainerPackage() const;
@@ -131,6 +132,8 @@ public:
 	ENGINE_API bool IsEmpty() const { return ActorDescList.IsEmpty(); }
 		
 	ENGINE_API void LoadAllActors(TArray<FWorldPartitionReference>& OutReferences);
+
+	ENGINE_API const TMap<FGuid, TObjectPtr<UActorDescContainerInstance>>& GetChildContainerInstances() const { return ChildContainerInstances; }
 		
 protected:
 	virtual void RegisterContainer(const FInitializeParams& InParams);
@@ -175,12 +178,12 @@ private:
 	TOptional<FLinkerInstancingContext>							InstancingContext;
 			
 	FActorContainerID											ContainerID;
+	FGuid														ContainerActorGuid;
 	TOptional<FTransform>										Transform;
 
 	UPROPERTY(Transient)
 	TObjectPtr<UActorDescContainer>								Container;
 		
-	// Keep reference off FActorDescInstance registered ChildContainerInstances for GC Ownership
 	UPROPERTY(Transient)
 	TMap<FGuid, TObjectPtr<UActorDescContainerInstance>>		ChildContainerInstances;
 
diff --git a/Engine/Source/Runtime/Engine/Public/WorldPartition/ActorDescContainerInstanceCollection.h b/Engine/Source/Runtime/Engine/Public/WorldPartition/ActorDescContainerInstanceCollection.h
--- a/Engine/Source/Runtime/Engine/Public/WorldPartition/ActorDescContainerInstanceCollection.h
+++ b/Engine/Source/Runtime/Engine/Public/WorldPartition/ActorDescContainerInstanceCollection.h
@@ -63,11 +63,11 @@ public:
 	DECLARE_EVENT_OneParam(UWorldPartition, FActorDescInstanceRemovedEvent, FWorldPartitionActorDescInstance*);
 	FActorDescInstanceRemovedEvent OnActorDescInstanceRemovedEvent;
 
-	void ForEachActorDescContainerInstanceBreakable(TFunctionRef<bool(ActorDescContPtrType)> Func);
-	void ForEachActorDescContainerInstanceBreakable(TFunctionRef<bool(ActorDescContPtrType)> Func) const;
+	void ForEachActorDescContainerInstanceBreakable(TFunctionRef<bool(ActorDescContPtrType)> Func, bool bRecursive = false);
+	void ForEachActorDescContainerInstanceBreakable(TFunctionRef<bool(ActorDescContPtrType)> Func, bool bRecursive = false) const;
 
-	void ForEachActorDescContainerInstance(TFunctionRef<void(ActorDescContPtrType)> Func);
-	void ForEachActorDescContainerInstance(TFunctionRef<void(ActorDescContPtrType)> Func) const;
+	void ForEachActorDescContainerInstance(TFunctionRef<void(ActorDescContPtrType)> Func, bool bRecursive = false);
+	void ForEachActorDescContainerInstance(TFunctionRef<void(ActorDescContPtrType)> Func, bool bRecursive = false) const;
 
 protected:
 	virtual void OnCollectionChanged() {};
@@ -439,11 +439,32 @@ void TActorDescContainerInstanceCollection<ActorDescContPtrType>::LoadAllActors(
 
 
 template<class ActorDescContPtrType>
-void TActorDescContainerInstanceCollection<ActorDescContPtrType>::ForEachActorDescContainerInstanceBreakable(TFunctionRef<bool(ActorDescContPtrType)> Func) const
+void TActorDescContainerInstanceCollection<ActorDescContPtrType>::ForEachActorDescContainerInstanceBreakable(TFunctionRef<bool(ActorDescContPtrType)> Func, bool bRecursive) const
 {
-	for (ActorDescContPtrType ActorDescContainerInstance : ActorDescContainerInstanceCollection)
+	TFunction<bool(ActorDescContPtrType)> InvokeFunc = [&Func, &InvokeFunc, bRecursive](ActorDescContPtrType ActorDescContainerInstance)
 	{
 		if (!Func(ActorDescContainerInstance))
+		{
+			return false;
+		}
+
+		if (bRecursive)
+		{
+			for (auto [Guid, ChildActorDescContainerInstance] : ActorDescContainerInstance->GetChildContainerInstances())
+			{
+				if (!InvokeFunc(ChildActorDescContainerInstance))
+				{
+					return false;
+				}
+			}
+		}
+
+		return true;
+	};
+
+	for (ActorDescContPtrType ActorDescContainerInstance : ActorDescContainerInstanceCollection)
+	{
+		if (!InvokeFunc(ActorDescContainerInstance))
 		{
 			break;
 		}
@@ -451,24 +472,24 @@ void TActorDescContainerInstanceCollection<ActorDescContPtrType>::ForEachActorDe
 }
 
 template<class ActorDescContPtrType>
-void TActorDescContainerInstanceCollection<ActorDescContPtrType>::ForEachActorDescContainerInstanceBreakable(TFunctionRef<bool(ActorDescContPtrType)> Func)
+void TActorDescContainerInstanceCollection<ActorDescContPtrType>::ForEachActorDescContainerInstanceBreakable(TFunctionRef<bool(ActorDescContPtrType)> Func, bool bRecursive)
 {
-	const_cast<const TActorDescContainerInstanceCollection*>(this)->ForEachActorDescContainerInstanceBreakable(Func);
+	const_cast<const TActorDescContainerInstanceCollection*>(this)->ForEachActorDescContainerInstanceBreakable(Func, bRecursive);
 }
 
 template<class ActorDescContPtrType>
-void TActorDescContainerInstanceCollection<ActorDescContPtrType>::ForEachActorDescContainerInstance(TFunctionRef<void(ActorDescContPtrType)> Func) const
+void TActorDescContainerInstanceCollection<ActorDescContPtrType>::ForEachActorDescContainerInstance(TFunctionRef<void(ActorDescContPtrType)> Func, bool bRecursive) const
 {
-	for (ActorDescContPtrType ActorDescContainerInstance : ActorDescContainerInstanceCollection)
+	ForEachActorDescContainerInstanceBreakable([&Func](ActorDescContPtrType ActorDescContainerInstance)
 	{
-		Func(ActorDescContainerInstance);
-	}
+		Func(ActorDescContainerInstance); return true;
+	}, bRecursive);
 }
 
 template<class ActorDescContPtrType>
-void TActorDescContainerInstanceCollection<ActorDescContPtrType>::ForEachActorDescContainerInstance(TFunctionRef<void(ActorDescContPtrType)> Func)
+void TActorDescContainerInstanceCollection<ActorDescContPtrType>::ForEachActorDescContainerInstance(TFunctionRef<void(ActorDescContPtrType)> Func, bool bRecursive)
 {
-	const_cast<const TActorDescContainerInstanceCollection*>(this)->ForEachActorDescContainerInstance(Func);
+	const_cast<const TActorDescContainerInstanceCollection*>(this)->ForEachActorDescContainerInstance(Func, bRecursive);
 }
 
 template<class ActorDescContPtrType>
