Level Instance: force unloaded level instance to be collapsed in the scene outliner by default.

diff --git a/Engine/Source/Editor/SceneOutliner/Private/ActorDescTreeItem.cpp b/Engine/Source/Editor/SceneOutliner/Private/ActorDescTreeItem.cpp
--- a/Engine/Source/Editor/SceneOutliner/Private/ActorDescTreeItem.cpp
+++ b/Engine/Source/Editor/SceneOutliner/Private/ActorDescTreeItem.cpp
@@ -243,6 +243,8 @@ FActorDescTreeItem::FActorDescTreeItem(const FGuid& InActorGuid, UActorDescConta
 	{
 		DisplayString = LOCTEXT("ActorLabelForMissingActor", "(Deleted Actor)").ToString();
 	}
+
+	Flags.bIsExpanded = false;
 }
 
 FActorDescTreeItem::FActorDescTreeItem(const FWorldPartitionActorDescInstance* InActorDescInstance)
@@ -259,6 +261,8 @@ FActorDescTreeItem::FActorDescTreeItem(const FWorldPartitionActorDescInstance* I
 	{
 		DisplayString = LOCTEXT("ActorLabelForMissingActor", "(Deleted Actor)").ToString();
 	}
+
+	Flags.bIsExpanded = false;
 }
 
 FSceneOutlinerTreeItemID FActorDescTreeItem::ComputeTreeItemID(FGuid InActorGuid, UActorDescContainerInstance* InContainerInstance)
