Scene Outliner: force items to be expanded by default when a search filter is active.

diff --git a/Engine/Source/Editor/SceneOutliner/Private/SSceneOutliner.cpp b/Engine/Source/Editor/SceneOutliner/Private/SSceneOutliner.cpp
--- a/Engine/Source/Editor/SceneOutliner/Private/SSceneOutliner.cpp
+++ b/Engine/Source/Editor/SceneOutliner/Private/SSceneOutliner.cpp
@@ -83,6 +83,7 @@ void SSceneOutliner::Construct(const FArguments& InArgs, const FSceneOutlinerIni
 	bNeedsRefresh = true;
 	bNeedsColumRefresh = true;
 	bShouldCacheColumnVisibility = true;
+	bForceParentItemsExpanded = false;
 	bIsReentrant = false;
 	bSortDirty = true;
 	bSelectionDirty = true;
@@ -1068,12 +1069,15 @@ void SSceneOutliner::AddUnfilteredItemToTree(FSceneOutlinerTreeItemRef Item)
 
 void SSceneOutliner::SetParentsExpansionState() const
 {
+	// If we have an active search filter, auto expand parents of items that passes the filter so they appear automatically in the outliner
+	bForceParentItemsExpanded = !SearchBoxFilter->GetRawFilterText().IsEmpty();
+
 	for (const auto& Pair : TreeItemMap)
 	{
 		auto& Item = Pair.Value;
 		if (Item->GetChildren().Num())
 		{
-			OutlinerTreeView->SetItemExpansion(Item, Item->Flags.bIsExpanded);
+			OutlinerTreeView->SetItemExpansion(Item, bForceParentItemsExpanded || Item->Flags.bIsExpanded);
 		}
 	}
 }
@@ -1979,6 +1983,11 @@ void SSceneOutliner::OnOutlinerTreeItemScrolledIntoView( FSceneOutlinerTreeItemP
 	
 void SSceneOutliner::OnItemExpansionChanged(FSceneOutlinerTreeItemPtr TreeItem, bool bIsExpanded) const
 {
+	if (bForceParentItemsExpanded)
+	{
+		return;
+	}
+
 	TreeItem->Flags.bIsExpanded = bIsExpanded;
 	TreeItem->OnExpansionChanged();
 
diff --git a/Engine/Source/Editor/SceneOutliner/Public/SSceneOutliner.h b/Engine/Source/Editor/SceneOutliner/Public/SSceneOutliner.h
--- a/Engine/Source/Editor/SceneOutliner/Public/SSceneOutliner.h
+++ b/Engine/Source/Editor/SceneOutliner/Public/SSceneOutliner.h
@@ -923,6 +923,9 @@ private:
 	/** True if the outliner should cache changes to column visibility into the config */
 	uint8 bShouldCacheColumnVisibility : 1;
 
+	/** True if we are forcing the underlying tree view to automatically expand all parents when searching */
+	mutable uint8 bForceParentItemsExpanded : 1;
+
 	/** Reentrancy guard */
 	bool bIsReentrant;
 
