Fixed crash when providing a folder with spaces at the beginning or at the end.
Folder label was sanitized when looking for existing but not when committing the new label.
Added auto-fix of folder labels in postload in editor.
Added sanitize of provided path in FFolder.

diff --git a/Engine/Source/Editor/SceneOutliner/Private/ActorFolderTreeItem.cpp b/Engine/Source/Editor/SceneOutliner/Private/ActorFolderTreeItem.cpp
--- a/Engine/Source/Editor/SceneOutliner/Private/ActorFolderTreeItem.cpp
+++ b/Engine/Source/Editor/SceneOutliner/Private/ActorFolderTreeItem.cpp
@@ -212,27 +212,31 @@ private:
 	void OnLabelCommitted(const FText& InLabel, ETextCommit::Type InCommitInfo)
 	{
 		auto TreeItem = TreeItemPtr.Pin();
-		if (TreeItem.IsValid() && !InLabel.ToString().Equals(TreeItem->GetLeafName().ToString(), ESearchCase::CaseSensitive))
+		if (TreeItem.IsValid())
 		{
-			// Rename the item
-			FFolder Folder = TreeItem->GetFolder();
-			FName NewPath = Folder.GetParent().GetPath();
-			if (NewPath.IsNone())
+			const FText TrimmedLabel = FText::TrimPrecedingAndTrailing(InLabel);
+			if (!TrimmedLabel.ToString().Equals(TreeItem->GetLeafName().ToString(), ESearchCase::CaseSensitive))
 			{
-				NewPath = FName(*InLabel.ToString());
-			}
-			else
-			{
-				NewPath = FName(*(NewPath.ToString() / InLabel.ToString()));
-			}
-			FFolder TreeItemNewFolder(Folder.GetRootObject(), NewPath);
+				// Rename the item
+				FFolder Folder = TreeItem->GetFolder();
+				FName NewPath = Folder.GetParent().GetPath();
+				if (NewPath.IsNone())
+				{
+					NewPath = FName(*TrimmedLabel.ToString());
+				}
+				else
+				{
+					NewPath = FName(*(NewPath.ToString() / TrimmedLabel.ToString()));
+				}
+				FFolder TreeItemNewFolder(Folder.GetRootObject(), NewPath);
 
-			FActorFolders::Get().RenameFolderInWorld(*TreeItem->World.Get(), Folder, TreeItemNewFolder);
+				FActorFolders::Get().RenameFolderInWorld(*TreeItem->World.Get(), Folder, TreeItemNewFolder);
 
-			auto Outliner = WeakSceneOutliner.Pin();
-			if (Outliner.IsValid())
-			{
-				Outliner->SetKeyboardFocus();
+				auto Outliner = WeakSceneOutliner.Pin();
+				if (Outliner.IsValid())
+				{
+					Outliner->SetKeyboardFocus();
+				}
 			}
 		}
 	}
diff --git a/Engine/Source/Runtime/Engine/Private/ActorFolder.cpp b/Engine/Source/Runtime/Engine/Private/ActorFolder.cpp
--- a/Engine/Source/Runtime/Engine/Private/ActorFolder.cpp
+++ b/Engine/Source/Runtime/Engine/Private/ActorFolder.cpp
@@ -52,6 +52,13 @@ bool UActorFolder::IsAsset() const
 	return IsPackageExternal() && !GetPackage()->HasAnyFlags(RF_Transient) && !HasAnyFlags(RF_Transient | RF_ClassDefaultObject) && !GetPackage()->HasAnyPackageFlags(PKG_PlayInEditor);
 }
 
+void UActorFolder::PostLoad()
+{
+	Super::PostLoad();
+
+	FolderLabel.TrimStartAndEndInline();
+}
+
 namespace ActorFolder
 {
 	static const FName NAME_FolderGuid(TEXT("FolderGuid"));
@@ -126,12 +133,13 @@ FActorFolderDesc UActorFolder::GetAssetRegistryInfoFromPackage(FName ActorFolder
 
 void UActorFolder::SetLabel(const FString& InFolderLabel)
 {
+	FString TrimmedFolderLabel = InFolderLabel.TrimStartAndEnd();
 	check(IsValid());
-	if (!FolderLabel.Equals(InFolderLabel, ESearchCase::CaseSensitive))
+	if (!FolderLabel.Equals(TrimmedFolderLabel, ESearchCase::CaseSensitive))
 	{
 		Modify();
 		FString OldFolderLabel = FolderLabel;
-		FolderLabel = InFolderLabel;
+		FolderLabel = TrimmedFolderLabel;
 		GetTypedOuter<ULevel>()->OnFolderLabelChanged(this, OldFolderLabel);
 	}
 }
diff --git a/Engine/Source/Runtime/Engine/Public/ActorFolder.h b/Engine/Source/Runtime/Engine/Public/ActorFolder.h
--- a/Engine/Source/Runtime/Engine/Public/ActorFolder.h
+++ b/Engine/Source/Runtime/Engine/Public/ActorFolder.h
@@ -33,6 +33,7 @@ public:
 	ENGINE_API virtual void GetAssetRegistryTags(FAssetRegistryTagsContext Context) const override;
 	UE_DEPRECATED(5.4, "Implement the version that takes FAssetRegistryTagsContext instead.")
 	ENGINE_API virtual void GetAssetRegistryTags(TArray<FAssetRegistryTag>& OutTags) const override;
+	ENGINE_API virtual void PostLoad() override;
 	//~ End UObject
 
 	ENGINE_API void SetParent(UActorFolder* InParent);
diff --git a/Engine/Source/Runtime/Engine/Public/Folder.h b/Engine/Source/Runtime/Engine/Public/Folder.h
--- a/Engine/Source/Runtime/Engine/Public/Folder.h
+++ b/Engine/Source/Runtime/Engine/Public/Folder.h
@@ -30,7 +30,23 @@ struct FFolder
 		: bPathInitialized(true)
 		, Path(InPath)
 		, RootObject(InRootObject)
-	{}
+	{
+		if (Path != GetEmptyPath())
+		{
+			TStringBuilder<512> Result;
+			TArray<FString> Parts;
+			Path.ToString().ParseIntoArray(Parts, TEXT("/"), true);
+			for (FString& Part : Parts)
+			{
+				if (Result.Len())
+				{
+					Result += TEXT("/");
+				}
+				Result += Part.TrimStartAndEnd();
+			}
+			Path = *Result;
+		}
+	}
 
 	FFolder(const FRootObject& InRootObject, const FGuid& InActorFolderGuid)
 		: bPathInitialized(false)
