FActorBrowsingMode: Make some functions virtual and export them

diff --git a/Engine/Source/Editor/SceneOutliner/Public/ActorBrowsingMode.h b/Engine/Source/Editor/SceneOutliner/Public/ActorBrowsingMode.h
--- a/Engine/Source/Editor/SceneOutliner/Public/ActorBrowsingMode.h
+++ b/Engine/Source/Editor/SceneOutliner/Public/ActorBrowsingMode.h
@@ -91,16 +91,16 @@ public:
 	void OnEditCopyActorsEnd();
 
 	/** Called by engine when edit paste actors begins */
-	void OnEditPasteActorsBegin();
+	SCENEOUTLINER_API virtual void OnEditPasteActorsBegin();
 
 	/** Called by engine when edit paste actors ends */
-	void OnEditPasteActorsEnd();
+	SCENEOUTLINER_API virtual void OnEditPasteActorsEnd();
 
 	/** Called by engine when edit duplicate actors begins */
-	void OnDuplicateActorsBegin();
+	SCENEOUTLINER_API virtual void OnDuplicateActorsBegin();
 
 	/** Called by engine when edit duplicate actors ends */
-	void OnDuplicateActorsEnd();
+	SCENEOUTLINER_API virtual void OnDuplicateActorsEnd();
 
 	/** Called by engine when edit delete actors begins */
 	void OnDeleteActorsBegin();
