World Outliner Optimizations:
- Don't query content bundles for actors that have their content bundle guid invalid.

diff --git a/Engine/Source/Editor/SceneOutliner/Private/SceneOutlinerModule.cpp b/Engine/Source/Editor/SceneOutliner/Private/SceneOutlinerModule.cpp
--- a/Engine/Source/Editor/SceneOutliner/Private/SceneOutlinerModule.cpp
+++ b/Engine/Source/Editor/SceneOutliner/Private/SceneOutlinerModule.cpp
@@ -414,9 +414,9 @@ void FSceneOutlinerModule::CreateActorInfoColumns(FSceneOutlinerInitializationOp
 		UContentBundleEngineSubsystem* ContentBundleEngineSubsystem = UContentBundleEngineSubsystem::Get();
 		if (const FActorTreeItem* ActorItem = Item.CastTo<FActorTreeItem>())
 		{
-			if (AActor* Actor = ActorItem->Actor.Get())
+			if (AActor* Actor = ActorItem->Actor.Get(); Actor && Actor->GetContentBundleGuid().IsValid())
 			{
-				if (const UContentBundleDescriptor * Descriptor = ContentBundleEngineSubsystem->GetContentBundleDescriptor(Actor->GetContentBundleGuid()))
+				if (const UContentBundleDescriptor* Descriptor = ContentBundleEngineSubsystem->GetContentBundleDescriptor(Actor->GetContentBundleGuid()))
 				{
 					return Descriptor->GetDisplayName();
 				}
@@ -424,7 +424,7 @@ void FSceneOutlinerModule::CreateActorInfoColumns(FSceneOutlinerInitializationOp
 		}
 		else if (const FActorDescTreeItem* ActorDescItem = Item.CastTo<FActorDescTreeItem>())
 		{
-			if (const FWorldPartitionActorDescInstance* ActorDescInstance = *ActorDescItem->ActorDescHandle)
+			if (const FWorldPartitionActorDescInstance* ActorDescInstance = *ActorDescItem->ActorDescHandle; ActorDescInstance && ActorDescInstance->GetContentBundleGuid().IsValid())
 			{
 				if (const UContentBundleDescriptor* Descriptor = ContentBundleEngineSubsystem->GetContentBundleDescriptor(ActorDescInstance->GetContentBundleGuid()))
 				{
