Merged External Data Layer column with Data Layer column. Put External Data Layer first if any.

diff --git a/Engine/Source/Editor/SceneOutliner/Private/SceneOutlinerModule.cpp b/Engine/Source/Editor/SceneOutliner/Private/SceneOutlinerModule.cpp
--- a/Engine/Source/Editor/SceneOutliner/Private/SceneOutlinerModule.cpp
+++ b/Engine/Source/Editor/SceneOutliner/Private/SceneOutlinerModule.cpp
@@ -5,6 +5,7 @@
 #include "Widgets/DeclarativeSyntaxSupport.h"
 #include "Framework/Application/SlateApplication.h"
 #include "SSceneOutliner.h"
+#include "Algo/Find.h"
 
 #include "SceneOutlinerActorInfoColumn.h"
 #include "SceneOutlinerGutter.h"
@@ -315,25 +316,7 @@ void FSceneOutlinerModule::CreateActorInfoColumns(FSceneOutlinerInitializationOp
 
 		return Result;
 	});
-
-	FGetTextForItem ExternalDatalayerInfoText = FGetTextForItem::CreateLambda([](const ISceneOutlinerTreeItem& Item) -> FString
-	{
-		const UExternalDataLayerAsset* ExternalDataLayerAsset = nullptr;
-		if (const FActorTreeItem* ActorItem = Item.CastTo<FActorTreeItem>())
-		{
-			if (AActor* Actor = ActorItem->Actor.Get())
-			{
-				ExternalDataLayerAsset = Actor->GetExternalDataLayerAsset();
-			}
-		}
-		else if (const FActorDescTreeItem* ActorDescItem = Item.CastTo<FActorDescTreeItem>())
-		{
-			ExternalDataLayerAsset = ActorDescItem->GetExternalDataLayerAsset();
-		}
-
-		return ExternalDataLayerAsset ? ExternalDataLayerAsset->GetName() : TEXT("");
-	});
-
+	
 	FGetTextForItem DataLayerInfoText = FGetTextForItem::CreateLambda([](const ISceneOutlinerTreeItem& Item) -> FString
 	{
 		TStringBuilder<128> Builder;
@@ -341,6 +324,11 @@ void FSceneOutlinerModule::CreateActorInfoColumns(FSceneOutlinerInitializationOp
 
 		auto BuildDataLayers = [&Builder, &DataLayerShortNames](const auto& DataLayerInstances, bool bPartOfOtherLevel)
 		{
+			if (const UDataLayerInstance* const* ExternalDataLayerInstance = !bPartOfOtherLevel ? Algo::FindByPredicate(DataLayerInstances, [](const UDataLayerInstance* DataLayerInstance) { return DataLayerInstance->IsA<UExternalDataLayerInstance>(); }) : nullptr)
+			{
+				Builder += (*ExternalDataLayerInstance)->GetDataLayerShortName();
+			}
+
 			for (const UDataLayerInstance* DataLayerInstance : DataLayerInstances)
 			{
 				if (!DataLayerInstance->IsA<UExternalDataLayerInstance>())
@@ -583,7 +571,6 @@ void FSceneOutlinerModule::CreateActorInfoColumns(FSceneOutlinerInitializationOp
 	AddTextInfoColumn(FSceneOutlinerBuiltInColumnTypes::Level(), LevelColumnName, LevelInfoText);
 	AddTextInfoColumn(FSceneOutlinerBuiltInColumnTypes::Layer(), FSceneOutlinerBuiltInColumnTypes::Layer_Localized(), LayerInfoText);
 	AddTextInfoColumn(FSceneOutlinerBuiltInColumnTypes::DataLayer(), FSceneOutlinerBuiltInColumnTypes::DataLayer_Localized(), DataLayerInfoText);
-	AddTextInfoColumn(FSceneOutlinerBuiltInColumnTypes::ExternalDataLayer(), FSceneOutlinerBuiltInColumnTypes::ExternalDataLayer_Localized(), ExternalDatalayerInfoText);
 	AddTextInfoColumn(FSceneOutlinerBuiltInColumnTypes::ContentBundle(), FSceneOutlinerBuiltInColumnTypes::ContentBundle_Localized(), ContentBundleInfoText);
 	AddTextInfoColumn(FSceneOutlinerBuiltInColumnTypes::SubPackage(), FSceneOutlinerBuiltInColumnTypes::SubPackage_Localized(), SubPackageInfoText);
 	AddTextInfoColumn(FSceneOutlinerBuiltInColumnTypes::Socket(), FSceneOutlinerBuiltInColumnTypes::Socket_Localized(), SocketInfoText);
diff --git a/Engine/Source/Editor/SceneOutliner/Public/SceneOutlinerPublicTypes.h b/Engine/Source/Editor/SceneOutliner/Public/SceneOutlinerPublicTypes.h
--- a/Engine/Source/Editor/SceneOutliner/Public/SceneOutlinerPublicTypes.h
+++ b/Engine/Source/Editor/SceneOutliner/Public/SceneOutlinerPublicTypes.h
@@ -50,7 +50,6 @@ struct FSceneOutlinerBuiltInColumnTypes
 	DEFINE_SCENEOUTLINER_BUILTIN_COLUMN_TYPE(Level, "Level", "SceneOutlinerLevelColumn", "Level");
 	DEFINE_SCENEOUTLINER_BUILTIN_COLUMN_TYPE(Layer, "Layer", "SceneOutlinerLayerColumn", "Layer");
 	DEFINE_SCENEOUTLINER_BUILTIN_COLUMN_TYPE(DataLayer, "Data Layer", "SceneOutlinerDataLayerColumn", "Data Layer");
-	DEFINE_SCENEOUTLINER_BUILTIN_COLUMN_TYPE(ExternalDataLayer, "External Data Layer", "SceneOutlinerExternalDataLayerColumn", "External Data Layer");
 	DEFINE_SCENEOUTLINER_BUILTIN_COLUMN_TYPE(ContentBundle, "Content Bundle", "SceneOutlinerContentBundleColumn", "Content Bundle");
 	DEFINE_SCENEOUTLINER_BUILTIN_COLUMN_TYPE(SubPackage, "Sub Package", "SceneOutlinerSubPackageColumn", "Sub Package");
 	DEFINE_SCENEOUTLINER_BUILTIN_COLUMN_TYPE(Pinned, "Pinned", "SceneOutlinerPinnedColumn", "Pinned");
