[World Outliner] Added option to freeze updates while in PIE and the associated game viewport has focus.

diff --git a/Engine/Source/Editor/SceneOutliner/Private/ActorBrowsingMode.cpp b/Engine/Source/Editor/SceneOutliner/Private/ActorBrowsingMode.cpp
--- a/Engine/Source/Editor/SceneOutliner/Private/ActorBrowsingMode.cpp
+++ b/Engine/Source/Editor/SceneOutliner/Private/ActorBrowsingMode.cpp
@@ -105,6 +105,7 @@ FActorBrowsingMode::FActorBrowsingMode(SSceneOutliner* InSceneOutliner, TWeakObj
 
 	bHideLevelInstanceHierarchy = LocalSettings.bHideLevelInstanceHierarchy;
 	InSceneOutliner->SetShowTransient(!LocalSettings.bHideTemporaryActors);
+	bShouldUpdateContentWhileInPIEFocused = LocalSettings.bShouldUpdateContentWhileInPIEFocused;
 
 	// Get the OutlinerModule to register FilterInfos with the FilterInfoMap
 	FSceneOutlinerFilterInfo ShowOnlySelectedActorsInfo(LOCTEXT("ToggleShowOnlySelected", "Only Selected"), LOCTEXT("ToggleShowOnlySelectedToolTip", "When enabled, only displays actors that are currently selected."), LocalSettings.bShowOnlySelectedActors, FCreateSceneOutlinerFilter::CreateStatic(&FActorBrowsingMode::CreateShowOnlySelectedActorsFilter));
@@ -544,6 +545,20 @@ void FActorBrowsingMode::InitializeViewMenuExtender(TSharedPtr<FExtender> Extend
 			EUserInterfaceActionType::ToggleButton
 		);
 
+		MenuBuilder.AddMenuEntry(
+			LOCTEXT("ShouldUpdateContentWhileInPIEFocusedLabel", "Update In PIE"),
+			LOCTEXT("ShouldUpdateContentWhileInPIEFocusedLabelTooltip", "When enabled, the Outliner will update in PIE."),
+			FSlateIcon(),
+			FUIAction(
+				FExecuteAction::CreateRaw(this, &FActorBrowsingMode::OnToggleShouldUpdateContentWhileInPIEFocused),
+				FCanExecuteAction(),
+				FIsActionChecked::CreateRaw(this, &FActorBrowsingMode::ShouldUpdateContentWhileInPIEFocused)
+			),
+			NAME_None,
+			EUserInterfaceActionType::ToggleButton
+		);
+
+
 		MenuBuilder.EndSection();
 
 		MenuBuilder.BeginSection("World", LOCTEXT("ShowWorldHeading", "World"));
@@ -1210,6 +1225,21 @@ bool FActorBrowsingMode::DoesFolderDoubleClickMarkCurrentFolder() const
 	return false;
 }
 
+void FActorBrowsingMode::OnToggleShouldUpdateContentWhileInPIEFocused()
+{
+	if (FActorBrowsingModeConfig* Settings = GetMutableConfig())
+	{
+		Settings->bShouldUpdateContentWhileInPIEFocused = !Settings->bShouldUpdateContentWhileInPIEFocused;
+		bShouldUpdateContentWhileInPIEFocused = Settings->bShouldUpdateContentWhileInPIEFocused;
+		SaveConfig();
+	}
+}
+
+bool FActorBrowsingMode::ShouldUpdateContentWhileInPIEFocused() const
+{
+	return bShouldUpdateContentWhileInPIEFocused;
+}
+
 void FActorBrowsingMode::OnFilterTextCommited(FSceneOutlinerItemSelection& Selection, ETextCommit::Type CommitType)
 {
 	// Start batching selection changes
diff --git a/Engine/Source/Editor/SceneOutliner/Private/ActorBrowsingModeSettings.h b/Engine/Source/Editor/SceneOutliner/Private/ActorBrowsingModeSettings.h
--- a/Engine/Source/Editor/SceneOutliner/Private/ActorBrowsingModeSettings.h
+++ b/Engine/Source/Editor/SceneOutliner/Private/ActorBrowsingModeSettings.h
@@ -72,6 +72,10 @@ public:
 	/** Specifies the behavior of double click on a folder */
 	UPROPERTY()
 	EActorBrowsingFolderDoubleClickMethod FolderDoubleClickMethod = EActorBrowsingFolderDoubleClickMethod::ToggleExpansion;
+
+	/** True if we want to allow updating content when in PIE with the PIE viewport focused */
+	UPROPERTY()
+	bool bShouldUpdateContentWhileInPIEFocused = false;
 };
 
 UCLASS(EditorConfig="ActorBrowsingMode")
diff --git a/Engine/Source/Editor/SceneOutliner/Private/ActorMode.cpp b/Engine/Source/Editor/SceneOutliner/Private/ActorMode.cpp
--- a/Engine/Source/Editor/SceneOutliner/Private/ActorMode.cpp
+++ b/Engine/Source/Editor/SceneOutliner/Private/ActorMode.cpp
@@ -19,6 +19,7 @@
 #include "EditorFolderUtils.h"
 #include "EditorLevelUtils.h"
 #include "EditorModeManager.h"
+#include "Engine/GameViewportClient.h"
 #include "WorldTreeItem.h"
 #include "LevelInstance/LevelInstanceInterface.h"
 #include "LevelInstance/LevelInstanceSubsystem.h"
@@ -30,6 +31,7 @@
 #include "Modules/ModuleManager.h"
 #include "ScopedTransaction.h"
 #include "SSocketChooser.h"
+#include "UnrealClient.h"
 
 DEFINE_LOG_CATEGORY_STATIC(LogActorMode, Log, All);
 
@@ -120,6 +122,7 @@ FActorMode::FActorMode(const FActorModeParams& Params)
 	, bHideUnloadedActors(Params.bHideUnloadedActors)
 	, bHideEmptyFolders(Params.bHideEmptyFolders)
 	, bCanInteractWithSelectableActorsOnly(Params.bCanInteractWithSelectableActorsOnly)
+	, bShouldUpdateContentWhileInPIEFocused(Params.bShouldUpdateContentWhileInPIEFocused)
 {
 	SceneOutliner->AddFilter(MakeShared<FActorFilter>(FActorTreeItem::FFilterPredicate::CreateLambda([this](const AActor* Actor)
 	{
@@ -418,6 +421,18 @@ bool FActorMode::CanInteract(const ISceneOutlinerTreeItem& Item) const
 	return true;
 }
 
+bool FActorMode::CanPopulate() const
+{
+	if (!bShouldUpdateContentWhileInPIEFocused)
+	{
+		if (UGameViewportClient* GameViewport = RepresentingWorld->GetGameViewport())
+		{
+			return !GameViewport->Viewport || !GameViewport->Viewport->HasFocus();
+		}
+	}
+	return true;
+}
+
 bool FActorMode::IsActorLevelDisplayable(ULevel* InLevel)
 {
 	// Don't show level tree item for the persistent level
diff --git a/Engine/Source/Editor/SceneOutliner/Private/SSceneOutliner.cpp b/Engine/Source/Editor/SceneOutliner/Private/SSceneOutliner.cpp
--- a/Engine/Source/Editor/SceneOutliner/Private/SSceneOutliner.cpp
+++ b/Engine/Source/Editor/SceneOutliner/Private/SSceneOutliner.cpp
@@ -2308,7 +2308,10 @@ void SSceneOutliner::Tick(const FGeometry& AllottedGeometry, const double InCurr
 	{
 		if( !bIsReentrant )
 		{
-			Populate();
+			if (Mode->CanPopulate())
+			{
+				Populate();
+			}
 		}
 	}
 	SortOutlinerTimer -= InDeltaTime;
diff --git a/Engine/Source/Editor/SceneOutliner/Public/ActorBrowsingMode.h b/Engine/Source/Editor/SceneOutliner/Public/ActorBrowsingMode.h
--- a/Engine/Source/Editor/SceneOutliner/Public/ActorBrowsingMode.h
+++ b/Engine/Source/Editor/SceneOutliner/Public/ActorBrowsingMode.h
@@ -144,6 +144,9 @@ private:
 	void OnToggleFolderDoubleClickMarkCurrentFolder();
 	bool DoesFolderDoubleClickMarkCurrentFolder() const;
 
+	void OnToggleShouldUpdateContentWhileInPIEFocused();
+	bool ShouldUpdateContentWhileInPIEFocused() const;
+
 	/**
 	 * Get a mutable version of the ActorBrowser config for setting values.
 	 * @returns		The config for this ActorBrowser.
diff --git a/Engine/Source/Editor/SceneOutliner/Public/ActorMode.h b/Engine/Source/Editor/SceneOutliner/Public/ActorMode.h
--- a/Engine/Source/Editor/SceneOutliner/Public/ActorMode.h
+++ b/Engine/Source/Editor/SceneOutliner/Public/ActorMode.h
@@ -53,6 +53,7 @@ struct SCENEOUTLINER_API FActorModeParams
 	bool bHideUnloadedActors = true;
 	bool bHideEmptyFolders = true;
 	bool bCanInteractWithSelectableActorsOnly = true;
+	bool bShouldUpdateContentWhileInPIEFocused = false;
 };
 
 class SCENEOUTLINER_API FActorMode : public ISceneOutlinerMode
@@ -83,6 +84,7 @@ public:
 	virtual FFolder::FRootObject GetPasteTargetRootObject() const override;
 
 	virtual bool CanInteract(const ISceneOutlinerTreeItem& Item) const override;
+	virtual bool CanPopulate() const override;
 	
 	virtual TSharedPtr<FDragDropOperation> CreateDragDropOperation(const FPointerEvent& MouseEvent, const TArray<FSceneOutlinerTreeItemPtr>& InTreeItems) const override;
 	virtual bool ParseDragDrop(FSceneOutlinerDragDropPayload& OutPayload, const FDragDropOperation& Operation) const override;
@@ -133,4 +135,6 @@ protected:
 	bool bAlwaysFrameSelection;
 	/** If True, CanInteract will be restricted to selectable actors only. */
 	bool bCanInteractWithSelectableActorsOnly;
+	/** Should we update content when in PIE and the PIE viewport has focus. */
+	bool bShouldUpdateContentWhileInPIEFocused;
 };
diff --git a/Engine/Source/Editor/SceneOutliner/Public/ISceneOutlinerMode.h b/Engine/Source/Editor/SceneOutliner/Public/ISceneOutlinerMode.h
--- a/Engine/Source/Editor/SceneOutliner/Public/ISceneOutlinerMode.h
+++ b/Engine/Source/Editor/SceneOutliner/Public/ISceneOutlinerMode.h
@@ -89,6 +89,8 @@ public:
 	virtual bool CanCustomizeToolbar() const { return false; }
 	/** Check if an item is interactive */
 	virtual bool CanInteract(const ISceneOutlinerTreeItem& Item) const { return true; }
+	/** Check if we are allowed to repopulate */
+	virtual bool CanPopulate() const { return true; }
 
 	/** Synchronize the mode specific selection with the tree view */
 	virtual void SynchronizeSelection() {}
