Adding a missing null check on FActorTreeItem::OnVisibilityChanged

Some tools like the Chaos Visual Debugger can perform GC on their own, so if a visibility changed event in the outliner is processed while that is happening, we can end up with an invalid actor ptr which will result in a crash

diff --git a/Engine/Source/Editor/SceneOutliner/Private/ActorTreeItem.cpp b/Engine/Source/Editor/SceneOutliner/Private/ActorTreeItem.cpp
--- a/Engine/Source/Editor/SceneOutliner/Private/ActorTreeItem.cpp
+++ b/Engine/Source/Editor/SceneOutliner/Private/ActorTreeItem.cpp
@@ -373,11 +373,14 @@ bool FActorTreeItem::ShouldShowVisibilityState() const
 
 void FActorTreeItem::OnVisibilityChanged(const bool bNewVisibility)
 {
-	// Save the actor to the transaction buffer to support undo/redo, but do
-	// not call Modify, as we do not want to dirty the actor's package and
-	// we're only editing temporary, transient values
-	SaveToTransactionBuffer(Actor.Get(), false);
-	Actor->SetIsTemporarilyHiddenInEditor(!bNewVisibility);
+	if (AActor* ActorPtr = Actor.Get())
+	{
+		// Save the actor to the transaction buffer to support undo/redo, but do
+		// not call Modify, as we do not want to dirty the actor's package and
+		// we're only editing temporary, transient values
+		SaveToTransactionBuffer(ActorPtr, false);
+		ActorPtr->SetIsTemporarilyHiddenInEditor(!bNewVisibility);
+	}
 }
 
 bool FActorTreeItem::GetVisibility() const
