[Scene Outliner] Optimize FFolderKey and FSceneOutlinerTreeItemID type hash computations.

#preglight 662a4a54684c0e6f58efa217

diff --git a/Engine/Source/Editor/SceneOutliner/Private/ActorDescTreeItem.cpp b/Engine/Source/Editor/SceneOutliner/Private/ActorDescTreeItem.cpp
--- a/Engine/Source/Editor/SceneOutliner/Private/ActorDescTreeItem.cpp
+++ b/Engine/Source/Editor/SceneOutliner/Private/ActorDescTreeItem.cpp
@@ -26,10 +26,6 @@
 #include "ToolMenus.h"
 #include "LevelEditorViewport.h"
 #include "HAL/PlatformApplicationMisc.h"
-#include "Misc/ArchiveMD5.h"
-#include "Misc/SecureHash.h"
-#include "UObject/ObjectKey.h"
-
 
 #define LOCTEXT_NAMESPACE "SceneOutliner_ActorDescTreeItem"
 
@@ -267,13 +263,7 @@ FActorDescTreeItem::FActorDescTreeItem(const FWorldPartitionActorDescInstance* I
 
 FSceneOutlinerTreeItemID FActorDescTreeItem::ComputeTreeItemID(FGuid InActorGuid, UActorDescContainerInstance* InContainerInstance)
 {
-	FArchiveMD5 Ar;
-	Ar << InActorGuid;
-
-	FObjectKey ContainerKey(InContainerInstance);
-	Ar << ContainerKey;
-
-	return FSceneOutlinerTreeItemID(Ar.GetGuidFromHash());
+	return FSceneOutlinerTreeItemID(FGuid::Combine(InActorGuid, InContainerInstance->GetContainerActorGuid()));
 }
 
 bool FActorDescTreeItem::ShouldDisplayInOutliner(const FWorldPartitionActorDescInstance* InActorDescInstance)
diff --git a/Engine/Source/Editor/SceneOutliner/Public/SceneOutlinerStandaloneTypes.h b/Engine/Source/Editor/SceneOutliner/Public/SceneOutlinerStandaloneTypes.h
--- a/Engine/Source/Editor/SceneOutliner/Public/SceneOutlinerStandaloneTypes.h
+++ b/Engine/Source/Editor/SceneOutliner/Public/SceneOutlinerStandaloneTypes.h
@@ -27,7 +27,7 @@ struct FFolderKey
 
 FORCEINLINE uint32 GetTypeHash(FFolderKey Key)
 {
-	return HashCombine(GetTypeHash(Key.Path), GetTypeHash(Key.RootObjectKey));
+	return HashCombineFast(GetTypeHash(Key.Path), GetTypeHash(Key.RootObjectKey));
 }
 
 /** Variant type that defines an identifier for a tree item. Assumes 'trivial relocatability' as with many unreal containers. */
@@ -143,7 +143,7 @@ public:
 			default:																		break;
 		}
 
-		return HashCombine((uint8)Type, Hash);
+		return HashCombineFast((uint8)Type, Hash);
 	}
 
 	friend uint32 GetTypeHash(const FSceneOutlinerTreeItemID& ItemID)
