Add ItemClick delegate to scene outliner

diff --git a/Engine/Source/Editor/SceneOutliner/Private/SSceneOutliner.cpp b/Engine/Source/Editor/SceneOutliner/Private/SSceneOutliner.cpp
--- a/Engine/Source/Editor/SceneOutliner/Private/SSceneOutliner.cpp
+++ b/Engine/Source/Editor/SceneOutliner/Private/SSceneOutliner.cpp
@@ -276,6 +276,9 @@ void SSceneOutliner::Construct(const FArguments& InArgs, const FSceneOutlinerIni
 			// Called when the user double-clicks with LMB on an item in the list
 			.OnMouseButtonDoubleClick( this, &SSceneOutliner::OnOutlinerTreeDoubleClick )
 
+			// Called when the user single-clicks with LMB on an item in the list
+			.OnMouseButtonClick( this, &SSceneOutliner::OnOutlinerTreeSingleClick )
+
 			// Called when an item is scrolled into view
 			.OnItemScrolledIntoView( this, &SSceneOutliner::OnOutlinerTreeItemScrolledIntoView )
 
@@ -1972,6 +1975,11 @@ void SSceneOutliner::OnOutlinerTreeDoubleClick( FSceneOutlinerTreeItemPtr TreeIt
 	OnDoubleClickOnTreeEvent.Broadcast(TreeItem);
 }
 
+void SSceneOutliner::OnOutlinerTreeSingleClick(FSceneOutlinerTreeItemPtr TreeItem) const
+{
+	Mode->OnItemClicked(TreeItem);
+}
+
 void SSceneOutliner::OnOutlinerTreeItemScrolledIntoView( FSceneOutlinerTreeItemPtr TreeItem, const TSharedPtr<ITableRow>& Widget )
 {
 	if (TreeItem == PendingRenameItem.Pin())
diff --git a/Engine/Source/Editor/SceneOutliner/Public/ISceneOutlinerMode.h b/Engine/Source/Editor/SceneOutliner/Public/ISceneOutlinerMode.h
--- a/Engine/Source/Editor/SceneOutliner/Public/ISceneOutlinerMode.h
+++ b/Engine/Source/Editor/SceneOutliner/Public/ISceneOutlinerMode.h
@@ -72,6 +72,8 @@ public:
 	virtual void OnItemSelectionChanged(FSceneOutlinerTreeItemPtr Item, ESelectInfo::Type SelectionType, const FSceneOutlinerItemSelection& Selection) {}
 	/** Called by the outliner when a tree item is double clicked */
 	virtual void OnItemDoubleClick(FSceneOutlinerTreeItemPtr Item) {}
+	/** Called by the outliner when a tree item is single-clicked */
+	virtual void OnItemClicked(FSceneOutlinerTreeItemPtr Item) {}
 	/** Called by the outliner when search box text is changed */
 	virtual void OnFilterTextChanged(const FText& InFilterText) {}
 	/** Called by the outliner when search box text is committed (by pressing enter) */
diff --git a/Engine/Source/Editor/SceneOutliner/Public/SSceneOutliner.h b/Engine/Source/Editor/SceneOutliner/Public/SSceneOutliner.h
--- a/Engine/Source/Editor/SceneOutliner/Public/SSceneOutliner.h
+++ b/Engine/Source/Editor/SceneOutliner/Public/SSceneOutliner.h
@@ -719,6 +719,9 @@ private:
 	/** Called by STreeView when the user double-clicks on an item in the tree */
 	void OnOutlinerTreeDoubleClick( FSceneOutlinerTreeItemPtr TreeItem );
 
+	/** Called by STreeView when the user single-clicks on an item in the tree */
+	void OnOutlinerTreeSingleClick( FSceneOutlinerTreeItemPtr TreeItem ) const;
+
 	/** Called by STreeView when an item is scrolled into view */
 	void OnOutlinerTreeItemScrolledIntoView( FSceneOutlinerTreeItemPtr TreeItem, const TSharedPtr<ITableRow>& Widget );
 
