Scene Outliner: Move all label color logic into the helper function

diff --git a/Engine/Source/Editor/SceneOutliner/Private/ActorTreeItem.cpp b/Engine/Source/Editor/SceneOutliner/Private/ActorTreeItem.cpp
--- a/Engine/Source/Editor/SceneOutliner/Private/ActorTreeItem.cpp
+++ b/Engine/Source/Editor/SceneOutliner/Private/ActorTreeItem.cpp
@@ -247,26 +247,6 @@ private:
 		{
 			return BaseColor.GetValue();
 		}
-
-		if (!Actor)
-		{
-			// Deleted actor!
-			return FLinearColor(0.2f, 0.2f, 0.25f);
-		}
-
-		UWorld* OwningWorld = Actor->GetWorld();
-		if (!OwningWorld)
-		{
-			// Deleted world!
-			return FLinearColor(0.2f, 0.2f, 0.25f);
-		}
-
-		const bool bRepresentingPIEWorld = TreeItem->Actor->GetWorld()->IsPlayInEditor();
-		if (bRepresentingPIEWorld && !TreeItem->bExistsInCurrentWorldAndPIE)
-		{
-			// Highlight actors that are exclusive to PlayWorld
-			return FLinearColor(0.9f, 0.8f, 0.4f);
-		}
 		
 		return FSlateColor::UseForeground();
 	}
diff --git a/Engine/Source/Editor/SceneOutliner/Private/SceneOutlinerStandaloneTypes.cpp b/Engine/Source/Editor/SceneOutliner/Private/SceneOutlinerStandaloneTypes.cpp
--- a/Engine/Source/Editor/SceneOutliner/Private/SceneOutlinerStandaloneTypes.cpp
+++ b/Engine/Source/Editor/SceneOutliner/Private/SceneOutlinerStandaloneTypes.cpp
@@ -2,6 +2,7 @@
 
 #include "SceneOutlinerStandaloneTypes.h"
 
+#include "ActorTreeItem.h"
 #include "EditorActorFolders.h"
 #include "Framework/Application/SlateApplication.h"
 #include "ISceneOutlinerTreeItem.h"
@@ -41,6 +42,31 @@ TOptional<FLinearColor> FSceneOutlinerCommonLabelData::GetForegroundColor(const
 		return DarkColor;
 	}
 
+	if(const FActorTreeItem* ActorTreeItem = TreeItem.CastTo<FActorTreeItem>())
+	{
+		AActor* Actor = ActorTreeItem->Actor.Get();
+		
+		if (!Actor)
+		{
+			// Deleted actor!
+			return FLinearColor(0.2f, 0.2f, 0.25f);
+		}
+
+		UWorld* OwningWorld = Actor->GetWorld();
+		if (!OwningWorld)
+		{
+			// Deleted world!
+			return FLinearColor(0.2f, 0.2f, 0.25f);
+		}
+
+		const bool bRepresentingPIEWorld = Actor->GetWorld()->IsPlayInEditor();
+		if (bRepresentingPIEWorld && !ActorTreeItem->bExistsInCurrentWorldAndPIE)
+		{
+			// Highlight actors that are exclusive to PlayWorld
+			return FLinearColor(0.9f, 0.8f, 0.4f);
+		}
+	}
+	
 	return TOptional<FLinearColor>();
 }
 
diff --git a/Engine/Source/Runtime/TypedElementFramework/Public/Elements/Interfaces/Capabilities/TypedElementUiStyleOverrideCapability.h b/Engine/Source/Runtime/TypedElementFramework/Public/Elements/Interfaces/Capabilities/TypedElementUiStyleOverrideCapability.h
--- a/Engine/Source/Runtime/TypedElementFramework/Public/Elements/Interfaces/Capabilities/TypedElementUiStyleOverrideCapability.h
+++ b/Engine/Source/Runtime/TypedElementFramework/Public/Elements/Interfaces/Capabilities/TypedElementUiStyleOverrideCapability.h
@@ -16,7 +16,7 @@ public:
 
 	~ITypedElementUiStyleOverrideCapability() override = default;
 
-	virtual void SetColorAndOpacity(const TAttribute<FLinearColor>& InColorAndOpacity) = 0;
+	virtual void SetForegroundColor(const TAttribute<FSlateColor>& InColorAndOpacity) = 0;
 };
 
 template<typename WidgetType>
@@ -25,9 +25,9 @@ class TTypedElementUiStyleOverrideCapability : public ITypedElementUiStyleOverri
 public:
 	explicit TTypedElementUiStyleOverrideCapability(WidgetType& InWidget) : Widget(InWidget){}
 	
-	void SetColorAndOpacity(const TAttribute<FLinearColor>& InColorAndOpacity) override
+	virtual void SetForegroundColor(const TAttribute<FSlateColor>& InColorAndOpacity) override
 	{
-		Widget.SetColorAndOpacity(InColorAndOpacity);
+		Widget.SetForegroundColor(InColorAndOpacity);
 	}
 
 private:
