Fixed FVisibilityDragDropOp which was inheriting from multiple instances of TSharedFromThis, which is unsupported.

diff --git a/Engine/Source/Editor/SceneOutliner/Private/SceneOutlinerGutter.cpp b/Engine/Source/Editor/SceneOutliner/Private/SceneOutlinerGutter.cpp
--- a/Engine/Source/Editor/SceneOutliner/Private/SceneOutlinerGutter.cpp
+++ b/Engine/Source/Editor/SceneOutliner/Private/SceneOutlinerGutter.cpp
@@ -58,7 +58,7 @@ bool FSceneOutlinerVisibilityCache::GetVisibility(const ISceneOutlinerTreeItem&
 	}
 }
 
-class FVisibilityDragDropOp : public FDragDropOperation, public TSharedFromThis<FVisibilityDragDropOp>
+class FVisibilityDragDropOp : public FDragDropOperation
 {
 public:
 	
