Changed Pin/Unpin (user facing) to ForceLoad/ReleaseForceLoad in SceneOutliner
#tests Context menu in Scene Outliner, playing with the Force Load column

diff --git a/Engine/Plugins/Editor/ObjectMixer/ObjectMixer/Source/ObjectMixer/Private/Views/List/Modes/ObjectMixerOutlinerMode.cpp b/Engine/Plugins/Editor/ObjectMixer/ObjectMixer/Source/ObjectMixer/Private/Views/List/Modes/ObjectMixerOutlinerMode.cpp
--- a/Engine/Plugins/Editor/ObjectMixer/ObjectMixer/Source/ObjectMixer/Private/Views/List/Modes/ObjectMixerOutlinerMode.cpp
+++ b/Engine/Plugins/Editor/ObjectMixer/ObjectMixer/Source/ObjectMixer/Private/Views/List/Modes/ObjectMixerOutlinerMode.cpp
@@ -1560,9 +1560,9 @@ void FObjectMixerOutlinerMode::RegisterContextMenu()
 						if (!Context->bRepresentingGameWorld && Context->bRepresentingPartitionedWorld)
 						{
 							Section.AddMenuEntry(
-								"PinItems",
-								LOCTEXT("Pin", "Pin"),
-								LOCTEXT("PinTooltip", "Keep the selected items loaded in the editor even when they don't overlap a loaded World Partition region"),
+								"ForceLoadItems",
+								LOCTEXT("ForceLoad", "Force Load"),
+								LOCTEXT("ForceLoadTooltip", "Keep the selected items loaded in the editor even when they don't overlap a loaded World Partition region"),
 								FSlateIcon(),
 								FUIAction(
 									FExecuteAction::CreateSP(SceneOutliner, &SSceneOutliner::PinSelectedItems),
@@ -1576,9 +1576,9 @@ void FObjectMixerOutlinerMode::RegisterContextMenu()
 									})));
 
 							Section.AddMenuEntry(
-								"UnpinItems",
-								LOCTEXT("Unpin", "Unpin"),
-								LOCTEXT("UnpinTooltip", "Allow the World Partition system to load and unload the selected items automatically"),
+								"ReleaseForceLoadItems",
+								LOCTEXT("ReleaseForceLoad", "Release Force Load"),
+								LOCTEXT("ReleaseForceLoadTooltip", "Allow the World Partition system to load and unload the selected items automatically"),
 								FSlateIcon(),
 								FUIAction(
 									FExecuteAction::CreateSP(SceneOutliner, &SSceneOutliner::UnpinSelectedItems),
diff --git a/Engine/Source/Editor/SceneOutliner/Private/ActorBrowsingMode.cpp b/Engine/Source/Editor/SceneOutliner/Private/ActorBrowsingMode.cpp
--- a/Engine/Source/Editor/SceneOutliner/Private/ActorBrowsingMode.cpp
+++ b/Engine/Source/Editor/SceneOutliner/Private/ActorBrowsingMode.cpp
@@ -761,9 +761,9 @@ void FActorBrowsingMode::FillDefaultContextBaseMenu(UToolMenu* InMenu)
 			if (!Context->bRepresentingGameWorld && Context->bRepresentingPartitionedWorld)
 			{
 				MainSection.AddMenuEntry(
-					"PinItems",
-					LOCTEXT("Pin", "Pin"),
-					LOCTEXT("PinTooltip", "Keep the selected items loaded in the editor even when they don't overlap a loaded World Partition region"),
+					"ForceLoadItems",
+					LOCTEXT("ForceLoad", "Force Load"),
+					LOCTEXT("ForceLoadTooltip", "Keep the selected items loaded in the editor even when they don't overlap a loaded World Partition region"),
 					FSlateIcon(),
 					FUIAction(
 						FExecuteAction::CreateSP(SceneOutliner, &SSceneOutliner::PinSelectedItems),
@@ -777,9 +777,9 @@ void FActorBrowsingMode::FillDefaultContextBaseMenu(UToolMenu* InMenu)
 				})));
 
 				MainSection.AddMenuEntry(
-					"UnpinItems",
-					LOCTEXT("Unpin", "Unpin"),
-					LOCTEXT("UnpinTooltip", "Allow the World Partition system to load and unload the selected items automatically"),
+					"ReleaseForceLoadItems",
+					LOCTEXT("ReleaseForceLoad", "Release Force Load"),
+					LOCTEXT("ReleaseForceLoadTooltip", "Allow the World Partition system to load and unload the selected items automatically"),
 					FSlateIcon(),
 					FUIAction(
 						FExecuteAction::CreateSP(SceneOutliner, &SSceneOutliner::UnpinSelectedItems),
diff --git a/Engine/Source/Editor/SceneOutliner/Private/SceneOutlinerPinnedColumn.cpp b/Engine/Source/Editor/SceneOutliner/Private/SceneOutlinerPinnedColumn.cpp
--- a/Engine/Source/Editor/SceneOutliner/Private/SceneOutlinerPinnedColumn.cpp
+++ b/Engine/Source/Editor/SceneOutliner/Private/SceneOutlinerPinnedColumn.cpp
@@ -188,7 +188,7 @@ SHeaderRow::FColumn::FArguments FSceneOutlinerPinnedColumn::ConstructHeaderRowCo
 		.VAlignHeader(VAlign_Center)
 		.HAlignCell(HAlign_Center)
 		.VAlignCell(VAlign_Center)
-		.ToolTip(IDocumentation::Get()->CreateToolTip(LOCTEXT("SceneOutlinerPinnedColumnToolTip", "Pinned: always loaded in editor"), nullptr, "Shared/MenuEntries/SceneOutliner_ActorBrowsingMode", "PinTooltip"))
+		.ToolTip(IDocumentation::Get()->CreateToolTip(LOCTEXT("SceneOutlinerForceLoadColumnToolTip", "Force Load: Keep the selected items loaded in the editor even when they don't overlap a loaded World Partition region."), nullptr, "Shared/MenuEntries/SceneOutliner_ActorBrowsingMode", "ForceLoadTooltip"))
 		[
 			SNew(SImage)
 			.ColorAndOpacity(FSlateColor::UseForeground())
@@ -206,7 +206,7 @@ const TSharedRef<SWidget> FSceneOutlinerPinnedColumn::ConstructRowWidget(FSceneO
 		.VAlign(VAlign_Center)
 		[
 			SNew(SPinnedWidget, WeakSceneOutliner, TreeItem, SharedThis(this), &Row)
-			.ToolTip(IDocumentation::Get()->CreateToolTip(LOCTEXT("SceneOutlinerPinnedWidgetTooltip", "Toggles whether this object is pinned (always loaded) in the editor."), nullptr, "Shared/MenuEntries/SceneOutliner_ActorBrowsingMode", "PinTooltip"))
+			.ToolTip(IDocumentation::Get()->CreateToolTip(LOCTEXT("SceneOutlinerForceLoadWidgetTooltip", "Toggles Force Load: Keep the selected items loaded in the editor even when they don't overlap a loaded World Partition region."), nullptr, "Shared/MenuEntries/SceneOutliner_ActorBrowsingMode", "ForceLoadTooltip"))
 		];
 	}
 	return SNullWidget::NullWidget;
