[SceneOutliner] Outliner Hotkey Pass
Removing hard-coded Hotkey and use commands for ActorBrowsingMode.
Added command class for ActorBrowsingMode.
Adding CommandList to SSceneOutliner to extend it through modes.

diff --git a/Engine/Source/Editor/SceneOutliner/Private/ActorBrowsingMode.cpp b/Engine/Source/Editor/SceneOutliner/Private/ActorBrowsingMode.cpp
--- a/Engine/Source/Editor/SceneOutliner/Private/ActorBrowsingMode.cpp
+++ b/Engine/Source/Editor/SceneOutliner/Private/ActorBrowsingMode.cpp
@@ -1,6 +1,8 @@
 // Copyright Epic Games, Inc. All Rights Reserved.
 
 #include "ActorBrowsingMode.h"
+
+#include "ActorBrowsingModeCommands.h"
 #include "Engine/Blueprint.h"
 #include "SceneOutlinerFilters.h"
 #include "SceneOutlinerModule.h"
@@ -52,6 +54,7 @@
 #include "Misc/ScopedSlowTask.h"
 #include "Elements/Framework/EngineElementsLibrary.h"
 #include "Elements/Framework/TypedElementHandle.h"
+#include "Framework/Commands/GenericCommands.h"
 
 DEFINE_LOG_CATEGORY_STATIC(LogActorBrowser, Log, All);
 
@@ -1284,76 +1287,24 @@ void FActorBrowsingMode::OnItemPassesFilters(const ISceneOutlinerTreeItem& Item)
 FReply FActorBrowsingMode::OnKeyDown(const FKeyEvent& InKeyEvent)
 {
 	const FSceneOutlinerItemSelection& Selection = SceneOutliner->GetSelection();
+	const FModifierKeysState ModifierKeys = FSlateApplication::Get().GetModifierKeys();
+	const FInputChord CheckChord( InKeyEvent.GetKey(), EModifierKey::FromBools(ModifierKeys.IsControlDown(), ModifierKeys.IsAltDown(), ModifierKeys.IsShiftDown(), ModifierKeys.IsCommandDown()) );
 
-	// Rename key: Rename selected actors (not rebindable, because it doesn't make much sense to bind.)
-	if (InKeyEvent.GetKey() == EKeys::F2)
+	// Use the keyboard shortcut bound to 'Focus Viewport To Selection'
+	if (FEditorViewportCommands::Get().FocusViewportToSelection->HasActiveChord(CheckChord))
 	{
 		if (Selection.Num() == 1)
 		{
-			FSceneOutlinerTreeItemPtr ItemToRename = Selection.SelectedItems[0].Pin();
-
-			if (ItemToRename.IsValid() && CanRenameItem(*ItemToRename) && ItemToRename->CanInteract())
-			{
-				SceneOutliner->SetPendingRenameItem(ItemToRename);
-				SceneOutliner->ScrollItemIntoView(ItemToRename);
-			}
-
-			return FReply::Handled();
-		}
-	}
-
-	// F5 forces a full refresh
-	else if (InKeyEvent.GetKey() == EKeys::F5)
-	{
-		SceneOutliner->FullRefresh();
-		return FReply::Handled();
-	}
-
-	// Delete key: Delete selected actors (not rebindable, because it doesn't make much sense to bind.)
-	// Use Delete and Backspace instead of Platform_Delete because the LevelEditor default Edit Delete is bound to both
-	else if (InKeyEvent.GetKey() == EKeys::Delete || InKeyEvent.GetKey() == EKeys::BackSpace)
-	{
-		if (SceneOutliner->GetSharedData().CustomDelete.IsBound())
-		{
-			SceneOutliner->GetSharedData().CustomDelete.Execute(Selection.SelectedItems);
-		}
-		else
-		{
-			if (RepresentingWorld.IsValid())
-			{
-				GUnrealEd->Exec(RepresentingWorld.Get(), TEXT("DELETE"));
-			}
-		}
-		return FReply::Handled();
-
-	}
+			FSceneOutlinerTreeItemPtr ItemToFocus = Selection.SelectedItems[0].Pin();
 
-	/* Allow the user to scroll to the current selection (and expand if needed) by pressing the key bound to
-	 * FEditorViewportCommands::Get().FocusViewportToSelection (Default: 'F')
-	 */
-	else
-	{
-		const FModifierKeysState ModifierKeys = FSlateApplication::Get().GetModifierKeys();
-		const FInputChord CheckChord( InKeyEvent.GetKey(), EModifierKey::FromBools(ModifierKeys.IsControlDown(), ModifierKeys.IsAltDown(), ModifierKeys.IsShiftDown(), ModifierKeys.IsCommandDown()) );
-
-		// Use the keyboard shortcut bound to 'Focus Viewport To Selection'
-		if (FEditorViewportCommands::Get().FocusViewportToSelection->HasActiveChord(CheckChord))
-		{
-			if (Selection.Num() == 1)
+			if (ItemToFocus.IsValid())
 			{
-				FSceneOutlinerTreeItemPtr ItemToFocus = Selection.SelectedItems[0].Pin();
-
-				if (ItemToFocus.IsValid())
-				{
-					SceneOutliner->ScrollItemIntoView(ItemToFocus);
-				}
-
-				// Return Unhandled here so that the level editor viewport can handle this event and focus the selected item
-				return FReply::Unhandled();
+				SceneOutliner->ScrollItemIntoView(ItemToFocus);
 			}
 		}
 	}
 
+	// Always return Unhandled here even if it entered the previous if so that the level editor viewport can handle the FocusViewport command as well
 	return FReply::Unhandled();
 }
 
@@ -1504,6 +1455,25 @@ void FActorBrowsingMode::RepairErrors() const
 	}
 }
 
+void FActorBrowsingMode::BindCommands(const TSharedRef<FUICommandList>& OutCommandList)
+{
+	OutCommandList->MapAction(
+		FGenericCommands::Get().Rename,
+		FExecuteAction::CreateRaw(this, &FActorBrowsingMode::OnExecuteRename ),
+		FCanExecuteAction::CreateRaw(this, &FActorBrowsingMode::CanExecuteRename));
+
+	OutCommandList->MapAction(
+		FGenericCommands::Get().Delete,
+		FExecuteAction::CreateRaw(this, &FActorBrowsingMode::OnExecuteDelete),
+		FCanExecuteAction::CreateRaw(this, &FActorBrowsingMode::CanDelete));
+	
+	OutCommandList->MapAction(
+		FActorBrowsingModeCommands::Get().Refresh,
+		FExecuteAction::CreateRaw(this, &FActorBrowsingMode::OnExecuteRefresh));
+
+	FInputBindingManager::Get().RegisterCommandList(FActorBrowsingModeCommands::Get().GetContextName(), OutCommandList);
+}
+
 bool FActorBrowsingMode::CanPasteFoldersOnlyFromClipboard() const
 {
 	// Intentionally not checking if the level is locked/hidden here, as it's better feedback for the user if they attempt to paste
@@ -1513,6 +1483,53 @@ bool FActorBrowsingMode::CanPasteFoldersOnlyFromClipboard() const
 	return PasteString.StartsWith("BEGIN FOLDERLIST");
 }
 
+void FActorBrowsingMode::OnExecuteDelete()
+{
+	const FSceneOutlinerItemSelection& Selection = SceneOutliner->GetSelection();
+	if (SceneOutliner->GetSharedData().CustomDelete.IsBound())
+	{
+		SceneOutliner->GetSharedData().CustomDelete.Execute(Selection.SelectedItems);
+	}
+	else
+	{
+		if (RepresentingWorld.IsValid())
+		{
+			GUnrealEd->Exec(RepresentingWorld.Get(), TEXT("DELETE"));
+		}
+	}
+}
+
+void FActorBrowsingMode::OnExecuteRefresh()
+{
+	SceneOutliner->FullRefresh();
+}
+
+void FActorBrowsingMode::OnExecuteRename()
+{
+	const FSceneOutlinerItemSelection& Selection = SceneOutliner->GetSelection();
+	if (Selection.Num() == 1)
+	{
+		FSceneOutlinerTreeItemPtr ItemToRename = Selection.SelectedItems[0].Pin();
+
+		if (ItemToRename.IsValid() && CanRenameItem(*ItemToRename) && ItemToRename->CanInteract())
+		{
+			SceneOutliner->SetPendingRenameItem(ItemToRename);
+			SceneOutliner->ScrollItemIntoView(ItemToRename);
+		}
+	}
+}
+
+bool FActorBrowsingMode::CanExecuteRename()
+{
+	const FSceneOutlinerItemSelection& Selection = SceneOutliner->GetSelection();
+	if (Selection.Num() == 1)
+	{
+		FSceneOutlinerTreeItemPtr ItemToRename = Selection.SelectedItems[0].Pin();
+		return ItemToRename.IsValid() && CanRenameItem(*ItemToRename) && ItemToRename->CanInteract();
+	}
+	return false;
+}
+
 void FActorBrowsingMode::SynchronizeSelectedActorDescs()
 {
 	if (UWorldPartitionSubsystem* WorldPartitionSubsystem = UWorld::GetSubsystem<UWorldPartitionSubsystem>(RepresentingWorld.Get()))
diff --git a/Engine/Source/Editor/SceneOutliner/Private/ActorBrowsingModeCommands.cpp b/Engine/Source/Editor/SceneOutliner/Private/ActorBrowsingModeCommands.cpp
new file mode 100644
--- /dev/null
+++ b/Engine/Source/Editor/SceneOutliner/Private/ActorBrowsingModeCommands.cpp
@@ -0,0 +1,17 @@
+// Copyright Epic Games, Inc. All Rights Reserved.
+
+#include "ActorBrowsingModeCommands.h"
+
+#include "Framework/Commands/InputChord.h"
+#include "Framework/Commands/UICommandInfo.h"
+#include "GenericPlatform/GenericApplication.h"
+#include "InputCoreTypes.h"
+
+#define LOCTEXT_NAMESPACE "ActorBrowsingModeCommands"
+
+void FActorBrowsingModeCommands::RegisterCommands()
+{
+	UI_COMMAND(Refresh, "Refresh", "Refresh Outliner", EUserInterfaceActionType::Button, FInputChord(EKeys::F5));
+}
+
+#undef LOCTEXT_NAMESPACE
diff --git a/Engine/Source/Editor/SceneOutliner/Private/ActorBrowsingModeCommands.h b/Engine/Source/Editor/SceneOutliner/Private/ActorBrowsingModeCommands.h
new file mode 100644
--- /dev/null
+++ b/Engine/Source/Editor/SceneOutliner/Private/ActorBrowsingModeCommands.h
@@ -0,0 +1,37 @@
+// Copyright Epic Games, Inc. All Rights Reserved.
+
+#pragma once
+
+#include "Framework/Commands/Commands.h"
+#include "HAL/Platform.h"
+#include "Internationalization/Internationalization.h"
+#include "Styling/AppStyle.h"
+#include "Templates/SharedPointer.h"
+#include "UObject/NameTypes.h"
+#include "UObject/UnrealNames.h"
+
+class FUICommandInfo;
+
+class FActorBrowsingModeCommands
+	: public TCommands<FActorBrowsingModeCommands>
+{
+public:
+	/** Default constructor. */
+	FActorBrowsingModeCommands()
+		: TCommands<FActorBrowsingModeCommands>( TEXT("ActorBrowsingModeCommands"), NSLOCTEXT("ActorBrowsingModeCommands", "ActorBrowsingModeCommands", "Actor Browsing Mode Commands"), NAME_None, FAppStyle::GetAppStyleSetName() )
+	{
+	}
+
+	virtual ~FActorBrowsingModeCommands()
+	{
+	}
+
+public:
+
+	//~ TCommands interface
+
+	virtual void RegisterCommands() override;
+
+public:
+	TSharedPtr<FUICommandInfo> Refresh;
+};
diff --git a/Engine/Source/Editor/SceneOutliner/Private/SSceneOutliner.cpp b/Engine/Source/Editor/SceneOutliner/Private/SSceneOutliner.cpp
--- a/Engine/Source/Editor/SceneOutliner/Private/SSceneOutliner.cpp
+++ b/Engine/Source/Editor/SceneOutliner/Private/SSceneOutliner.cpp
@@ -106,7 +106,10 @@ void SSceneOutliner::Construct(const FArguments& InArgs, const FSceneOutlinerIni
 
 	// @todo outliner: Should probably save this in layout!
 	// @todo outliner: Should save spacing for list view in layout
-	
+
+	// Setup Commands
+	BindCommands();
+
 	FSceneOutlinerModule& SceneOutlinerModule = FModuleManager::LoadModuleChecked<FSceneOutlinerModule>("SceneOutliner");
 	SceneOutlinerModule.OnColumnPermissionListChanged().AddSP(this, &SSceneOutliner::OnColumnPermissionListChanged);
 
@@ -2160,6 +2163,14 @@ void SSceneOutliner::OnItemLabelChanged(FSceneOutlinerTreeItemPtr ChangedItem)
 	}
 }
 
+void SSceneOutliner::BindCommands()
+{
+	CommandList = MakeShared<FUICommandList>();
+
+	// Bind Mode Commands
+	Mode->BindCommands(CommandList.ToSharedRef());
+}
+
 void SSceneOutliner::OnAssetReloaded(const EPackageReloadPhase InPackageReloadPhase, FPackageReloadedEvent* InPackageReloadedEvent)
 {
 	if (InPackageReloadPhase == EPackageReloadPhase::PostBatchPostGC)
@@ -2294,7 +2305,15 @@ bool SSceneOutliner::SupportsKeyboardFocus() const
 
 FReply SSceneOutliner::OnKeyDown( const FGeometry& MyGeometry, const FKeyEvent& InKeyEvent )
 {
-	// @todo outliner: Use command system for these for discoverability? (allow bindings?)
+	if (CommandList.IsValid())
+	{
+		if (CommandList->ProcessCommandBindings(InKeyEvent))
+		{
+			return FReply::Handled();
+		}
+	}
+
+	// Fallback to the Mode OnKeyDown to check if it's handled there
 	return Mode->OnKeyDown(InKeyEvent);
 }
 
diff --git a/Engine/Source/Editor/SceneOutliner/Private/SceneOutlinerModule.cpp b/Engine/Source/Editor/SceneOutliner/Private/SceneOutlinerModule.cpp
--- a/Engine/Source/Editor/SceneOutliner/Private/SceneOutlinerModule.cpp
+++ b/Engine/Source/Editor/SceneOutliner/Private/SceneOutlinerModule.cpp
@@ -18,6 +18,7 @@
 
 #include "ActorPickingMode.h"
 #include "ActorBrowsingMode.h"
+#include "ActorBrowsingModeCommands.h"
 #include "ActorTreeItem.h"
 #include "ActorFolderTreeItem.h"
 #include "ComponentTreeItem.h"
@@ -59,6 +60,9 @@ void FSceneOutlinerModule::StartupModule()
 	RegisterColumnType<FSceneOutlinerSourceControlColumn>();
 	RegisterColumnType<FSceneOutlinerPinnedColumn>();
 	RegisterColumnType<FSceneOutlinerActorUnsavedColumn>();
+
+	// Register Commands
+	FActorBrowsingModeCommands::Register();
 }
 
 
@@ -70,6 +74,9 @@ void FSceneOutlinerModule::ShutdownModule()
 	UnRegisterColumnType<FSceneOutlinerSourceControlColumn>();
 	UnRegisterColumnType<FSceneOutlinerPinnedColumn>();
 	UnRegisterColumnType<FSceneOutlinerActorUnsavedColumn>();
+
+	// Un-Register Commands
+	FActorBrowsingModeCommands::Unregister();
 }
 
 TSharedRef<ISceneOutliner> FSceneOutlinerModule::CreateSceneOutliner(const FSceneOutlinerInitializationOptions& InitOptions) const
diff --git a/Engine/Source/Editor/SceneOutliner/Public/ActorBrowsingMode.h b/Engine/Source/Editor/SceneOutliner/Public/ActorBrowsingMode.h
--- a/Engine/Source/Editor/SceneOutliner/Public/ActorBrowsingMode.h
+++ b/Engine/Source/Editor/SceneOutliner/Public/ActorBrowsingMode.h
@@ -48,6 +48,7 @@ public:
 	SCENEOUTLINER_API virtual bool HasErrors() const;
 	SCENEOUTLINER_API virtual FText GetErrorsText() const;
 	SCENEOUTLINER_API virtual void RepairErrors() const;	
+	SCENEOUTLINER_API virtual void BindCommands(const TSharedRef<FUICommandList>& OutCommandList) override;
 	SCENEOUTLINER_API virtual FFolder CreateNewFolder() override;
 	SCENEOUTLINER_API virtual FFolder GetFolder(const FFolder& ParentPath, const FName& LeafName) override;
 	SCENEOUTLINER_API virtual bool CreateFolder(const FFolder& NewFolder) override;
@@ -117,7 +118,19 @@ private:
 	static void RegisterContextMenu();
 	static void FillDefaultContextBaseMenu(UToolMenu* InMenu);
 	bool CanPasteFoldersOnlyFromClipboard() const;
-	
+
+	/** Delete the current selection */
+	void OnExecuteDelete();
+
+	/** Refresh the Outliner */
+	void OnExecuteRefresh();
+
+	/** Rename the selected Item (Only work with 1 selected Item) */
+	void OnExecuteRename();
+
+	/** Check whether the Rename can be executed */
+	bool CanExecuteRename();
+
 	void SynchronizeSelectedActorDescs();
 
 	void OnActorEditorContextSubsystemChanged();
diff --git a/Engine/Source/Editor/SceneOutliner/Public/ISceneOutlinerMode.h b/Engine/Source/Editor/SceneOutliner/Public/ISceneOutlinerMode.h
--- a/Engine/Source/Editor/SceneOutliner/Public/ISceneOutlinerMode.h
+++ b/Engine/Source/Editor/SceneOutliner/Public/ISceneOutlinerMode.h
@@ -11,6 +11,8 @@
 #include "SceneOutlinerDragDrop.h"
 #include "Folder.h"
 
+class FUICommandList;
+
 class ISceneOutlinerMode
 {
 public:
@@ -150,6 +152,9 @@ public:
 	/** Repair errors*/
 	virtual void RepairErrors() const {}
 
+	/** Bind commands for this mode to the given CommandList */
+	virtual void BindCommands(const TSharedRef<FUICommandList>& OutCommandList) {}
+
 public:
 	/* Folder management */
 		
diff --git a/Engine/Source/Editor/SceneOutliner/Public/SSceneOutliner.h b/Engine/Source/Editor/SceneOutliner/Public/SSceneOutliner.h
--- a/Engine/Source/Editor/SceneOutliner/Public/SSceneOutliner.h
+++ b/Engine/Source/Editor/SceneOutliner/Public/SSceneOutliner.h
@@ -657,6 +657,9 @@ private:
 	void OnItemLabelChanged(FSceneOutlinerTreeItemPtr ChangedItem);
 private:
 
+	/** BindCommands for this Outliner */
+	void BindCommands();
+
 	/** Map of columns that are shown on this outliner. */
 	TMap<FName, TSharedPtr<ISceneOutlinerColumn>> Columns;
 
@@ -895,6 +898,9 @@ private:
 	/** The button that displays view options */
 	TSharedPtr<SComboButton> ViewOptionsComboButton;
 
+	/** Command list to be processed during the OnKeyDown, extended by derived Mode through BindCommands */
+	TSharedPtr<FUICommandList> CommandList;
+
 private:
 
 	/** Called when SceneOutlinerModule column permission list changes. */
