Scene Outliner: Support searching by component names.

Component items are now created and made searchable independant of the owning actor item.

diff --git a/Engine/Source/Editor/SceneOutliner/Private/ActorBrowsingMode.cpp b/Engine/Source/Editor/SceneOutliner/Private/ActorBrowsingMode.cpp
--- a/Engine/Source/Editor/SceneOutliner/Private/ActorBrowsingMode.cpp
+++ b/Engine/Source/Editor/SceneOutliner/Private/ActorBrowsingMode.cpp
@@ -951,10 +951,12 @@ void FActorBrowsingMode::OnItemAdded(FSceneOutlinerTreeItemPtr Item)
 {
 	if (const FActorTreeItem* ActorItem = Item->CastTo<FActorTreeItem>())
 	{
+		// We incremented the count regardless of Flags.bIsFilteredOut because the count should match what the user sees, which includes things like
+		// actors which don't pass the filter themselves but are force shown by children being visible.
+		++FilteredActorCount;
+		
 		if (!Item->Flags.bIsFilteredOut)
 		{
-			++FilteredActorCount;
-
 			// Synchronize selection
 			if (GEditor->GetSelectedActors()->IsSelected(ActorItem->Actor.Get()))
 			{
@@ -964,11 +966,8 @@ void FActorBrowsingMode::OnItemAdded(FSceneOutlinerTreeItemPtr Item)
 	}
 	else if (Item->IsA<FActorDescTreeItem>())
 	{
-		if (!Item->Flags.bIsFilteredOut)
-		{
-			++FilteredActorCount;
-			++FilteredUnloadedActorCount;
-		}
+		++FilteredActorCount;
+		++FilteredUnloadedActorCount;
 	}
 }
 
@@ -976,18 +975,12 @@ void FActorBrowsingMode::OnItemRemoved(FSceneOutlinerTreeItemPtr Item)
 {
 	if (Item->IsA<FActorTreeItem>())
 	{
-		if (!Item->Flags.bIsFilteredOut)
-		{
-			--FilteredActorCount;
-		}
+		--FilteredActorCount;
 	}
 	else if (Item->IsA<FActorDescTreeItem>())
 	{
-		if (!Item->Flags.bIsFilteredOut)
-		{
-			--FilteredActorCount;
-			--FilteredUnloadedActorCount;
-		}
+		--FilteredActorCount;
+		--FilteredUnloadedActorCount;
 	}
 }
 
diff --git a/Engine/Source/Editor/SceneOutliner/Private/ActorHierarchy.cpp b/Engine/Source/Editor/SceneOutliner/Private/ActorHierarchy.cpp
--- a/Engine/Source/Editor/SceneOutliner/Private/ActorHierarchy.cpp
+++ b/Engine/Source/Editor/SceneOutliner/Private/ActorHierarchy.cpp
@@ -261,7 +261,7 @@ FSceneOutlinerTreeItemPtr FActorHierarchy::FindOrCreateParentItem(const ISceneOu
 			{
 				return *ParentItem;
 			}
-			else
+			else if (Owner->IsListedInSceneOutliner())
 			{
 				return bCreate ? CreateItemForActor(Owner, true) : nullptr;
 			}
@@ -384,7 +384,8 @@ void FActorHierarchy::CreateComponentItems(const AActor* Actor, TArray<FSceneOut
 		{
 			if (Component != nullptr)
 			{
-				if (FSceneOutlinerTreeItemPtr ComponentItem = Mode->CreateItemFor<FComponentTreeItem>(Component))
+				if (FSceneOutlinerTreeItemPtr ComponentItem =
+					Mode->CreateItemFor<FComponentTreeItem>(FComponentTreeItem(Component, bSearchComponentsByActorName)))
 				{
 					OutItems.Add(ComponentItem);
 				}
@@ -450,25 +451,25 @@ bool FActorHierarchy::CheckLevelInstanceEditing(UWorld* World, AActor* Actor) co
 
 void FActorHierarchy::InsertActorItemAndCreateComponents(AActor* InActor, FSceneOutlinerTreeItemPtr ActorItem, TArray<FSceneOutlinerTreeItemPtr>& OutItems) const
 {
+	const int32 InsertLocation = OutItems.Num();
+	
+	// Create all component items
+	CreateComponentItems(InActor, OutItems);
+
+	// If we are only showing actors with valid components, don't add the actor if no components were added
 	if (bShowingOnlyActorWithValidComponents)
 	{
-		int32 InsertLocation = OutItems.Num();
-
-		// Create all component items
-		CreateComponentItems(InActor, OutItems);
-
-		if (OutItems.Num() != InsertLocation)
+		// If the number of items remained the same after component insertion no components were inserted
+		if (OutItems.Num() == InsertLocation)
 		{
-			// Add the actor before the components
-			OutItems.Insert(ActorItem, InsertLocation);
+			return;
 		}
 	}
-	else
-	{
-		OutItems.Add(ActorItem);
 
-		// Create all component items
-		CreateComponentItems(InActor, OutItems);
+	if(ActorItem)
+	{
+		// Add the actor before the components
+		OutItems.Insert(ActorItem, InsertLocation);
 	}
 }
 
@@ -574,11 +575,11 @@ void FActorHierarchy::CreateWorldChildren(UWorld* World, TArray<FSceneOutlinerTr
 		{
 			continue;
 		}
+
+		FSceneOutlinerTreeItemPtr ActorItem = CreateItemForActor(Actor);
+		
+		InsertActorItemAndCreateComponents(Actor, ActorItem, OutItems);
 		
-		if (FSceneOutlinerTreeItemPtr ActorItem = CreateItemForActor(Actor))
-		{
-			InsertActorItemAndCreateComponents(Actor, ActorItem, OutItems);
-		}
 	}
 
 	CreateUnloadedItems(World, OutItems);
diff --git a/Engine/Source/Editor/SceneOutliner/Private/ActorMode.cpp b/Engine/Source/Editor/SceneOutliner/Private/ActorMode.cpp
--- a/Engine/Source/Editor/SceneOutliner/Private/ActorMode.cpp
+++ b/Engine/Source/Editor/SceneOutliner/Private/ActorMode.cpp
@@ -45,6 +45,7 @@ static FAutoConsoleVariableRef CVarAutoRepresentingWorldNetMode(
 
 using FActorFilter = TSceneOutlinerPredicateFilter<FActorTreeItem>;
 using FFolderFilter = TSceneOutlinerPredicateFilter<FFolderTreeItem>;
+using FComponentFilter = TSceneOutlinerPredicateFilter<FComponentTreeItem>;
 
 namespace SceneOutliner
 {
@@ -123,12 +124,20 @@ FActorMode::FActorMode(const FActorModeParams& Params)
 	, bHideEmptyFolders(Params.bHideEmptyFolders)
 	, bCanInteractWithSelectableActorsOnly(Params.bCanInteractWithSelectableActorsOnly)
 	, bShouldUpdateContentWhileInPIEFocused(Params.bShouldUpdateContentWhileInPIEFocused)
+	, bSearchComponentsByActorName(Params.bSearchComponentsByActorName)
 {
 	SceneOutliner->AddFilter(MakeShared<FActorFilter>(FActorTreeItem::FFilterPredicate::CreateLambda([this](const AActor* Actor)
 	{
 		return IsActorDisplayable(Actor);
 	}), FSceneOutlinerFilter::EDefaultBehaviour::Pass));
 
+	// Don't show components if the owner actor is not displayable
+	SceneOutliner->AddFilter(MakeShared<FComponentFilter>(FComponentTreeItem::FFilterPredicate::CreateLambda([this](const UActorComponent* ActorComponent)
+	{
+		return ActorComponent && IsActorDisplayable(ActorComponent->GetOwner());
+	}), FSceneOutlinerFilter::EDefaultBehaviour::Pass));
+
+
 	auto FolderPassesFilter = [this](const FFolder& InFolder, bool bInCheckHideLevelInstanceFlag)
 	{
 		if (ILevelInstanceInterface* LevelInstance = Cast<ILevelInstanceInterface>(InFolder.GetRootObjectPtr()))
@@ -172,6 +181,7 @@ TUniquePtr<ISceneOutlinerHierarchy> FActorMode::CreateHierarchy()
 	ActorHierarchy->SetShowingLevelInstances(!bHideLevelInstanceHierarchy);
 	ActorHierarchy->SetShowingUnloadedActors(!bHideUnloadedActors);
 	ActorHierarchy->SetShowingEmptyFolders(!bHideEmptyFolders);
+	ActorHierarchy->SetSearchComponentsByActorName(bSearchComponentsByActorName);
 
 	return ActorHierarchy;
 }
diff --git a/Engine/Source/Editor/SceneOutliner/Private/ComponentTreeItem.cpp b/Engine/Source/Editor/SceneOutliner/Private/ComponentTreeItem.cpp
--- a/Engine/Source/Editor/SceneOutliner/Private/ComponentTreeItem.cpp
+++ b/Engine/Source/Editor/SceneOutliner/Private/ComponentTreeItem.cpp
@@ -39,7 +39,7 @@ struct SComponentTreeLabel : FSceneOutlinerCommonLabelData, public SCompoundWidg
 				SAssignNew(InlineTextBlock, SInlineEditableTextBlock)
 				.Text(this, &SComponentTreeLabel::GetDisplayText)
 				.ToolTipText(this, &SComponentTreeLabel::GetTooltipText)
-				//.HighlightText(HighlightText)
+				.HighlightText(HighlightText)
 				.ColorAndOpacity(this, &SComponentTreeLabel::GetForegroundColor)
 				//.OnTextCommitted(this, &SComponentTreeLabel::OnLabelCommitted)
 				//.OnVerifyTextChanged(this, &SComponentTreeLabel::OnVerifyItemLabelChanged)
@@ -173,10 +173,11 @@ private:
 
 const FSceneOutlinerTreeItemType FComponentTreeItem::Type(&ISceneOutlinerTreeItem::Type);
 
-FComponentTreeItem::FComponentTreeItem(UActorComponent* InComponent)
+FComponentTreeItem::FComponentTreeItem(UActorComponent* InComponent, bool bInSearchComponentsByActorName)
 	: ISceneOutlinerTreeItem(Type)
 	, Component(InComponent)
 	, ID(InComponent)
+	, bSearchComponentsByActorName(bInSearchComponentsByActorName)
 {
 	AActor* OwningActor = InComponent->GetOwner();
 	bExistsInCurrentWorldAndPIE = GEditor->ObjectsThatExistInEditorWorld.Get(OwningActor);
@@ -228,4 +229,9 @@ TSharedRef<SWidget> FComponentTreeItem::GenerateLabelWidget(ISceneOutliner& Outl
 	return SNew(SComponentTreeLabel, *this, Outliner, InRow);
 }
 
+bool FComponentTreeItem::GetSearchComponentByActorName() const
+{
+	return bSearchComponentsByActorName;
+}
+
 #undef LOCTEXT_NAMESPACE
diff --git a/Engine/Source/Editor/SceneOutliner/Private/SceneOutlinerHelpers.cpp b/Engine/Source/Editor/SceneOutliner/Private/SceneOutlinerHelpers.cpp
--- a/Engine/Source/Editor/SceneOutliner/Private/SceneOutlinerHelpers.cpp
+++ b/Engine/Source/Editor/SceneOutliner/Private/SceneOutlinerHelpers.cpp
@@ -4,6 +4,7 @@
 #include "ActorTreeItem.h"
 #include "ActorDescTreeItem.h"
 #include "ActorFolderTreeItem.h"
+#include "ComponentTreeItem.h"
 #include "EditorClassUtils.h"
 #include "Engine/Blueprint.h"
 #include "Misc/PackageName.h"
@@ -82,7 +83,7 @@ namespace SceneOutliner
 		return nullptr;
 	}
 
-TSharedPtr<SWidget> FSceneOutlinerHelpers::GetClassHyperlink(UObject* InObject)
+	TSharedPtr<SWidget> FSceneOutlinerHelpers::GetClassHyperlink(UObject* InObject)
 	{
 		if (InObject)
 		{
@@ -119,4 +120,25 @@ TSharedPtr<SWidget> FSceneOutlinerHelpers::GetClassHyperlink(UObject* InObject)
 		return nullptr;
 	}
 
+	void FSceneOutlinerHelpers::PopulateExtraSearchStrings(const ISceneOutlinerTreeItem& TreeItem, TArray< FString >& OutSearchStrings)
+	{
+		// For components, we want them to be searchable by the actor name if they request so. This is so you can search by actors in component
+		// pickers without the actual components themselves being filtered out.
+		if (const FComponentTreeItem* ComponentTreeItem = TreeItem.CastTo<FComponentTreeItem>())
+		{
+			if (ComponentTreeItem->GetSearchComponentByActorName())
+			{
+				if (const UActorComponent* Component = ComponentTreeItem->Component.Get())
+				{
+					if (const AActor* Owner = Component->GetOwner())
+					{
+						constexpr bool bCreateIfNone = false;
+						OutSearchStrings.Add(Owner->GetActorLabel(bCreateIfNone));
+					}
+					
+				}
+			}
+		}
+	}
+
 ;}
\ No newline at end of file
diff --git a/Engine/Source/Editor/SceneOutliner/Private/SceneOutlinerItemLabelColumn.cpp b/Engine/Source/Editor/SceneOutliner/Private/SceneOutlinerItemLabelColumn.cpp
--- a/Engine/Source/Editor/SceneOutliner/Private/SceneOutlinerItemLabelColumn.cpp
+++ b/Engine/Source/Editor/SceneOutliner/Private/SceneOutlinerItemLabelColumn.cpp
@@ -15,6 +15,7 @@
 #include "SceneOutlinerStandaloneTypes.h"
 #include "ISceneOutlinerMode.h"
 #include "SceneOutlinerDragDrop.h"
+#include "SceneOutlinerHelpers.h"
 #include "Widgets/Views/SListView.h"
 #include "SortHelper.h"
 #include "Widgets/SToolTip.h"
@@ -41,6 +42,7 @@ const TSharedRef<SWidget> FSceneOutlinerItemLabelColumn::ConstructRowWidget(FSce
 void FSceneOutlinerItemLabelColumn::PopulateSearchStrings( const ISceneOutlinerTreeItem& Item, TArray< FString >& OutSearchStrings ) const
 {
 	OutSearchStrings.Add(Item.GetDisplayString());
+	SceneOutliner::FSceneOutlinerHelpers::PopulateExtraSearchStrings(Item, OutSearchStrings);
 }
 
 void FSceneOutlinerItemLabelColumn::SortItems(TArray<FSceneOutlinerTreeItemPtr>& OutItems, const EColumnSortMode::Type SortMode) const
diff --git a/Engine/Source/Editor/SceneOutliner/Private/SceneOutlinerModule.cpp b/Engine/Source/Editor/SceneOutliner/Private/SceneOutlinerModule.cpp
--- a/Engine/Source/Editor/SceneOutliner/Private/SceneOutlinerModule.cpp
+++ b/Engine/Source/Editor/SceneOutliner/Private/SceneOutlinerModule.cpp
@@ -152,6 +152,7 @@ TSharedRef<ISceneOutliner> FSceneOutlinerModule::CreateComponentPicker(const FSc
 			Params.bHideLevelInstanceHierarchy = true;
 			Params.bHideUnloadedActors = true;
 			Params.bHideEmptyFolders = true;
+			Params.bSearchComponentsByActorName = true;
 			return new FActorPickingMode(Params, OnItemPicked);
 		});
 
diff --git a/Engine/Source/Editor/SceneOutliner/Public/ActorHierarchy.h b/Engine/Source/Editor/SceneOutliner/Public/ActorHierarchy.h
--- a/Engine/Source/Editor/SceneOutliner/Public/ActorHierarchy.h
+++ b/Engine/Source/Editor/SceneOutliner/Public/ActorHierarchy.h
@@ -30,6 +30,7 @@ public:
 	void SetShowingLevelInstances(bool bInShowingLevelInstances) { bShowingLevelInstances = bInShowingLevelInstances; }
 	void SetShowingUnloadedActors(bool bInShowingUnloadedActors) { bShowingUnloadedActors = bInShowingUnloadedActors; }
 	void SetShowingEmptyFolders(bool bInShowingEmptyFolders) { bShowingEmptyFolders = bInShowingEmptyFolders; }
+	void SetSearchComponentsByActorName(bool bInSearchComponentsByActorName) { bSearchComponentsByActorName = bInSearchComponentsByActorName;}
 
 private:
 	bool IsShowingUnloadedActors() const;
@@ -123,6 +124,7 @@ protected:
 	bool bShowingLevelInstances = false;
 	bool bShowingUnloadedActors = false;
 	bool bShowingEmptyFolders = false;
+	bool bSearchComponentsByActorName = false;
 
 protected:
 
diff --git a/Engine/Source/Editor/SceneOutliner/Public/ActorMode.h b/Engine/Source/Editor/SceneOutliner/Public/ActorMode.h
--- a/Engine/Source/Editor/SceneOutliner/Public/ActorMode.h
+++ b/Engine/Source/Editor/SceneOutliner/Public/ActorMode.h
@@ -35,7 +35,9 @@ struct SCENEOUTLINER_API FActorModeParams
 {
 	FActorModeParams() {}
 
-	FActorModeParams(SSceneOutliner* InSceneOutliner, const TWeakObjectPtr<UWorld>& InSpecifiedWorldToDisplay = nullptr, bool bInHideComponents = true, bool bInHideLevelInstanceHierarchy = true, bool bInHideUnloadedActors = true, bool bInHideEmptyFolders = true, bool bInCanInteractWithSelectableActorsOnly = true)
+	FActorModeParams(SSceneOutliner* InSceneOutliner, const TWeakObjectPtr<UWorld>& InSpecifiedWorldToDisplay = nullptr, bool bInHideComponents = true,
+		bool bInHideLevelInstanceHierarchy = true, bool bInHideUnloadedActors = true, bool bInHideEmptyFolders = true,
+		bool bInCanInteractWithSelectableActorsOnly = true, bool binSearchComponentsByActorName = false)
 		: SpecifiedWorldToDisplay(InSpecifiedWorldToDisplay)
 		, SceneOutliner(InSceneOutliner)
 		, bHideComponents(bInHideComponents)
@@ -43,6 +45,7 @@ struct SCENEOUTLINER_API FActorModeParams
 		, bHideUnloadedActors(bInHideUnloadedActors)
 		, bHideEmptyFolders(bInHideEmptyFolders)
 		, bCanInteractWithSelectableActorsOnly(bInCanInteractWithSelectableActorsOnly)
+		, bSearchComponentsByActorName(binSearchComponentsByActorName)
 	{}
 
 	TWeakObjectPtr<UWorld> SpecifiedWorldToDisplay = nullptr;
@@ -54,6 +57,7 @@ struct SCENEOUTLINER_API FActorModeParams
 	bool bHideEmptyFolders = true;
 	bool bCanInteractWithSelectableActorsOnly = true;
 	bool bShouldUpdateContentWhileInPIEFocused = false;
+	bool bSearchComponentsByActorName = false;
 };
 
 class SCENEOUTLINER_API FActorMode : public ISceneOutlinerMode
@@ -137,4 +141,8 @@ protected:
 	bool bCanInteractWithSelectableActorsOnly;
 	/** Should we update content when in PIE and the PIE viewport has focus. */
 	bool bShouldUpdateContentWhileInPIEFocused;
+	/** If true and bHideComponents is false, components will be shown if the owning actor is searched for even if the search text does not match
+	 * the components
+	 */
+	bool bSearchComponentsByActorName;
 };
diff --git a/Engine/Source/Editor/SceneOutliner/Public/ComponentTreeItem.h b/Engine/Source/Editor/SceneOutliner/Public/ComponentTreeItem.h
--- a/Engine/Source/Editor/SceneOutliner/Public/ComponentTreeItem.h
+++ b/Engine/Source/Editor/SceneOutliner/Public/ComponentTreeItem.h
@@ -35,7 +35,9 @@ public:
 	static const FSceneOutlinerTreeItemType Type;
 
 	/** Construct this item from an Component */
-	FComponentTreeItem(UActorComponent* InComponent);
+	FComponentTreeItem(UActorComponent* InComponent, bool bInSearchComponentsByActorName = false);
+
+	bool GetSearchComponentByActorName() const;
 
 	/* Begin ISceneOutlinerTreeItem Implementation */
 	virtual bool IsValid() const override { return Component.IsValid(); }
@@ -47,6 +49,10 @@ public:
 public:
 	/** true if this item exists in both the current world and PIE. */
 	bool bExistsInCurrentWorldAndPIE;
+
+	/** If true components will be shown if the owning actor is searched for even if the search text does not match the component */
+	bool bSearchComponentsByActorName;
+	
 	/** Cache the string displayed */
 	FString CachedDisplayString;
 };
diff --git a/Engine/Source/Editor/SceneOutliner/Public/SceneOutlinerHelpers.h b/Engine/Source/Editor/SceneOutliner/Public/SceneOutlinerHelpers.h
--- a/Engine/Source/Editor/SceneOutliner/Public/SceneOutlinerHelpers.h
+++ b/Engine/Source/Editor/SceneOutliner/Public/SceneOutlinerHelpers.h
@@ -12,8 +12,12 @@ namespace SceneOutliner
 	public:
 		UE_DEPRECATED(5.5, "FSceneOutlinerHelpers::GetExternalPackageName has been deprecated implement/use ISceneOutlinerTreeItem::GetPackageName instead")
 		static FString GetExternalPackageName(const ISceneOutlinerTreeItem& TreeItem);
+		
 		UE_DEPRECATED(5.5, "FSceneOutlinerHelpers::GetExternalPackage has been deprecated implement/use ISceneOutlinerTreeItem::GetPackageName instead")
 		static UPackage* GetExternalPackage(const ISceneOutlinerTreeItem& TreeItem);
+		
 		static TSharedPtr<SWidget> GetClassHyperlink(UObject* InObject);
+		
+		static void PopulateExtraSearchStrings(const ISceneOutlinerTreeItem& TreeItem, TArray< FString >& OutSearchStrings);
 	};
 };
