Add checks to make sure RepresentingWorld is valid before it's accessed.

diff --git a/Engine/Source/Editor/SceneOutliner/Private/ActorBrowsingMode.cpp b/Engine/Source/Editor/SceneOutliner/Private/ActorBrowsingMode.cpp
--- a/Engine/Source/Editor/SceneOutliner/Private/ActorBrowsingMode.cpp
+++ b/Engine/Source/Editor/SceneOutliner/Private/ActorBrowsingMode.cpp
@@ -893,16 +893,19 @@ TSharedPtr<SWidget> FActorBrowsingMode::BuildContextMenu()
 	ContextObject->bRepresentingPartitionedWorld = bRepresentingWorldPartitionedWorld;
 
 	int32 NumPinnedItems = 0;
-	if (const UWorldPartition* const WorldPartition = RepresentingWorld->GetWorldPartition())
+	if (RepresentingWorld.IsValid())
 	{
-		ItemSelection.ForEachItem<IActorBaseTreeItem>([WorldPartition, &NumPinnedItems](const IActorBaseTreeItem& ActorItem)
+		if (const UWorldPartition* const WorldPartition = RepresentingWorld->GetWorldPartition())
 		{
-			if (WorldPartition->IsActorPinned(ActorItem.GetGuid()))
-			{
-				++NumPinnedItems;
-			}
-			return true;
-		});
+			ItemSelection.ForEachItem<IActorBaseTreeItem>([WorldPartition, &NumPinnedItems](const IActorBaseTreeItem& ActorItem)
+				{
+					if (WorldPartition->IsActorPinned(ActorItem.GetGuid()))
+					{
+						++NumPinnedItems;
+					}
+					return true;
+				});
+		}
 	}
 	ContextObject->NumPinnedItems = NumPinnedItems;
 
@@ -996,29 +999,32 @@ void FActorBrowsingMode::OnLevelActorDeleted(AActor* Actor)
 
 void FActorBrowsingMode::OnSelectUnloadedActors(const TArray<FGuid>& ActorGuids)
 {
-	if (UWorldPartition* WorldPartition = RepresentingWorld->GetWorldPartition())
+	if (RepresentingWorld.IsValid())
 	{
-		TArray<FSceneOutlinerTreeItemPtr> ItemsToSelect;
-		ItemsToSelect.Reserve(ActorGuids.Num());
-		for (const FGuid& ActorGuid : ActorGuids)
+		if (UWorldPartition* WorldPartition = RepresentingWorld->GetWorldPartition())
 		{
-			if (FWorldPartitionActorDescInstance* ActorDescInstance = WorldPartition->GetActorDescInstance(ActorGuid))
+			TArray<FSceneOutlinerTreeItemPtr> ItemsToSelect;
+			ItemsToSelect.Reserve(ActorGuids.Num());
+			for (const FGuid& ActorGuid : ActorGuids)
 			{
-				if (FSceneOutlinerTreeItemPtr ItemPtr = SceneOutliner->GetTreeItem(FActorDescTreeItem::ComputeTreeItemID(ActorDescInstance->GetGuid(), ActorDescInstance->GetContainerInstance())))
+				if (FWorldPartitionActorDescInstance* ActorDescInstance = WorldPartition->GetActorDescInstance(ActorGuid))
 				{
-					ItemsToSelect.Add(ItemPtr);
+					if (FSceneOutlinerTreeItemPtr ItemPtr = SceneOutliner->GetTreeItem(FActorDescTreeItem::ComputeTreeItemID(ActorDescInstance->GetGuid(), ActorDescInstance->GetContainerInstance())))
+					{
+						ItemsToSelect.Add(ItemPtr);
+					}
 				}
 			}
-		}
-
-		if (ItemsToSelect.Num())
-		{
-			SceneOutliner->SetItemSelection(ItemsToSelect, true);
-			SceneOutliner->ScrollItemIntoView(ItemsToSelect.Last());
 
-			if (const FActorDescTreeItem* ActorDescItem = ItemsToSelect.Last()->CastTo<FActorDescTreeItem>())
+			if (ItemsToSelect.Num())
 			{
-				ActorDescItem->FocusActorBounds();
+				SceneOutliner->SetItemSelection(ItemsToSelect, true);
+				SceneOutliner->ScrollItemIntoView(ItemsToSelect.Last());
+
+				if (const FActorDescTreeItem* ActorDescItem = ItemsToSelect.Last()->CastTo<FActorDescTreeItem>())
+				{
+					ActorDescItem->FocusActorBounds();
+				}
 			}
 		}
 	}
@@ -1370,7 +1376,7 @@ FText FActorBrowsingMode::GetErrorsText() const
 
 void FActorBrowsingMode::RepairErrors() const
 {
-	if (!bRepresentingWorldGameWorld && bRepresentingWorldPartitionedWorld)
+	if (RepresentingWorld.IsValid() && !bRepresentingWorldGameWorld && bRepresentingWorldPartitionedWorld)
 	{
 		if (UWorldPartition* const WorldPartition = RepresentingWorld->GetWorldPartition())
 		{
