Speculative fix for a crash when opening a level with a renamed Verse component and IDOs enabled.

The crash seems to be caused by the filters in the Scene Outliner not being created. This could in theory happen when the Scene Outliner is created but not constructed yet. It's unclear how this stage can be reached, so this is just a targeted speculative fix until a deeper investigation can be done.

diff --git a/Engine/Source/Editor/SceneOutliner/Private/SSceneOutliner.cpp b/Engine/Source/Editor/SceneOutliner/Private/SSceneOutliner.cpp
--- a/Engine/Source/Editor/SceneOutliner/Private/SSceneOutliner.cpp
+++ b/Engine/Source/Editor/SceneOutliner/Private/SSceneOutliner.cpp
@@ -54,9 +54,6 @@ void SSceneOutliner::Construct(const FArguments& InArgs, const FSceneOutlinerIni
 	// We use the filter collection provided, otherwise we create our own
 	Filters = InInitOptions.Filters.IsValid() ? InInitOptions.Filters : MakeShareable(new FSceneOutlinerFilters);
 
-	// The interactive filter collection
-	InteractiveFilters = MakeShareable(new FSceneOutlinerFilters);
-
 	OutlinerIdentifier = InInitOptions.OutlinerIdentifier;
 	
 	// Setup the SearchBox
diff --git a/Engine/Source/Editor/SceneOutliner/Public/SSceneOutliner.h b/Engine/Source/Editor/SceneOutliner/Public/SSceneOutliner.h
--- a/Engine/Source/Editor/SceneOutliner/Public/SSceneOutliner.h
+++ b/Engine/Source/Editor/SceneOutliner/Public/SSceneOutliner.h
@@ -246,7 +246,11 @@ public:
 		void Construct(const FArguments& InArgs, const FSceneOutlinerInitializationOptions& InitOptions);
 
 	/** Default constructor - initializes data that is shared between all tree items */
-	SSceneOutliner() : SharedData(MakeShareable(new FSharedSceneOutlinerData)) {}
+	SSceneOutliner() 
+		: SharedData(MakeShareable(new FSharedSceneOutlinerData)) 
+		, Filters(MakeShareable(new FSceneOutlinerFilters))
+		, InteractiveFilters(MakeShareable(new FSceneOutlinerFilters))
+	{}
 
 	/** SSceneOutliner destructor */
 	virtual ~SSceneOutliner();
@@ -597,7 +601,10 @@ private:
 				OnItemPassesFilters(Temporary);
 			}
 
-		bPassesFilters &= SearchBoxFilter->PassesFilter(Temporary);
+			if (SearchBoxFilter)
+			{
+				bPassesFilters &= SearchBoxFilter->PassesFilter(Temporary);
+			}
 			
 		if (bForce || bPassesFilters)
 		{
@@ -637,7 +644,12 @@ private:
 
 	bool PassesAllFilters(const FSceneOutlinerTreeItemPtr& Item) const
 	{
-		bool bPassesFilters = SearchBoxFilter->PassesFilter(*Item);
+		bool bPassesFilters = false;
+		
+		if (SearchBoxFilter)
+		{
+			bPassesFilters = SearchBoxFilter->PassesFilter(*Item);
+		}
 
 		if(FilterCollection)
 		{
