FActorMode: Make IsActorDisplayable virtual

diff --git a/Engine/Source/Editor/SceneOutliner/Public/ActorMode.h b/Engine/Source/Editor/SceneOutliner/Public/ActorMode.h
--- a/Engine/Source/Editor/SceneOutliner/Public/ActorMode.h
+++ b/Engine/Source/Editor/SceneOutliner/Public/ActorMode.h
@@ -108,7 +108,7 @@ private:
 
 protected:
 	void SynchronizeActorSelection();
-	bool IsActorDisplayable(const AActor* InActor) const;
+	virtual bool IsActorDisplayable(const AActor* InActor) const;
 
 	/** Set the Scene Outliner attached to this mode as the most recently used outliner in the Level Editor */
 	void SetAsMostRecentOutliner() const;
