[SceneOutliner] ActorBrowsingMode only do a full refresh if a Actor Editor Context filter is active

#rnx

diff --git a/Engine/Source/Editor/SceneOutliner/Private/ActorBrowsingMode.cpp b/Engine/Source/Editor/SceneOutliner/Private/ActorBrowsingMode.cpp
--- a/Engine/Source/Editor/SceneOutliner/Private/ActorBrowsingMode.cpp
+++ b/Engine/Source/Editor/SceneOutliner/Private/ActorBrowsingMode.cpp
@@ -404,7 +404,11 @@ FSlateColor FActorBrowsingMode::GetStatusTextColor() const
 
 void FActorBrowsingMode::OnActorEditorContextSubsystemChanged()
 {
-	SceneOutliner->FullRefresh();
+	// For performance reasons avoid doing full refresh if we don't have an active filter relying on the current Actor Editor Context
+	if (const FActorBrowsingModeConfig* Settings = GetConstConfig(); Settings && (Settings->bShowOnlyActorsInCurrentDataLayers || Settings->bShowOnlyActorsInCurrentContentBundle))
+	{
+		SceneOutliner->FullRefresh();
+	}
 }
 
 void FActorBrowsingMode::OnToggleAlwaysFrameSelection()
