Early out of SSceneOutliner::PasteFoldersEnd if there's no work to do

This avoids it creating an empty transaction when duplicating an actor

#rnx

diff --git a/Engine/Source/Editor/SceneOutliner/Private/SSceneOutliner.cpp b/Engine/Source/Editor/SceneOutliner/Private/SSceneOutliner.cpp
--- a/Engine/Source/Editor/SceneOutliner/Private/SSceneOutliner.cpp
+++ b/Engine/Source/Editor/SceneOutliner/Private/SSceneOutliner.cpp
@@ -1700,6 +1700,19 @@ void SSceneOutliner::PasteFoldersBegin(TArray<FName> InFolders)
 
 void SSceneOutliner::PasteFoldersEnd()
 {
+	ON_SCOPE_EXIT
+	{
+		CacheFoldersEdit.Reset();
+		CacheFoldersEditRootObject = FFolder::GetInvalidRootObject();
+		CacheFolderMap.Reset();
+		CachePasteFolderExistingChildrenMap.Reset();
+	};
+
+	if (CacheFoldersEdit.IsEmpty())
+	{
+		return;
+	}
+
 	const FScopedTransaction Transaction(NSLOCTEXT("UnrealEd", "PasteItems", "Paste Items"));
 
 	// Create new folder
@@ -1745,10 +1758,6 @@ void SSceneOutliner::PasteFoldersEnd()
 		}
 	}
 
-	CacheFoldersEdit.Reset();
-	CacheFoldersEditRootObject = FFolder::GetInvalidRootObject();
-	CacheFolderMap.Reset();
-	CachePasteFolderExistingChildrenMap.Reset();
 	FullRefresh();
 }
 
