FActorMode: Add an internal function that is called when actors are attached via drag/drop in the Outliner

diff --git a/Engine/Source/Editor/SceneOutliner/Private/ActorMode.cpp b/Engine/Source/Editor/SceneOutliner/Private/ActorMode.cpp
--- a/Engine/Source/Editor/SceneOutliner/Private/ActorMode.cpp
+++ b/Engine/Source/Editor/SceneOutliner/Private/ActorMode.cpp
@@ -947,24 +947,26 @@ void FActorMode::OnDrop(ISceneOutlinerTreeItem& DropTarget, const FSceneOutliner
 			}
 			else
 			{
-				auto PerformAttachment = [](FName SocketName, TWeakObjectPtr<AActor> Parent, const TArray<TWeakObjectPtr<AActor>> NewAttachments)
+				auto PerformAttachment = [this](FName SocketName, TWeakObjectPtr<AActor> Parent, const TArray<TWeakObjectPtr<AActor>> NewAttachments)
 				{
 					AActor* ParentActor = Parent.Get();
 					if (ParentActor)
 					{
+						TArray<TWeakObjectPtr<AActor>> AttachedActors;
 						// modify parent and child
 						const FScopedTransaction Transaction(LOCTEXT("UndoAction_PerformAttachment", "Attach actors"));
 
 						// Attach each child
-						bool bAttached = false;
 						for (auto& Child : NewAttachments)
 						{
 							AActor* ChildActor = Child.Get();
 							if (GEditor->CanParentActors(ParentActor, ChildActor))
 							{
 								GEditor->ParentActors(ParentActor, ChildActor, SocketName);
+								AttachedActors.Add(ChildActor);
 							}
 						}
+						OnActorsAttached(ParentActor, AttachedActors);
 					}
 				};
 
diff --git a/Engine/Source/Editor/SceneOutliner/Public/ActorMode.h b/Engine/Source/Editor/SceneOutliner/Public/ActorMode.h
--- a/Engine/Source/Editor/SceneOutliner/Public/ActorMode.h
+++ b/Engine/Source/Editor/SceneOutliner/Public/ActorMode.h
@@ -114,6 +114,9 @@ protected:
 	void SetAsMostRecentOutliner() const;
 
 	virtual TUniquePtr<ISceneOutlinerHierarchy> CreateHierarchy() override;
+
+	// Called when actors are attached to a parent actor via drag and drop
+	virtual void OnActorsAttached(AActor* ParentActor, TArray<TWeakObjectPtr<AActor>> ChildActors) const {}
 	
 	FFolder GetWorldDefaultRootFolder() const;
 protected:
