SceneOutliner: Add VeryVerbose logs when a refresh is requested with a delay

diff --git a/Engine/Source/Editor/SceneOutliner/Private/ActorHierarchy.cpp b/Engine/Source/Editor/SceneOutliner/Private/ActorHierarchy.cpp
--- a/Engine/Source/Editor/SceneOutliner/Private/ActorHierarchy.cpp
+++ b/Engine/Source/Editor/SceneOutliner/Private/ActorHierarchy.cpp
@@ -29,6 +29,7 @@
 #include "WorldPartition/WorldPartitionSubsystem.h"
 #include "ActorFolder.h"
 #include "ActorMode.h"
+#include "SSceneOutliner.h"
 
 TUniquePtr<FActorHierarchy> FActorHierarchy::Create(ISceneOutlinerMode* Mode, const TWeakObjectPtr<UWorld>& World)
 {
@@ -951,6 +952,8 @@ void FActorHierarchy::OnLevelActorFolderChanged(const AActor* InActor, FName Old
 {
 	if (Mode->ShouldShowFolders() && RepresentingWorld.Get() == InActor->GetWorld())
 	{
+		UE_LOG(LogSceneOutliner, VeryVerbose, TEXT("OnLevelActorFolderChanged refresh requested by Actor: %s"), *InActor->GetActorLabel());
+
 		FSceneOutlinerHierarchyChangedData EventData;
 		EventData.Type = FSceneOutlinerHierarchyChangedData::Moved;
 		EventData.ItemIDs.Add(FSceneOutlinerTreeItemID(InActor));
diff --git a/Engine/Source/Editor/SceneOutliner/Private/SSceneOutliner.cpp b/Engine/Source/Editor/SceneOutliner/Private/SSceneOutliner.cpp
--- a/Engine/Source/Editor/SceneOutliner/Private/SSceneOutliner.cpp
+++ b/Engine/Source/Editor/SceneOutliner/Private/SSceneOutliner.cpp
@@ -32,13 +32,13 @@
 #include "Widgets/Input/SMultiLineEditableTextBox.h"
 #include "DetailLayoutBuilder.h"
 
-DEFINE_LOG_CATEGORY_STATIC(LogSceneOutliner, Log, All);
+DEFINE_LOG_CATEGORY(LogSceneOutliner);
 
 static float GSceneOutlinerProcessingBudgetPerFrame = 5.0f;
 static FAutoConsoleVariableRef CVarGuardBandMultiplier(
 	TEXT("SceneOutliner.ProcessingBudgetPerFrame"),
 	GSceneOutlinerProcessingBudgetPerFrame,
-	TEXT("Maximum time in milliseconds to spend processing operations per frame"));
+	TEXT("Maximum time in mZilliseconds to spend processing operations per frame"));
 
 
 #define LOCTEXT_NAMESPACE "SSceneOutliner"
@@ -671,6 +671,7 @@ ESelectionMode::Type SSceneOutliner::GetSelectionMode() const
 
 void SSceneOutliner::Refresh()
 {
+	UE_LOG(LogSceneOutliner, VeryVerbose, TEXT("Refresh requested, current refresh delay: %f"), UIRefreshDelay);
 	bNeedsRefresh = true;
 }
 
@@ -939,6 +940,7 @@ FSceneOutlinerTreeItemPtr SSceneOutliner::GetTreeItem(FSceneOutlinerTreeItemID I
 
 void SSceneOutliner::SetNextUIRefreshDelay(float InDelay)
 {
+	UE_LOG(LogSceneOutliner, VeryVerbose, TEXT("UI refresh delay set to %f"), UIRefreshDelay);
 	UIRefreshDelay = InDelay;
 }
 
@@ -2073,6 +2075,9 @@ void SSceneOutliner::OnHierarchyChangedEvent(FSceneOutlinerHierarchyChangedData
 				PendingTreeItemMap_Removal.Add(TreeItemID, Item->ToSharedRef());
 			}
 		}
+		
+		UE_LOG(LogSceneOutliner, VeryVerbose, TEXT("Refresh requested by FSceneOutlinerHierarchyChangedData::Removed"));
+
 		Refresh();
 	}
 	else if (Event.Type == FSceneOutlinerHierarchyChangedData::Moved)
@@ -2098,6 +2103,8 @@ void SSceneOutliner::OnHierarchyChangedEvent(FSceneOutlinerHierarchyChangedData
 				PendingOperations.Emplace(SceneOutliner::FPendingTreeOperation::Moved, TreeItemPtr.ToSharedRef());
 			}
 		}
+		UE_LOG(LogSceneOutliner, VeryVerbose, TEXT("Refresh requested by FSceneOutlinerHierarchyChangedData::Moved"));
+
 		Refresh();
 	}
 	else if (Event.Type == FSceneOutlinerHierarchyChangedData::FolderMoved)
@@ -2122,6 +2129,8 @@ void SSceneOutliner::OnHierarchyChangedEvent(FSceneOutlinerHierarchyChangedData
 				PendingOperations.Emplace(SceneOutliner::FPendingTreeOperation::Moved, Item.ToSharedRef());
 			}
 		}
+		
+		UE_LOG(LogSceneOutliner, VeryVerbose, TEXT("Refresh requested by FSceneOutlinerHierarchyChangedData::FolderMoved"));
 		Refresh();
 	}
 	else if (Event.Type == FSceneOutlinerHierarchyChangedData::FullRefresh)
@@ -2371,6 +2380,7 @@ void SSceneOutliner::Tick(const FGeometry& AllottedGeometry, const double InCurr
 			Pair.Value->Flags.bChildrenRequireSort = true;
 		}
 
+		UE_LOG(LogSceneOutliner, VeryVerbose, TEXT("Sort completed, UI refresh pending. UIRefreshDelay = %f"), UIRefreshDelay);
 		bNeedsUIRefresh = true;
 		bSortDirty = false;
 	}
@@ -2383,6 +2393,7 @@ void SSceneOutliner::Tick(const FGeometry& AllottedGeometry, const double InCurr
 		// if we are currently pending a sort, don't refresh until that is completed
 		if(UIRefreshDelay <= 0.0f && !bSortDirty)
 		{
+			UE_LOG(LogSceneOutliner, VeryVerbose, TEXT("UI Refresh executed"));
 			OutlinerTreeView->RequestTreeRefresh();
 			bNeedsUIRefresh = false;
 			UIRefreshDelay = 0.0f;
diff --git a/Engine/Source/Editor/SceneOutliner/Public/SSceneOutliner.h b/Engine/Source/Editor/SceneOutliner/Public/SSceneOutliner.h
--- a/Engine/Source/Editor/SceneOutliner/Public/SSceneOutliner.h
+++ b/Engine/Source/Editor/SceneOutliner/Public/SSceneOutliner.h
@@ -42,6 +42,8 @@ template<typename FilterType> class SFilterBar;
 template<typename ItemType> class STreeView;
 class SFilterSearchBox;
 
+SCENEOUTLINER_API DECLARE_LOG_CATEGORY_EXTERN(LogSceneOutliner, Log, All);
+
 /**
  * Scene Outliner definition
  * Note the Scene Outliner is also called the World Outliner
