SceneOutliner: Add more logging around refreshes

diff --git a/Engine/Source/Editor/SceneOutliner/Private/ActorBrowsingMode.cpp b/Engine/Source/Editor/SceneOutliner/Private/ActorBrowsingMode.cpp
--- a/Engine/Source/Editor/SceneOutliner/Private/ActorBrowsingMode.cpp
+++ b/Engine/Source/Editor/SceneOutliner/Private/ActorBrowsingMode.cpp
@@ -993,6 +993,8 @@ void FActorBrowsingMode::OnItemRemoved(FSceneOutlinerTreeItemPtr Item)
 
 void FActorBrowsingMode::OnComponentsUpdated()
 {
+	UE_LOG(LogSceneOutliner, VeryVerbose, TEXT("FullRefresh requested by FActorBrowsingMode::OnComponentsUpdated"));
+
 	SceneOutliner->FullRefresh();
 }
 
diff --git a/Engine/Source/Editor/SceneOutliner/Private/ActorHierarchy.cpp b/Engine/Source/Editor/SceneOutliner/Private/ActorHierarchy.cpp
--- a/Engine/Source/Editor/SceneOutliner/Private/ActorHierarchy.cpp
+++ b/Engine/Source/Editor/SceneOutliner/Private/ActorHierarchy.cpp
@@ -701,6 +701,7 @@ void FActorHierarchy::OnWorldPartitionCreated(UWorld* InWorld)
 {
 	if (RepresentingWorld.Get() == InWorld)
 	{
+		UE_LOG(LogSceneOutliner, VeryVerbose, TEXT("OnWorldPartitionCreated Full Refresh Event"));
 		FullRefreshEvent();
 	}
 }
@@ -821,26 +822,32 @@ void FActorHierarchy::OnActorDescInstanceRemoved(FWorldPartitionActorDescInstanc
 
 void FActorHierarchy::OnWorldPartitionInitialized(UWorldPartition* InWorldPartition)
 {
+	UE_LOG(LogSceneOutliner, VeryVerbose, TEXT("OnWorldPartitionInitialized Full Refresh Event"));
+
 	FullRefreshEvent();
 }
 
 void FActorHierarchy::OnWorldPartitionUninitialized(UWorldPartition* InWorldPartition)
 {
+	UE_LOG(LogSceneOutliner, VeryVerbose, TEXT("OnWorldPartitionUninitialized Full Refresh Event"));
 	FullRefreshEvent();
 }
 
 void FActorHierarchy::OnComponentsUpdated()
 {
+	UE_LOG(LogSceneOutliner, VeryVerbose, TEXT("OnComponentsUpdated Full Refresh Event"));
 	FullRefreshEvent();
 }
 
 void FActorHierarchy::OnLevelActorListChanged()
 {
+	UE_LOG(LogSceneOutliner, VeryVerbose, TEXT("OnLevelActorListChanged Full Refresh Event"));
 	FullRefreshEvent();
 }
 
 void FActorHierarchy::OnActorFoldersUpdatedEvent(ULevel* InLevel)
 {
+	UE_LOG(LogSceneOutliner, VeryVerbose, TEXT("OnActorFoldersUpdatedEvent Full Refresh Event"));
 	FullRefreshEvent();
 }
 
diff --git a/Engine/Source/Editor/SceneOutliner/Private/SSceneOutliner.cpp b/Engine/Source/Editor/SceneOutliner/Private/SSceneOutliner.cpp
--- a/Engine/Source/Editor/SceneOutliner/Private/SSceneOutliner.cpp
+++ b/Engine/Source/Editor/SceneOutliner/Private/SSceneOutliner.cpp
@@ -671,7 +671,7 @@ ESelectionMode::Type SSceneOutliner::GetSelectionMode() const
 
 void SSceneOutliner::Refresh()
 {
-	UE_LOG(LogSceneOutliner, VeryVerbose, TEXT("Refresh requested, current refresh delay: %f"), UIRefreshDelay);
+	UE_LOG(LogSceneOutliner, VeryVerbose, TEXT("Refresh requested, current refresh delay: %f SceneOutliner = %p"), UIRefreshDelay, this);
 	bNeedsRefresh = true;
 }
 
@@ -940,7 +940,7 @@ FSceneOutlinerTreeItemPtr SSceneOutliner::GetTreeItem(FSceneOutlinerTreeItemID I
 
 void SSceneOutliner::SetNextUIRefreshDelay(float InDelay)
 {
-	UE_LOG(LogSceneOutliner, VeryVerbose, TEXT("UI refresh delay set to %f"), UIRefreshDelay);
+	UE_LOG(LogSceneOutliner, VeryVerbose, TEXT("UI refresh delay set to %f, previously %f Scene Outliner = %p"), InDelay, UIRefreshDelay, this);
 	UIRefreshDelay = InDelay;
 }
 
@@ -2135,6 +2135,7 @@ void SSceneOutliner::OnHierarchyChangedEvent(FSceneOutlinerHierarchyChangedData
 	}
 	else if (Event.Type == FSceneOutlinerHierarchyChangedData::FullRefresh)
 	{
+		UE_LOG(LogSceneOutliner, VeryVerbose, TEXT("Full Refresh requested by FSceneOutlinerHierarchyChangedData::FullRefresh"));
 		FullRefresh();
 	}
 }
@@ -2144,6 +2145,7 @@ void SSceneOutliner::PostUndo(bool bSuccess)
 	// Refresh our tree in case any changes have been made to the scene that might effect our list
 	if( !bIsReentrant )
 	{
+		UE_LOG(LogSceneOutliner, VeryVerbose, TEXT("FullRefresh requested by SSceneOutliner::PostUndo"));
 		FullRefresh();
 	}
 }
@@ -2196,6 +2198,7 @@ void SSceneOutliner::OnAssetReloaded(const EPackageReloadPhase InPackageReloadPh
 	if (InPackageReloadPhase == EPackageReloadPhase::PostBatchPostGC)
 	{
 		// perhaps overkill but a simple Refresh() doesn't appear to work.
+		UE_LOG(LogSceneOutliner, VeryVerbose, TEXT("FullRefresh requested by SSceneOutliner::OnAssetReloaded"));
 		FullRefresh();
 	}
 }
@@ -2380,7 +2383,7 @@ void SSceneOutliner::Tick(const FGeometry& AllottedGeometry, const double InCurr
 			Pair.Value->Flags.bChildrenRequireSort = true;
 		}
 
-		UE_LOG(LogSceneOutliner, VeryVerbose, TEXT("Sort completed, UI refresh pending. UIRefreshDelay = %f"), UIRefreshDelay);
+		UE_LOG(LogSceneOutliner, VeryVerbose, TEXT("Sort completed, UI refresh pending. UIRefreshDelay = %f SceneOutliner = %p"), UIRefreshDelay, this);
 		bNeedsUIRefresh = true;
 		bSortDirty = false;
 	}
@@ -2388,12 +2391,13 @@ void SSceneOutliner::Tick(const FGeometry& AllottedGeometry, const double InCurr
 	// If we are pending a UI refresh
 	if(bNeedsUIRefresh)
 	{
+		UE_LOG(LogSceneOutliner, VeryVerbose, TEXT("Current UIRefreshDelay = %f, DeltaTime = %f "), UIRefreshDelay, InDeltaTime);
 		UIRefreshDelay -= InDeltaTime;
 
 		// if we are currently pending a sort, don't refresh until that is completed
 		if(UIRefreshDelay <= 0.0f && !bSortDirty)
 		{
-			UE_LOG(LogSceneOutliner, VeryVerbose, TEXT("UI Refresh executed"));
+			UE_LOG(LogSceneOutliner, VeryVerbose, TEXT("UI Refresh executed SceneOutliner = %p"), this);
 			OutlinerTreeView->RequestTreeRefresh();
 			bNeedsUIRefresh = false;
 			UIRefreshDelay = 0.0f;
