Outliner: Modify UIRefreshDelay logic to delay both the UI refresh and the internal item sorting

diff --git a/Engine/Source/Editor/SceneOutliner/Private/SSceneOutliner.cpp b/Engine/Source/Editor/SceneOutliner/Private/SSceneOutliner.cpp
--- a/Engine/Source/Editor/SceneOutliner/Private/SSceneOutliner.cpp
+++ b/Engine/Source/Editor/SceneOutliner/Private/SSceneOutliner.cpp
@@ -2372,43 +2372,52 @@ void SSceneOutliner::Tick(const FGeometry& AllottedGeometry, const double InCurr
 			}
 		}
 	}
-	SortOutlinerTimer -= InDeltaTime;
 
-	// Delay sorting when in PIE
-	if (bSortDirty && (GEditor->PlayWorld == nullptr || SortOutlinerTimer <= 0))
+	// If we are pending a sort
+	if (bSortDirty)
 	{
-		SortItems(RootTreeItems);
-		for (const auto& Pair : TreeItemMap)
+		bool bExecuteSort = false;
+		
+		SortOutlinerTimer -= InDeltaTime;
+		UIRefreshDelay -= InDeltaTime;
+
+		// If we are in PIE, we will use SortOutlinerTimer as the min delay between sorts
+		if (GEditor->PlayWorld)
 		{
-			Pair.Value->Flags.bChildrenRequireSort = true;
+			if (SortOutlinerTimer <= 0)
+			{
+				bExecuteSort = true;
+			}
+		}
+		// In editor worlds, we will check if there is any user specific UIRefreshDelay
+		else
+		{
+			if (UIRefreshDelay <= 0)
+			{
+				bExecuteSort = true;
+			}
 		}
 
-		UE_LOG(LogSceneOutliner, VeryVerbose, TEXT("Sort completed, UI refresh pending. UIRefreshDelay = %f SceneOutliner = %p"), UIRefreshDelay, this);
-		bNeedsUIRefresh = true;
-		bSortDirty = false;
-	}
-
-	// If we are pending a UI refresh
-	if(bNeedsUIRefresh)
-	{
-		UE_LOG(LogSceneOutliner, VeryVerbose, TEXT("Current UIRefreshDelay = %f, DeltaTime = %f "), UIRefreshDelay, InDeltaTime);
-		UIRefreshDelay -= InDeltaTime;
-
-		// if we are currently pending a sort, don't refresh until that is completed
-		if(UIRefreshDelay <= 0.0f && !bSortDirty)
+		if (bExecuteSort)
 		{
-			UE_LOG(LogSceneOutliner, VeryVerbose, TEXT("UI Refresh executed SceneOutliner = %p"), this);
+			UE_LOG(LogSceneOutliner, VeryVerbose, TEXT("Sort Executed"));
+
+			SortItems(RootTreeItems);
+			
+			for (const auto& Pair : TreeItemMap)
+			{
+				Pair.Value->Flags.bChildrenRequireSort = true;
+			}
+			
 			OutlinerTreeView->RequestTreeRefresh();
-			bNeedsUIRefresh = false;
+			bSortDirty = false;
+
+			// Reset both timers
+			SortOutlinerTimer = SCENE_OUTLINER_RESORT_TIMER;
 			UIRefreshDelay = 0.0f;
 		}
 	}
-
-	if (SortOutlinerTimer <= 0)
-	{
-		SortOutlinerTimer = SCENE_OUTLINER_RESORT_TIMER;
-	}
-
+	
 	if (bSelectionDirty)
 	{
 		Mode->SynchronizeSelection();
diff --git a/Engine/Source/Editor/SceneOutliner/Public/SSceneOutliner.h b/Engine/Source/Editor/SceneOutliner/Public/SSceneOutliner.h
--- a/Engine/Source/Editor/SceneOutliner/Public/SSceneOutliner.h
+++ b/Engine/Source/Editor/SceneOutliner/Public/SSceneOutliner.h
@@ -963,12 +963,9 @@ private:
 	/** Reentrancy guard */
 	bool bIsReentrant;
 
-	/** The delay (in seconds) before a UI refresh is executed after being requested. 0 means the next refresh will be immediate */
+	/** The delay (in seconds) before a UI sort and refresh is executed after being requested. 0 means the next refresh will be immediate */
 	float UIRefreshDelay = 0.0f;
 
-	/** If true, the TreeView UI is pending a refresh */
-	bool bNeedsUIRefresh = false;
-
 	/* Widget containing the filtering text box */
 	TSharedPtr< SFilterSearchBox > FilterTextBoxWidget;
 
