SourceControl: Fixed checking out via the SceneOutliner checking out other packages (such as player spawn pads or water bodies) as well.

#rnx

diff --git a/Engine/Source/Editor/SceneOutliner/Private/SceneOutlinerSCCHandler.cpp b/Engine/Source/Editor/SceneOutliner/Private/SceneOutlinerSCCHandler.cpp
--- a/Engine/Source/Editor/SceneOutliner/Private/SceneOutlinerSCCHandler.cpp
+++ b/Engine/Source/Editor/SceneOutliner/Private/SceneOutlinerSCCHandler.cpp
@@ -382,8 +382,8 @@ void FSceneOutlinerSCCHandler::ExecuteSCCRefresh()
 
 void FSceneOutlinerSCCHandler::ExecuteSCCCheckOut()
 {
-	TArray<UPackage*> PackagesToCheckOut;
-	GetSelectedPackages(PackagesToCheckOut);
+	TArray<FString> PackagesToCheckOut;
+	GetSelectedPackageNames(PackagesToCheckOut);
 
 	if ( PackagesToCheckOut.Num() > 0 )
 	{
