SourceControl: Fixed reverting via the SceneOutliner checking out other packages (such as player spawn pads or water bodies). This regression was introduced in #33762883.

#rnx

diff --git a/Engine/Source/Editor/SceneOutliner/Private/SceneOutlinerSCCHandler.cpp b/Engine/Source/Editor/SceneOutliner/Private/SceneOutlinerSCCHandler.cpp
--- a/Engine/Source/Editor/SceneOutliner/Private/SceneOutlinerSCCHandler.cpp
+++ b/Engine/Source/Editor/SceneOutliner/Private/SceneOutlinerSCCHandler.cpp
@@ -356,10 +356,15 @@ void FSceneOutlinerSCCHandler::GetSelectedPackages(TArray<UPackage*>& OutPackage
 			continue;
 		}
 
-		UPackage* Package = SourceControl->LoadPackage();
-		if (Package != nullptr)
+		if (UPackage* Package = SourceControl->GetPackage())
 		{
 			OutPackages.Add(Package);
+			continue;
+		}
+		if (UPackage* Package = SourceControl->LoadPackage())
+		{
+			OutPackages.Add(Package);
+			continue;
 		}
 	}
 }
