Outliner: Move the code modifying item expansion state in Populate() to when the actual sort is going to execute in Tick() because updating item expansion state causes a list refresh that would ignore the UIRefreshDelay

diff --git a/Engine/Source/Editor/SceneOutliner/Private/SSceneOutliner.cpp b/Engine/Source/Editor/SceneOutliner/Private/SSceneOutliner.cpp
--- a/Engine/Source/Editor/SceneOutliner/Private/SSceneOutliner.cpp
+++ b/Engine/Source/Editor/SceneOutliner/Private/SSceneOutliner.cpp
@@ -775,8 +775,6 @@ void SSceneOutliner::Populate()
 	bool bFinalSort = false;
 	if (PendingOperations.Num() == 0)
 	{
-		// Update expansion state based on item states
-		SetParentsExpansionState();
 		// When done processing a FullRefresh Scroll to First item in selection as it may have been
 		// scrolled out of view by the Refresh
 		if (bProcessingFullRefresh)
@@ -1077,21 +1075,6 @@ void SSceneOutliner::AddUnfilteredItemToTree(FSceneOutlinerTreeItemRef Item)
 	Mode->OnItemAdded(Item);
 }
 
-void SSceneOutliner::SetParentsExpansionState() const
-{
-	// If we have an active search filter, auto expand parents of items that passes the filter so they appear automatically in the outliner
-	bForceParentItemsExpanded = !SearchBoxFilter->GetRawFilterText().IsEmpty();
-
-	for (const auto& Pair : TreeItemMap)
-	{
-		auto& Item = Pair.Value;
-		if (Item->GetChildren().Num())
-		{
-			OutlinerTreeView->SetItemExpansion(Item, bForceParentItemsExpanded || Item->Flags.bIsExpanded);
-		}
-	}
-}
-
 void SSceneOutliner::PopulateSearchStrings(const ISceneOutlinerTreeItem& Item, TArray< FString >& OutSearchStrings) const
 {
 	for (const auto& Pair : Columns)
@@ -2403,10 +2386,18 @@ void SSceneOutliner::Tick(const FGeometry& AllottedGeometry, const double InCurr
 			UE_LOG(LogSceneOutliner, VeryVerbose, TEXT("Sort Executed"));
 
 			SortItems(RootTreeItems);
+
+			// Also update expansion state based on item states
+			// This is done here because updating expansion state causes a refresh, so we want to make sure that does not ignore any UIRefreshDelay
+			bForceParentItemsExpanded = !SearchBoxFilter->GetRawFilterText().IsEmpty();
 			
 			for (const auto& Pair : TreeItemMap)
 			{
 				Pair.Value->Flags.bChildrenRequireSort = true;
+				if (Pair.Value->GetChildren().Num())
+				{
+					OutlinerTreeView->SetItemExpansion(Pair.Value, bForceParentItemsExpanded || Pair.Value->Flags.bIsExpanded);
+				}
 			}
 			
 			OutlinerTreeView->RequestTreeRefresh();
diff --git a/Engine/Source/Editor/SceneOutliner/Public/SSceneOutliner.h b/Engine/Source/Editor/SceneOutliner/Public/SSceneOutliner.h
--- a/Engine/Source/Editor/SceneOutliner/Public/SSceneOutliner.h
+++ b/Engine/Source/Editor/SceneOutliner/Public/SSceneOutliner.h
@@ -931,9 +931,6 @@ private:
 	/** Cached expansion state info */
 	mutable FParentsExpansionState CachedExpansionStateInfo;
 
-	/** Updates the expansion state of parent items after a repopulate */
-	void SetParentsExpansionState() const;
-
 private:
 
 	/** True if the outliner needs to be repopulated at the next appropriate opportunity, usually because our
