Add an EditorCanDetachFrom delegate to AActor that is checked before an actor is detached from it's parent

diff --git a/Engine/Source/Editor/LevelEditor/Private/LevelEditorActions.cpp b/Engine/Source/Editor/LevelEditor/Private/LevelEditorActions.cpp
--- a/Engine/Source/Editor/LevelEditor/Private/LevelEditorActions.cpp
+++ b/Engine/Source/Editor/LevelEditor/Private/LevelEditorActions.cpp
@@ -1565,7 +1565,22 @@ void FLevelEditorActionCallbacks::DetachActor_Clicked()
 bool FLevelEditorActionCallbacks::DetachActor_CanExecute()
 {
 	FSelectedActorInfo SelectionInfo = AssetSelectionUtils::GetSelectedActorInfo();
-	return SelectionInfo.NumSelected > 0 && SelectionInfo.bHaveAttachedActor;
+	
+	if (SelectionInfo.NumSelected > 0 && SelectionInfo.bHaveAttachedActor)
+	{
+		TArray<AActor*> SelectedActors;
+		GEditor->GetSelectedActors()->GetSelectedObjects<AActor>(SelectedActors);
+		FText DetachErrorMsg;
+		for (AActor* SelectedActor : SelectedActors)
+		{
+			if (!SelectedActor->EditorCanDetachFrom(SelectedActor->GetSceneOutlinerParent(), DetachErrorMsg))
+			{
+				return false;
+			}
+		}
+		return true;
+	}
+	return false;
 }
 
 void FLevelEditorActionCallbacks::AttachSelectedActors()
diff --git a/Engine/Source/Editor/SceneOutliner/Private/ActorMode.cpp b/Engine/Source/Editor/SceneOutliner/Private/ActorMode.cpp
--- a/Engine/Source/Editor/SceneOutliner/Private/ActorMode.cpp
+++ b/Engine/Source/Editor/SceneOutliner/Private/ActorMode.cpp
@@ -613,6 +613,16 @@ FSceneOutlinerDragValidationInfo FActorMode::ValidateDrop(const ISceneOutlinerTr
 		const FText ActorLabel = FText::FromString(ActorTarget->GetActorLabel());
 		if (bDraggedOntoAttachmentParent)
 		{
+			for (const auto& DragActorPtr : DragActors)
+			{
+				AActor* DragActor = DragActorPtr.Get();
+				if (!DragActor->EditorCanDetachFrom(DragActor->GetSceneOutlinerParent(), AttachErrorMsg))
+				{
+					// Cannot detach from parent into root
+					return FSceneOutlinerDragValidationInfo(ESceneOutlinerDropCompatibility::IncompatibleGeneric, AttachErrorMsg); 
+				}
+			}
+
 			if (DragActors.Num() == 1)
 			{
 				return FSceneOutlinerDragValidationInfo(ESceneOutlinerDropCompatibility::CompatibleDetach, ActorLabel);
@@ -784,6 +794,14 @@ FSceneOutlinerDragValidationInfo FActorMode::ValidateDrop(const ISceneOutlinerTr
 
 					return FSceneOutlinerDragValidationInfo(ESceneOutlinerDropCompatibility::IncompatibleGeneric, Text);
 				}
+				else if (Actor->GetSceneOutlinerParent())
+				{
+					FText DetachErrorMsg;
+					if (!Actor->EditorCanDetachFrom(Actor->GetSceneOutlinerParent(), DetachErrorMsg))
+					{
+						return FSceneOutlinerDragValidationInfo(ESceneOutlinerDropCompatibility::IncompatibleGeneric, DetachErrorMsg);
+					}
+				}
 			}
 		}
 
diff --git a/Engine/Source/Runtime/Engine/Classes/GameFramework/Actor.h b/Engine/Source/Runtime/Engine/Classes/GameFramework/Actor.h
--- a/Engine/Source/Runtime/Engine/Classes/GameFramework/Actor.h
+++ b/Engine/Source/Runtime/Engine/Classes/GameFramework/Actor.h
@@ -2599,6 +2599,9 @@ public:
 	/** Returns true if this actor is allowed to be attached from the given actor */
 	ENGINE_API virtual bool EditorCanAttachFrom(const AActor* InChild, FText& OutReason) const;
 
+	/** Returns true if this actor is allowed to be detached from the given actor */
+	ENGINE_API virtual bool EditorCanDetachFrom(const AActor* InParent, FText& OutReason) const { return true; }
+
 	/** Returns the actor attachement parent that should be used in editor */
 	ENGINE_API virtual AActor* GetSceneOutlinerParent() const;
 
