Outliner: Add an option to automatically collapse other folders on viewport selection

diff --git a/Engine/Source/Editor/SceneOutliner/Private/ActorBrowsingMode.cpp b/Engine/Source/Editor/SceneOutliner/Private/ActorBrowsingMode.cpp
--- a/Engine/Source/Editor/SceneOutliner/Private/ActorBrowsingMode.cpp
+++ b/Engine/Source/Editor/SceneOutliner/Private/ActorBrowsingMode.cpp
@@ -279,7 +279,8 @@ FActorBrowsingMode::FActorBrowsingMode(SSceneOutliner* InSceneOutliner, TWeakObj
 		}), FSceneOutlinerFilter::EDefaultBehaviour::Pass));
 
 	bAlwaysFrameSelection = LocalSettings.bAlwaysFrameSelection;
-
+	bCollapseOutlinerTreeOnNewSelection = LocalSettings.bCollapseOutlinerTreeOnNewSelection;
+	
 	Rebuild();
 }
 
@@ -565,6 +566,18 @@ void FActorBrowsingMode::InitializeViewMenuExtender(TSharedPtr<FExtender> Extend
 			EUserInterfaceActionType::ToggleButton
 		);
 
+		MenuBuilder.AddMenuEntry(
+			LOCTEXT("CollapseOutlinerTreeOnNewSelectionLabel", "Collapse Hierarchy on Viewport Selection"),
+			LOCTEXT("CollapseOutlinerTreeOnNewSelectionTooltip", "When enabled, all items except the one that was just selected are collapsed."),
+			FSlateIcon(),
+			FUIAction(
+				FExecuteAction::CreateRaw(this, &FActorBrowsingMode::OnToggleCollapseOutlinerTreeOnNewSelection),
+				FCanExecuteAction(),
+				FIsActionChecked::CreateRaw(this, &FActorBrowsingMode::CollapseOutlinerTreeOnNewSelection)
+			),
+			NAME_None,
+			EUserInterfaceActionType::ToggleButton
+		);
 
 		MenuBuilder.EndSection();
 
@@ -1248,6 +1261,21 @@ bool FActorBrowsingMode::ShouldUpdateContentWhileInPIEFocused() const
 	return bShouldUpdateContentWhileInPIEFocused;
 }
 
+void FActorBrowsingMode::OnToggleCollapseOutlinerTreeOnNewSelection()
+{
+	if (FActorBrowsingModeConfig* Settings = GetMutableConfig())
+	{
+		Settings->bCollapseOutlinerTreeOnNewSelection = !Settings->bCollapseOutlinerTreeOnNewSelection;
+		bCollapseOutlinerTreeOnNewSelection = Settings->bCollapseOutlinerTreeOnNewSelection;
+		SaveConfig();
+	}
+}
+
+bool FActorBrowsingMode::CollapseOutlinerTreeOnNewSelection() const
+{
+	return bCollapseOutlinerTreeOnNewSelection;
+}
+
 void FActorBrowsingMode::OnFilterTextCommited(FSceneOutlinerItemSelection& Selection, ETextCommit::Type CommitType)
 {
 	// Start batching selection changes
diff --git a/Engine/Source/Editor/SceneOutliner/Private/ActorBrowsingModeSettings.h b/Engine/Source/Editor/SceneOutliner/Private/ActorBrowsingModeSettings.h
--- a/Engine/Source/Editor/SceneOutliner/Private/ActorBrowsingModeSettings.h
+++ b/Engine/Source/Editor/SceneOutliner/Private/ActorBrowsingModeSettings.h
@@ -76,6 +76,10 @@ public:
 	/** True if we want to allow updating content when in PIE with the PIE viewport focused */
 	UPROPERTY()
 	bool bShouldUpdateContentWhileInPIEFocused = false;
+
+	/** True if we want to collapse Outliner tree on new selection, except for the item that was just selected */
+	UPROPERTY()
+	bool bCollapseOutlinerTreeOnNewSelection = false;
 };
 
 UCLASS(EditorConfig="ActorBrowsingMode")
diff --git a/Engine/Source/Editor/SceneOutliner/Private/ActorModeInteractive.cpp b/Engine/Source/Editor/SceneOutliner/Private/ActorModeInteractive.cpp
--- a/Engine/Source/Editor/SceneOutliner/Private/ActorModeInteractive.cpp
+++ b/Engine/Source/Editor/SceneOutliner/Private/ActorModeInteractive.cpp
@@ -66,6 +66,11 @@ void FActorModeInteractive::OnLevelSelectionChanged(UObject* Obj)
 		// Scroll last item into view  - this means if we are multi-selecting, we show newest selection. @TODO Not perfect though
 		if (const AActor* LastSelectedActor = GEditor->GetSelectedActors()->GetBottom<AActor>())
 		{
+			if (bCollapseOutlinerTreeOnNewSelection)
+			{
+				SceneOutliner->CollapseAll();
+			}
+			
 			if (FSceneOutlinerTreeItemPtr TreeItem = SceneOutliner->GetTreeItem(LastSelectedActor, false))
 			{
 				// Only scroll if selection framing is enabled
diff --git a/Engine/Source/Editor/SceneOutliner/Public/ActorBrowsingMode.h b/Engine/Source/Editor/SceneOutliner/Public/ActorBrowsingMode.h
--- a/Engine/Source/Editor/SceneOutliner/Public/ActorBrowsingMode.h
+++ b/Engine/Source/Editor/SceneOutliner/Public/ActorBrowsingMode.h
@@ -160,6 +160,9 @@ private:
 	void OnToggleShouldUpdateContentWhileInPIEFocused();
 	bool ShouldUpdateContentWhileInPIEFocused() const;
 
+	void OnToggleCollapseOutlinerTreeOnNewSelection();
+	bool CollapseOutlinerTreeOnNewSelection() const;
+
 	/**
 	 * Get a mutable version of the ActorBrowser config for setting values.
 	 * @returns		The config for this ActorBrowser.
diff --git a/Engine/Source/Editor/SceneOutliner/Public/ActorMode.h b/Engine/Source/Editor/SceneOutliner/Public/ActorMode.h
--- a/Engine/Source/Editor/SceneOutliner/Public/ActorMode.h
+++ b/Engine/Source/Editor/SceneOutliner/Public/ActorMode.h
@@ -148,4 +148,6 @@ protected:
 	 * the components
 	 */
 	bool bSearchComponentsByActorName;
+	/** Should Outliner Tree get collapsed on selection, except for the item that was just selected */
+	bool bCollapseOutlinerTreeOnNewSelection;
 };
