PR #12408: SSceneOutliner minor performance optimization

diff --git a/Engine/Source/Editor/SceneOutliner/Private/SSceneOutliner.cpp b/Engine/Source/Editor/SceneOutliner/Private/SSceneOutliner.cpp
--- a/Engine/Source/Editor/SceneOutliner/Private/SSceneOutliner.cpp
+++ b/Engine/Source/Editor/SceneOutliner/Private/SSceneOutliner.cpp
@@ -1940,6 +1940,7 @@ void SSceneOutliner::OnGetChildrenForOutlinerTree( FSceneOutlinerTreeItemPtr InP
 
 			// Empty out the children and repopulate them in the correct order
 			InParent->Children.Empty();
+			InParent->Children.Reserve(OutChildren.Num());
 			for (auto& Child : OutChildren)
 			{
 				InParent->Children.Emplace(Child);
diff --git a/Engine/Source/Editor/SceneOutliner/Public/SortHelper.h b/Engine/Source/Editor/SceneOutliner/Public/SortHelper.h
--- a/Engine/Source/Editor/SceneOutliner/Public/SortHelper.h
+++ b/Engine/Source/Editor/SceneOutliner/Public/SortHelper.h
@@ -87,6 +87,7 @@ struct FSceneOutlinerSortHelper
 	void Sort(TArray<FSceneOutlinerTreeItemPtr>& Array)
 	{
 		TArray<FSortPayload> SortData;
+		SortData.Reserve(Array.Num());
 		const auto End = Array.Num();
 		for (int32 Index = 0; Index < End; ++Index)
 		{
