World Partition: when repairing actor descriptor containers, make sure to reset loaders to avoid keeping file handles on the invalid files.

diff --git a/Engine/Source/Editor/SceneOutliner/Private/ActorBrowsingMode.cpp b/Engine/Source/Editor/SceneOutliner/Private/ActorBrowsingMode.cpp
--- a/Engine/Source/Editor/SceneOutliner/Private/ActorBrowsingMode.cpp
+++ b/Engine/Source/Editor/SceneOutliner/Private/ActorBrowsingMode.cpp
@@ -1428,6 +1428,8 @@ void FActorBrowsingMode::RepairErrors() const
 				ActorDescContainerInstance->GetContainer()->ClearInvalidActors();
 			});
 
+			FlushAsyncLoading();
+
 			TArray<FString> ActorFilesToDelete;
 			TArray<FString> ActorFilesToRevert;
 			{
@@ -1439,6 +1441,11 @@ void FActorBrowsingMode::RepairErrors() const
 					FPackagePath PackagePath;
 					if (FPackagePath::TryFromPackageName(InvalidActorAsset.PackageName, PackagePath))
 					{
+						if (UPackage* ExistingPackage = FindObject<UPackage>(nullptr, *InvalidActorAsset.PackageName.ToString()))
+						{
+							ResetLoaders(ExistingPackage);
+						}
+
 						const FString ActorFile = PackagePath.GetLocalFullPath();
 						FSourceControlStatePtr SCState = SourceControlProvider.GetState(ActorFile, EStateCacheUsage::ForceUpdate);
 
