ActorHierarchy: Small optimization to avoid creating the item in the first place for actors that don't want to be listen in the Outliner. This helps because usually the FActorTreeItem is created first and then checked against any filters which causes Actor->GetActorLabel() to be called which can be expensive

Also removes some logs that were added to help detect a perf issue that has been since fixed

diff --git a/Engine/Source/Editor/SceneOutliner/Private/ActorHierarchy.cpp b/Engine/Source/Editor/SceneOutliner/Private/ActorHierarchy.cpp
--- a/Engine/Source/Editor/SceneOutliner/Private/ActorHierarchy.cpp
+++ b/Engine/Source/Editor/SceneOutliner/Private/ActorHierarchy.cpp
@@ -701,7 +701,6 @@ void FActorHierarchy::OnWorldPartitionCreated(UWorld* InWorld)
 {
 	if (RepresentingWorld.Get() == InWorld)
 	{
-		UE_LOG(LogSceneOutliner, VeryVerbose, TEXT("OnWorldPartitionCreated Full Refresh Event"));
 		FullRefreshEvent();
 	}
 }
@@ -710,10 +709,15 @@ void FActorHierarchy::OnLevelActorAdded(AActor* InActor)
 {
 	if (InActor != nullptr && RepresentingWorld.Get() == InActor->GetWorld())
 	{
-		FSceneOutlinerHierarchyChangedData EventData;
-		EventData.Type = FSceneOutlinerHierarchyChangedData::Added;
-		EventData.Items.Add(CreateItemForActor(InActor));
-		HierarchyChangedEvent.Broadcast(EventData);
+		// CreateItemForActor already checks this by calling FActorMode::IsActorDisplayable, but that is after the actor tree item has
+		// already been created so this is a small optimization to avoid creating it and then calling Actor->GetActorLabel() which can be expensive
+		if (InActor->IsListedInSceneOutliner())
+		{
+			FSceneOutlinerHierarchyChangedData EventData;
+			EventData.Type = FSceneOutlinerHierarchyChangedData::Added;
+			EventData.Items.Add(CreateItemForActor(InActor));
+			HierarchyChangedEvent.Broadcast(EventData);
+		}
 	}
 }
 
@@ -822,32 +826,26 @@ void FActorHierarchy::OnActorDescInstanceRemoved(FWorldPartitionActorDescInstanc
 
 void FActorHierarchy::OnWorldPartitionInitialized(UWorldPartition* InWorldPartition)
 {
-	UE_LOG(LogSceneOutliner, VeryVerbose, TEXT("OnWorldPartitionInitialized Full Refresh Event"));
-
 	FullRefreshEvent();
 }
 
 void FActorHierarchy::OnWorldPartitionUninitialized(UWorldPartition* InWorldPartition)
 {
-	UE_LOG(LogSceneOutliner, VeryVerbose, TEXT("OnWorldPartitionUninitialized Full Refresh Event"));
 	FullRefreshEvent();
 }
 
 void FActorHierarchy::OnComponentsUpdated()
 {
-	UE_LOG(LogSceneOutliner, VeryVerbose, TEXT("OnComponentsUpdated Full Refresh Event"));
 	FullRefreshEvent();
 }
 
 void FActorHierarchy::OnLevelActorListChanged()
 {
-	UE_LOG(LogSceneOutliner, VeryVerbose, TEXT("OnLevelActorListChanged Full Refresh Event"));
 	FullRefreshEvent();
 }
 
 void FActorHierarchy::OnActorFoldersUpdatedEvent(ULevel* InLevel)
 {
-	UE_LOG(LogSceneOutliner, VeryVerbose, TEXT("OnActorFoldersUpdatedEvent Full Refresh Event"));
 	FullRefreshEvent();
 }
 
@@ -959,8 +957,6 @@ void FActorHierarchy::OnLevelActorFolderChanged(const AActor* InActor, FName Old
 {
 	if (Mode->ShouldShowFolders() && RepresentingWorld.Get() == InActor->GetWorld())
 	{
-		UE_LOG(LogSceneOutliner, VeryVerbose, TEXT("OnLevelActorFolderChanged refresh requested by Actor: %s"), *InActor->GetActorLabel());
-
 		FSceneOutlinerHierarchyChangedData EventData;
 		EventData.Type = FSceneOutlinerHierarchyChangedData::Moved;
 		EventData.ItemIDs.Add(FSceneOutlinerTreeItemID(InActor));
