Move Folder Name Validation Logic to Helper File.

Written by aditya.ravichandran, I just copy and pasted the code here out of an existing larger review. so the change could be in faster.
Initial Review: https://p4-swarm.epicgames.net/reviews/36192383

diff --git a/Engine/Source/Editor/SceneOutliner/Private/ActorFolderTreeItem.cpp b/Engine/Source/Editor/SceneOutliner/Private/ActorFolderTreeItem.cpp
--- a/Engine/Source/Editor/SceneOutliner/Private/ActorFolderTreeItem.cpp
+++ b/Engine/Source/Editor/SceneOutliner/Private/ActorFolderTreeItem.cpp
@@ -16,6 +16,7 @@
 #include "WorldPersistentFolders.h"
 #include "Engine/Level.h"
 #include "LevelUtils.h"
+#include "SceneOutlinerHelpers.h"
 #include "LevelInstance/LevelInstanceInterface.h"
 #include "LevelInstance/LevelInstanceSubsystem.h"
 
@@ -158,55 +159,7 @@ private:
 			return false;
 		}
 
-		const FText TrimmedLabel = FText::TrimPrecedingAndTrailing(InLabel);
-
-		if (TrimmedLabel.IsEmpty())
-		{
-			OutErrorMessage = LOCTEXT("RenameFailed_LeftBlank", "Names cannot be left blank");
-			return false;
-		}
-
-		if (TrimmedLabel.ToString().Len() >= NAME_SIZE)
-		{
-			FFormatNamedArguments Arguments;
-			Arguments.Add(TEXT("CharCount"), NAME_SIZE);
-			OutErrorMessage = FText::Format(LOCTEXT("RenameFailed_TooLong", "Names must be less than {CharCount} characters long."), Arguments);
-			return false;
-		}
-
-		const FString LabelString = TrimmedLabel.ToString();
-		if (TreeItem->GetLeafName().ToString() == LabelString)
-		{
-			return true;
-		}
-
-		int32 Dummy = 0;
-		if (LabelString.FindChar('/', Dummy) || LabelString.FindChar('\\', Dummy))
-		{
-			OutErrorMessage = LOCTEXT("RenameFailed_InvalidChar", "Folder names cannot contain / or \\.");
-			return false;
-		}
-
-		// Validate that this folder doesn't exist already
-		FFolder Folder = TreeItem->GetFolder();
-		FName NewPath = Folder.GetParent().GetPath();
-		if (NewPath.IsNone())
-		{
-			NewPath = FName(*LabelString);
-		}
-		else
-		{
-			NewPath = FName(*(NewPath.ToString() / LabelString));
-		}
-		FFolder NewFolder(Folder.GetRootObject(), NewPath);
-
-		if (FActorFolders::Get().ContainsFolder(*TreeItem->World.Get(), NewFolder))
-		{
-			OutErrorMessage = LOCTEXT("RenameFailed_AlreadyExists", "A folder with this name already exists at this level");
-			return false;
-		}
-
-		return true;
+		return SceneOutliner::FSceneOutlinerHelpers::ValidateFolderName(TreeItem->GetFolder(), TreeItem->World.Get(), InLabel, OutErrorMessage);
 	}
 
 	void OnLabelCommitted(const FText& InLabel, ETextCommit::Type InCommitInfo)
diff --git a/Engine/Source/Editor/SceneOutliner/Private/SceneOutlinerHelpers.cpp b/Engine/Source/Editor/SceneOutliner/Private/SceneOutlinerHelpers.cpp
--- a/Engine/Source/Editor/SceneOutliner/Private/SceneOutlinerHelpers.cpp
+++ b/Engine/Source/Editor/SceneOutliner/Private/SceneOutlinerHelpers.cpp
@@ -5,6 +5,7 @@
 #include "ActorDescTreeItem.h"
 #include "ActorFolderTreeItem.h"
 #include "ComponentTreeItem.h"
+#include "EditorActorFolders.h"
 #include "EditorClassUtils.h"
 #include "Engine/Blueprint.h"
 #include "Misc/PackageName.h"
@@ -14,6 +15,8 @@
 
 #include "SourceControlHelpers.h"
 
+#define LOCTEXT_NAMESPACE "SceneOutlinerHelpers"
+
 namespace SceneOutliner
 {
 	FString FSceneOutlinerHelpers::GetExternalPackageName(const ISceneOutlinerTreeItem& TreeItem)
@@ -141,4 +144,58 @@ namespace SceneOutliner
 		}
 	}
 
-;}
\ No newline at end of file
+	bool FSceneOutlinerHelpers::ValidateFolderName(const FFolder& InFolder, UWorld* World, const FText& InLabel, FText& OutErrorMessage)
+	{
+		const FText TrimmedLabel = FText::TrimPrecedingAndTrailing(InLabel);
+
+		if (TrimmedLabel.IsEmpty())
+		{
+			OutErrorMessage = LOCTEXT("RenameFailed_LeftBlank", "Names cannot be left blank");
+			return false;
+		}
+
+		if (TrimmedLabel.ToString().Len() >= NAME_SIZE)
+		{
+			FFormatNamedArguments Arguments;
+			Arguments.Add(TEXT("CharCount"), NAME_SIZE);
+			OutErrorMessage = FText::Format(LOCTEXT("RenameFailed_TooLong", "Names must be less than {CharCount} characters long."), Arguments);
+			return false;
+		}
+
+		const FString LabelString = TrimmedLabel.ToString();
+		if (InFolder.GetLeafName().ToString() == LabelString)
+		{
+			return true;
+		}
+
+		int32 Dummy = 0;
+		if (LabelString.FindChar('/', Dummy) || LabelString.FindChar('\\', Dummy))
+		{
+			OutErrorMessage = LOCTEXT("RenameFailed_InvalidChar", "Folder names cannot contain / or \\.");
+			return false;
+		}
+
+		// Validate that this folder doesn't exist already
+		FName NewPath = InFolder.GetParent().GetPath();
+		if (NewPath.IsNone())
+		{
+			NewPath = FName(*LabelString);
+		}
+		else
+		{
+			NewPath = FName(*(NewPath.ToString() / LabelString));
+		}
+		FFolder NewFolder(InFolder.GetRootObject(), NewPath);
+
+		if (World && FActorFolders::Get().ContainsFolder(*World, NewFolder))
+		{
+			OutErrorMessage = LOCTEXT("RenameFailed_AlreadyExists", "A folder with this name already exists at this level");
+			return false;
+		}
+
+		return true;
+	}
+
+}
+
+#undef LOCTEXT_NAMESPACE
\ No newline at end of file
diff --git a/Engine/Source/Editor/SceneOutliner/Public/SceneOutlinerHelpers.h b/Engine/Source/Editor/SceneOutliner/Public/SceneOutlinerHelpers.h
--- a/Engine/Source/Editor/SceneOutliner/Public/SceneOutlinerHelpers.h
+++ b/Engine/Source/Editor/SceneOutliner/Public/SceneOutlinerHelpers.h
@@ -19,5 +19,7 @@ namespace SceneOutliner
 		static TSharedPtr<SWidget> GetClassHyperlink(UObject* InObject);
 		
 		static void PopulateExtraSearchStrings(const ISceneOutlinerTreeItem& TreeItem, TArray< FString >& OutSearchStrings);
+
+		static bool ValidateFolderName(const FFolder& InFolder, UWorld* World, const FText& InLabel, FText& OutErrorMessage);
 	};
 };
