ActorBrowsingMode: Expose a function to derived classes by making it virtual

diff --git a/Engine/Source/Editor/SceneOutliner/Public/ActorBrowsingMode.h b/Engine/Source/Editor/SceneOutliner/Public/ActorBrowsingMode.h
--- a/Engine/Source/Editor/SceneOutliner/Public/ActorBrowsingMode.h
+++ b/Engine/Source/Editor/SceneOutliner/Public/ActorBrowsingMode.h
@@ -111,6 +111,9 @@ public:
 	
 	/** Function called by the Outliner Filter Bar to compare an item with Type Filters*/
 	SCENEOUTLINER_API virtual bool CompareItemWithClassName(SceneOutliner::FilterBarType InItem, const TSet<FTopLevelAssetPath>&) const override;
+	
+	/** Check whether the Rename can be executed */
+	SCENEOUTLINER_API virtual bool CanExecuteRename();
 private:
 	/** Build and up the context menu */
 	TSharedPtr<SWidget> BuildContextMenu();
@@ -128,9 +131,6 @@ private:
 	/** Rename the selected Item (Only work with 1 selected Item) */
 	void OnExecuteRename();
 
-	/** Check whether the Rename can be executed */
-	bool CanExecuteRename();
-
 	void SynchronizeSelectedActorDescs();
 
 	void OnActorEditorContextSubsystemChanged();
