[SceneOutliner] When PrimaryKeys & SecondaryKeys are equal, default to ItemID Hash for sort stability

#rnx

diff --git a/Engine/Source/Editor/SceneOutliner/Public/SortHelper.h b/Engine/Source/Editor/SceneOutliner/Public/SortHelper.h
--- a/Engine/Source/Editor/SceneOutliner/Public/SortHelper.h
+++ b/Engine/Source/Editor/SceneOutliner/Public/SortHelper.h
@@ -101,7 +101,8 @@ struct FSceneOutlinerSortHelper
 				SecondaryKey = SecondaryFunction(*Element);
 			}
 
-			SortData.Emplace(Index, MoveTemp(PrimaryKey), MoveTemp(SecondaryKey));
+			const int32 DefaultSortID = Element.IsValid() ? GetTypeHash(Element->GetID()) : 0;
+			SortData.Emplace(Index, MoveTemp(PrimaryKey), MoveTemp(SecondaryKey), DefaultSortID);
 		}
 
 		SortData.Sort([&](const FSortPayload& One, const FSortPayload& Two){
@@ -110,12 +111,12 @@ struct FSceneOutlinerSortHelper
 			else if (PrimarySortMode == EColumnSortMode::Descending && One.PrimaryKey != Two.PrimaryKey)
 				return One.PrimaryKey > Two.PrimaryKey;
 
-			if (SecondarySortMode == EColumnSortMode::Ascending)
+			if (SecondarySortMode == EColumnSortMode::Ascending && One.SecondaryKey != Two.SecondaryKey)
 				return One.SecondaryKey < Two.SecondaryKey;
-			else if (SecondarySortMode == EColumnSortMode::Descending)
+			else if (SecondarySortMode == EColumnSortMode::Descending && One.SecondaryKey != Two.SecondaryKey)
 				return One.SecondaryKey > Two.SecondaryKey;
 
-			return false;
+			return One.DefaultSortID < Two.DefaultSortID;
 		});
 
 		TArray<FSceneOutlinerTreeItemPtr> NewArray;
@@ -140,12 +141,14 @@ private:
 	struct FSortPayload
 	{
 		int32 OriginalIndex;
-			
+		
 		PrimaryKeyType PrimaryKey;
 		SecondaryKeyType SecondaryKey;
+		
+		int32 DefaultSortID;
 
-		FSortPayload(int32 InOriginalIndex, PrimaryKeyType&& InPrimaryKey, SecondaryKeyType&& InSecondaryKey)
-			: OriginalIndex(InOriginalIndex), PrimaryKey(MoveTemp(InPrimaryKey)), SecondaryKey(MoveTemp(InSecondaryKey))
+		FSortPayload(int32 InOriginalIndex, PrimaryKeyType&& InPrimaryKey, SecondaryKeyType&& InSecondaryKey, int32 InDefaultSortID)
+			: OriginalIndex(InOriginalIndex), PrimaryKey(MoveTemp(InPrimaryKey)), SecondaryKey(MoveTemp(InSecondaryKey)), DefaultSortID(InDefaultSortID)
 		{}
 
 		FSortPayload(FSortPayload&& Other)
@@ -157,6 +160,7 @@ private:
 			OriginalIndex = rhs.OriginalIndex;
 			PrimaryKey = MoveTemp(rhs.PrimaryKey);
 			SecondaryKey = MoveTemp(rhs.SecondaryKey);
+			DefaultSortID = rhs.DefaultSortID;
 			return *this;
 		}
 
