Scene Outliner: Implement function to add single item to selection

diff --git a/Engine/Source/Editor/SceneOutliner/Private/SSceneOutliner.cpp b/Engine/Source/Editor/SceneOutliner/Private/SSceneOutliner.cpp
--- a/Engine/Source/Editor/SceneOutliner/Private/SSceneOutliner.cpp
+++ b/Engine/Source/Editor/SceneOutliner/Private/SSceneOutliner.cpp
@@ -1279,6 +1279,11 @@ void SSceneOutliner::AddToSelection(const TArray<FSceneOutlinerTreeItemPtr>& InI
 	OutlinerTreeView->SetItemSelection(InItems, true, SelectInfo);
 }
 
+void SSceneOutliner::AddToSelection(const FSceneOutlinerTreeItemPtr& InItem, ESelectInfo::Type SelectInfo)
+{
+	OutlinerTreeView->SetItemSelection(InItem, true, SelectInfo);
+}
+
 void SSceneOutliner::RemoveFromSelection(const TArray<FSceneOutlinerTreeItemPtr>& InItems, ESelectInfo::Type SelectInfo)
 {
 	OutlinerTreeView->SetItemSelection(InItems, false, SelectInfo);
diff --git a/Engine/Source/Editor/SceneOutliner/Public/SSceneOutliner.h b/Engine/Source/Editor/SceneOutliner/Public/SSceneOutliner.h
--- a/Engine/Source/Editor/SceneOutliner/Public/SSceneOutliner.h
+++ b/Engine/Source/Editor/SceneOutliner/Public/SSceneOutliner.h
@@ -440,7 +440,10 @@ public:
 	void AddToSelection(const TArray<FSceneOutlinerTreeItemPtr>& InItems, ESelectInfo::Type SelectInfo = ESelectInfo::Direct);
 
 	/** Remove a set of items from the current selection */
-	void RemoveFromSelection(const TArray<FSceneOutlinerTreeItemPtr>& InItems, ESelectInfo::Type SelectInfo = ESelectInfo::Direct);
+	void RemoveFromSelection(const TArray<FSceneOutlinerTreeItemPtr>& InItems, ESelectInfo::Type SelectInfo = ESelectInfo::Direct);\
+
+	/** Adds a single item to the current selection */
+	void AddToSelection(const FSceneOutlinerTreeItemPtr& InItem, ESelectInfo::Type SelectInfo = ESelectInfo::Direct);
 
 	/**
 		* Returns the list of currently selected tree items
