Outliner: Modify folder color and opacity to match the Content Browser

diff --git a/Engine/Source/Editor/SceneOutliner/Private/ActorFolderTreeItem.cpp b/Engine/Source/Editor/SceneOutliner/Private/ActorFolderTreeItem.cpp
--- a/Engine/Source/Editor/SceneOutliner/Private/ActorFolderTreeItem.cpp
+++ b/Engine/Source/Editor/SceneOutliner/Private/ActorFolderTreeItem.cpp
@@ -68,7 +68,6 @@ struct SActorFolderTreeLabel : FSceneOutlinerCommonLabelData, public SCompoundWi
 					[
 						SNew(SImage)
 						.Image(this, &SActorFolderTreeLabel::GetIcon)
-						.ColorAndOpacity(FSlateColor::UseForeground())
 					]
 				]
 				+ SHorizontalBox::Slot()
