Bumping the scene outliner tree item id from uint32 to uint64.

TEDS outliner is using the generic id constructor for FSceneOutlinerTreeItemID, which takes a uint32 as an id value. But TEDS code is using the row handle as an id, which is uint64, therefore we are loosing data in the conversion.

This is blocking tools like the Chaos Visual Debugger as they have unsafe type cast warnings as errors enabled, but more importantly, it seems like bug that went unnoticed as we are loosing data for each id.

diff --git a/Engine/Source/Editor/SceneOutliner/Public/SceneOutlinerStandaloneTypes.h b/Engine/Source/Editor/SceneOutliner/Public/SceneOutlinerStandaloneTypes.h
--- a/Engine/Source/Editor/SceneOutliner/Public/SceneOutlinerStandaloneTypes.h
+++ b/Engine/Source/Editor/SceneOutliner/Public/SceneOutlinerStandaloneTypes.h
@@ -34,7 +34,7 @@ FORCEINLINE uint32 GetTypeHash(FFolderKey Key)
 struct FSceneOutlinerTreeItemID
 {
 public:
-	using FUniqueID = uint32;
+	using FUniqueID = uint64;
 
 	enum class EType : uint8 { Object, Folder, UniqueID, Guid, Null };
 
