[SceneOutliner]

Disable the visibility widget for component tree items. It only creates visual clutter and has no function.

diff --git a/Engine/Source/Editor/SceneOutliner/Public/ComponentTreeItem.h b/Engine/Source/Editor/SceneOutliner/Public/ComponentTreeItem.h
--- a/Engine/Source/Editor/SceneOutliner/Public/ComponentTreeItem.h
+++ b/Engine/Source/Editor/SceneOutliner/Public/ComponentTreeItem.h
@@ -45,6 +45,9 @@ public:
 	virtual FString GetDisplayString() const override;
 	virtual bool CanInteract() const override;
 	virtual TSharedRef<SWidget> GenerateLabelWidget(ISceneOutliner& Outliner, const STableRow<FSceneOutlinerTreeItemPtr>& InRow) override;
+	virtual bool ShouldShowVisibilityState() const { return false; }
+	virtual bool HasVisibilityInfo() const { return false; }
+	virtual bool GetVisibility() const { return false; }
 	/* End ISceneOutlinerTreeItem Implementation */
 public:
 	/** true if this item exists in both the current world and PIE. */
