Hide actor entity components proxy when users asked to show Actor Components in SceneOutliner

#rnx

diff --git a/Engine/Source/Editor/SceneOutliner/Private/ActorHierarchy.cpp b/Engine/Source/Editor/SceneOutliner/Private/ActorHierarchy.cpp
--- a/Engine/Source/Editor/SceneOutliner/Private/ActorHierarchy.cpp
+++ b/Engine/Source/Editor/SceneOutliner/Private/ActorHierarchy.cpp
@@ -383,7 +383,7 @@ void FActorHierarchy::CreateComponentItems(const AActor* Actor, TArray<FSceneOut
 	{
 		for (UActorComponent* Component : Actor->GetComponents())
 		{
-			if (Component != nullptr)
+			if (Component != nullptr && !Component->GetClass()->HasAnyClassFlags(CLASS_Hidden))
 			{
 				if (FSceneOutlinerTreeItemPtr ComponentItem =
 					Mode->CreateItemFor<FComponentTreeItem>(FComponentTreeItem(Component, bSearchComponentsByActorName)))
