[Truncation Warnings] Fix warnings in SceneOutliner and UndoHistoryEditor modules

#rnx

diff --git a/Engine/Source/Editor/SceneOutliner/Public/SSceneOutliner.h b/Engine/Source/Editor/SceneOutliner/Public/SSceneOutliner.h
--- a/Engine/Source/Editor/SceneOutliner/Public/SSceneOutliner.h
+++ b/Engine/Source/Editor/SceneOutliner/Public/SSceneOutliner.h
@@ -72,7 +72,7 @@ namespace SceneOutliner
 	/** Set of actions to apply to new tree items */
 	namespace ENewItemAction
 	{
-		enum Type
+		enum Type : uint8
 		{
 			/** Do nothing when it is created */
 			None = 0,
diff --git a/Engine/Source/Editor/SceneOutliner/Public/SSocketChooser.h b/Engine/Source/Editor/SceneOutliner/Public/SSocketChooser.h
--- a/Engine/Source/Editor/SceneOutliner/Public/SSocketChooser.h
+++ b/Engine/Source/Editor/SceneOutliner/Public/SSocketChooser.h
@@ -115,7 +115,7 @@ public:
 				SNew(SHorizontalBox)
 				+SHorizontalBox::Slot()
 				.AutoWidth()
-				.Padding(2)
+				.Padding(2.f)
 				.VAlign(VAlign_Center)
 				[
 					SNew(SImage)
@@ -185,7 +185,7 @@ public:
 		[
 			SNew(SBorder)
 			.BorderImage(FAppStyle::GetBrush(TEXT("Menu.Background")))
-			.Padding(5)
+			.Padding(5.f)
 			.Content()
 			[
 				SNew(SVerticalBox)
@@ -206,10 +206,10 @@ public:
 				]
 				+SVerticalBox::Slot()
 				.AutoHeight()
-				.MaxHeight(512)
+				.MaxHeight(512.f)
 				[
 					SNew(SBox)
-					.WidthOverride(256)
+					.WidthOverride(256.f)
 					.Content()
 					[
 						SAssignNew(SocketListView, SListView< TSharedPtr<FSocketInfo> >)
@@ -230,8 +230,8 @@ public:
 		TSharedPtr<SWindow> Window = CheckAndGetWindowPtr();
 		if(Window.IsValid())
 		{
-			FVector2D WindowLocation = Window->IsMorphing() ? Window->GetMorphTargetPosition() : Window->GetPositionInScreen();
-			FVector2D WindowSize = Window->GetDesiredSize();
+			FVector2f WindowLocation = Window->IsMorphing() ? Window->GetMorphTargetPosition() : Window->GetPositionInScreen();
+			FVector2f WindowSize = Window->GetDesiredSize();
 			FSlateRect Anchor(WindowLocation.X, WindowLocation.Y, WindowLocation.X, WindowLocation.Y);
 			WindowLocation = FSlateApplication::Get().CalculatePopupWindowPosition( Anchor, WindowSize, false );
 
diff --git a/Engine/Source/Editor/UndoHistoryEditor/Private/Widgets/SUndoHistory.cpp b/Engine/Source/Editor/UndoHistoryEditor/Private/Widgets/SUndoHistory.cpp
--- a/Engine/Source/Editor/UndoHistoryEditor/Private/Widgets/SUndoHistory.cpp
+++ b/Engine/Source/Editor/UndoHistoryEditor/Private/Widgets/SUndoHistory.cpp
@@ -189,7 +189,7 @@ void SUndoHistory::Construct( const FArguments& InArgs )
 					.HAlign(HAlign_Right)
 					[
 						SNew(SComboButton)
-							.ContentPadding(0)
+							.ContentPadding(0.f)
 							.ForegroundColor(FSlateColor::UseForeground())
 							.ButtonStyle(FAppStyle::Get(), "ToggleButton")
 							.OnGetMenuContent(this, &SUndoHistory::GetViewButtonContent)
@@ -430,7 +430,7 @@ void SUndoHistory::HandleUndoListJumpToTransaction(TSharedPtr<FTransactionInfo>
 
 	const int32 RemainingUndoRedo = FMath::Abs(InItem->QueueIndex - LastActiveTransactionIndex);
 
- 	FScopedSlowTask SlowTask(RemainingUndoRedo, LOCTEXT("ReplayingTransactions", "Replaying Transactions..."));
+ 	FScopedSlowTask SlowTask(static_cast<float>(RemainingUndoRedo), LOCTEXT("ReplayingTransactions", "Replaying Transactions..."));
  	SlowTask.MakeDialogDelayed(0.5f);
 
 	const auto ProgressLambda = [&SlowTask](const int32 InTransactionIndex)
