[Truncation Warnings] Enable truncation warnings in SceneOutliner module

#rnx

diff --git a/Engine/Source/Editor/SceneOutliner/Private/ActorDescTreeItem.cpp b/Engine/Source/Editor/SceneOutliner/Private/ActorDescTreeItem.cpp
--- a/Engine/Source/Editor/SceneOutliner/Private/ActorDescTreeItem.cpp
+++ b/Engine/Source/Editor/SceneOutliner/Private/ActorDescTreeItem.cpp
@@ -91,8 +91,8 @@ struct SActorDescTreeLabel : FSceneOutlinerCommonLabelData, public SCompoundWidg
 				.Padding(FSceneOutlinerDefaultTreeItemMetrics::IconPadding())
 				[
 					SNew(SBox)
-					.WidthOverride(FSceneOutlinerDefaultTreeItemMetrics::IconSize())
-					.HeightOverride(FSceneOutlinerDefaultTreeItemMetrics::IconSize())
+					.WidthOverride(static_cast<float>(FSceneOutlinerDefaultTreeItemMetrics::IconSize()))
+					.HeightOverride(static_cast<float>(FSceneOutlinerDefaultTreeItemMetrics::IconSize()))
 					[
 						SNew(SImage)
 						.Image(this, &SActorDescTreeLabel::GetIcon)
diff --git a/Engine/Source/Editor/SceneOutliner/Private/ActorFolderTreeItem.cpp b/Engine/Source/Editor/SceneOutliner/Private/ActorFolderTreeItem.cpp
--- a/Engine/Source/Editor/SceneOutliner/Private/ActorFolderTreeItem.cpp
+++ b/Engine/Source/Editor/SceneOutliner/Private/ActorFolderTreeItem.cpp
@@ -63,8 +63,8 @@ struct SActorFolderTreeLabel : FSceneOutlinerCommonLabelData, public SCompoundWi
 				.Padding(FSceneOutlinerDefaultTreeItemMetrics::IconPadding())
 				[
 					SNew(SBox)
-					.WidthOverride(FSceneOutlinerDefaultTreeItemMetrics::IconSize())
-					.HeightOverride(FSceneOutlinerDefaultTreeItemMetrics::IconSize())
+					.WidthOverride(static_cast<float>(FSceneOutlinerDefaultTreeItemMetrics::IconSize()))
+					.HeightOverride(static_cast<float>(FSceneOutlinerDefaultTreeItemMetrics::IconSize()))
 					[
 						SNew(SImage)
 						.Image(this, &SActorFolderTreeLabel::GetIcon)
diff --git a/Engine/Source/Editor/SceneOutliner/Private/ActorTreeItem.cpp b/Engine/Source/Editor/SceneOutliner/Private/ActorTreeItem.cpp
--- a/Engine/Source/Editor/SceneOutliner/Private/ActorTreeItem.cpp
+++ b/Engine/Source/Editor/SceneOutliner/Private/ActorTreeItem.cpp
@@ -81,8 +81,8 @@ struct SActorTreeLabel : FSceneOutlinerCommonLabelData, public SCompoundWidget
 				.Padding(FSceneOutlinerDefaultTreeItemMetrics::IconPadding())
 				[
 					SNew(SBox)
-					.WidthOverride(FSceneOutlinerDefaultTreeItemMetrics::IconSize())
-					.HeightOverride(FSceneOutlinerDefaultTreeItemMetrics::IconSize())
+					.WidthOverride(static_cast<float>(FSceneOutlinerDefaultTreeItemMetrics::IconSize()))
+					.HeightOverride(static_cast<float>(FSceneOutlinerDefaultTreeItemMetrics::IconSize()))
 					[
 						SNew(SImage)
 						.Image(this, &SActorTreeLabel::GetIcon)
diff --git a/Engine/Source/Editor/SceneOutliner/Private/ComponentTreeItem.cpp b/Engine/Source/Editor/SceneOutliner/Private/ComponentTreeItem.cpp
--- a/Engine/Source/Editor/SceneOutliner/Private/ComponentTreeItem.cpp
+++ b/Engine/Source/Editor/SceneOutliner/Private/ComponentTreeItem.cpp
@@ -65,8 +65,8 @@ struct SComponentTreeLabel : FSceneOutlinerCommonLabelData, public SCompoundWidg
 				.Padding(FSceneOutlinerDefaultTreeItemMetrics::IconPadding())
 				[
 					SNew(SBox)
-					.WidthOverride(FSceneOutlinerDefaultTreeItemMetrics::IconSize())
-					.HeightOverride(FSceneOutlinerDefaultTreeItemMetrics::IconSize())
+					.WidthOverride(static_cast<float>(FSceneOutlinerDefaultTreeItemMetrics::IconSize()))
+					.HeightOverride(static_cast<float>(FSceneOutlinerDefaultTreeItemMetrics::IconSize()))
 					[
 						SNew(SImage)
 						.Image(this, &SComponentTreeLabel::GetIcon)
diff --git a/Engine/Source/Editor/SceneOutliner/Private/LevelTreeItem.cpp b/Engine/Source/Editor/SceneOutliner/Private/LevelTreeItem.cpp
--- a/Engine/Source/Editor/SceneOutliner/Private/LevelTreeItem.cpp
+++ b/Engine/Source/Editor/SceneOutliner/Private/LevelTreeItem.cpp
@@ -36,8 +36,8 @@ struct SLevelTreeLabel : FSceneOutlinerCommonLabelData, public SCompoundWidget
 			.Padding(FSceneOutlinerDefaultTreeItemMetrics::IconPadding())
 			[
 				SNew(SBox)
-				.WidthOverride(FSceneOutlinerDefaultTreeItemMetrics::IconSize())
-				.HeightOverride(FSceneOutlinerDefaultTreeItemMetrics::IconSize())
+				.WidthOverride(static_cast<float>(FSceneOutlinerDefaultTreeItemMetrics::IconSize()))
+				.HeightOverride(static_cast<float>(FSceneOutlinerDefaultTreeItemMetrics::IconSize()))
 				[
 					SNew(SImage)
 					.Image(FSlateIconFinder::FindIconBrushForClass(UWorld::StaticClass()))
diff --git a/Engine/Source/Editor/SceneOutliner/Private/SOutlinerTreeView.cpp b/Engine/Source/Editor/SceneOutliner/Private/SOutlinerTreeView.cpp
--- a/Engine/Source/Editor/SceneOutliner/Private/SOutlinerTreeView.cpp
+++ b/Engine/Source/Editor/SceneOutliner/Private/SOutlinerTreeView.cpp
@@ -295,13 +295,13 @@ TSharedRef<SWidget> SSceneOutlinerTreeRow::GenerateWidgetForColumn( const FName&
 	{
 		// The first column gets the tree expansion arrow for this row
 		return SNew(SBox)
-			.MinDesiredHeight(FSceneOutlinerDefaultTreeItemMetrics::RowHeight())
+			.MinDesiredHeight(static_cast<float>(FSceneOutlinerDefaultTreeItemMetrics::RowHeight()))
 			[
 				SNew( SHorizontalBox )
 
 				+SHorizontalBox::Slot()
 				.AutoWidth()
-				.Padding(6, 0, 0, 0)
+				.Padding(6.f, 0.f, 0.f, 0.f)
 				[
 					SNew( SExpanderArrow, SharedThis(this) ).IndentAmount(12)
 				]
@@ -473,12 +473,12 @@ TSharedRef<SWidget> SSceneOutlinerPinnedTreeRow::GenerateWidgetForColumn(const F
 		{
 			// We add space for the expander arrow for the label widget to make sure the spacing/indentation is consistent with non-pinned rows
 			NewItemWidget = SNew(SBox)
-				.MinDesiredHeight(FSceneOutlinerDefaultTreeItemMetrics::RowHeight())
+				.MinDesiredHeight(static_cast<float>(FSceneOutlinerDefaultTreeItemMetrics::RowHeight()))
 				[
 					SNew(SHorizontalBox)
 					+ SHorizontalBox::Slot()
 					.AutoWidth()
-					.Padding(6, 0, 0, 0)
+					.Padding(6.f, 0.f, 0.f, 0.f)
 					[
 						SNew(SExpanderArrow, SharedThis(this)).IndentAmount(12)
 						.Visibility(EVisibility::Hidden) // Hidden SExpanderArrow to occupy the same space as non-pinned rows
diff --git a/Engine/Source/Editor/SceneOutliner/Private/WorldTreeItem.cpp b/Engine/Source/Editor/SceneOutliner/Private/WorldTreeItem.cpp
--- a/Engine/Source/Editor/SceneOutliner/Private/WorldTreeItem.cpp
+++ b/Engine/Source/Editor/SceneOutliner/Private/WorldTreeItem.cpp
@@ -96,8 +96,8 @@ struct SWorldTreeLabel : FSceneOutlinerCommonLabelData, public SCompoundWidget
 			.Padding(FSceneOutlinerDefaultTreeItemMetrics::IconPadding())
 			[
 				SNew(SBox)
-				.WidthOverride(FSceneOutlinerDefaultTreeItemMetrics::IconSize())
-				.HeightOverride(FSceneOutlinerDefaultTreeItemMetrics::IconSize())
+				.WidthOverride(static_cast<float>(FSceneOutlinerDefaultTreeItemMetrics::IconSize()))
+				.HeightOverride(static_cast<float>(FSceneOutlinerDefaultTreeItemMetrics::IconSize()))
 				[
 					SNew(SImage)
 					.Image(FSlateIconFinder::FindIconBrushForClass(UWorld::StaticClass()))
diff --git a/Engine/Source/Editor/SceneOutliner/SceneOutliner.Build.cs b/Engine/Source/Editor/SceneOutliner/SceneOutliner.Build.cs
--- a/Engine/Source/Editor/SceneOutliner/SceneOutliner.Build.cs
+++ b/Engine/Source/Editor/SceneOutliner/SceneOutliner.Build.cs
@@ -6,6 +6,9 @@ public class SceneOutliner : ModuleRules
 {
 	public SceneOutliner(ReadOnlyTargetRules Target) : base(Target)
 	{
+		// Enable truncation warnings in this module
+		CppCompileWarningSettings.UnsafeTypeCastWarningLevel = WarningLevel.Warning;
+
 		PrivateDependencyModuleNames.AddRange(
 			new string[]
 			{
