Fix labels not being set when pasting actors

Re-added a refresh that was previously removed but is apparently required for labels being updated in the outliner post-paste

#rnx

diff --git a/Engine/Source/Editor/SceneOutliner/Private/SSceneOutliner.cpp b/Engine/Source/Editor/SceneOutliner/Private/SSceneOutliner.cpp
--- a/Engine/Source/Editor/SceneOutliner/Private/SSceneOutliner.cpp
+++ b/Engine/Source/Editor/SceneOutliner/Private/SSceneOutliner.cpp
@@ -1705,6 +1705,7 @@ void SSceneOutliner::PasteFoldersEnd()
 
 	if (CacheFoldersEdit.IsEmpty())
 	{
+		FullRefresh();
 		return;
 	}
 
