Speculative fix for crash when selecting menu entry while underlying world is torn down

diff --git a/Engine/Source/Editor/SceneOutliner/Private/ActorBrowsingMode.cpp b/Engine/Source/Editor/SceneOutliner/Private/ActorBrowsingMode.cpp
--- a/Engine/Source/Editor/SceneOutliner/Private/ActorBrowsingMode.cpp
+++ b/Engine/Source/Editor/SceneOutliner/Private/ActorBrowsingMode.cpp
@@ -1493,23 +1493,35 @@ void FActorBrowsingMode::SynchronizeSelectedActorDescs()
 
 FFolder FActorBrowsingMode::CreateNewFolder()
 {
-	const FScopedTransaction Transaction(LOCTEXT("UndoAction_CreateFolder", "Create Folder"));
-	TArray<FFolder> SelectedFolders = SceneOutliner->GetSelection().GetData<FFolder>(SceneOutliner::FFolderPathSelector());
-	const FFolder NewFolderName = FActorFolders::Get().GetDefaultFolderForSelection(*RepresentingWorld, &SelectedFolders);
-	FActorFolders::Get().CreateFolderContainingSelection(*RepresentingWorld, NewFolderName);
+	if (TStrongObjectPtr<UWorld> World = RepresentingWorld.Pin())
+	{
+		const FScopedTransaction Transaction(LOCTEXT("UndoAction_CreateFolder", "Create Folder"));
+		TArray<FFolder> SelectedFolders = SceneOutliner->GetSelection().GetData<FFolder>(SceneOutliner::FFolderPathSelector());
+		const FFolder NewFolderName = FActorFolders::Get().GetDefaultFolderForSelection(*World, &SelectedFolders);
+		FActorFolders::Get().CreateFolderContainingSelection(*World, NewFolderName);
+		return NewFolderName;
+	}
 
-	return NewFolderName;
+	return FFolder();
 }
 
 FFolder FActorBrowsingMode::GetFolder(const FFolder& ParentPath, const FName& LeafName)
 {
 	// Return a unique folder under the provided parent path & root object and using the provided leaf name
-	return FActorFolders::Get().GetFolderName(*RepresentingWorld, ParentPath, LeafName);
+	if (TStrongObjectPtr<UWorld> World = RepresentingWorld.Pin())
+	{
+		return FActorFolders::Get().GetFolderName(*World, ParentPath, LeafName);
+	}
+	return FFolder();
 }
 
 bool FActorBrowsingMode::CreateFolder(const FFolder& NewPath)
 {
-	return FActorFolders::Get().CreateFolder(*RepresentingWorld, NewPath);
+	if (TStrongObjectPtr<UWorld> World = RepresentingWorld.Pin())
+	{
+		return FActorFolders::Get().CreateFolder(*World, NewPath);
+	}
+	return false;
 }
 
 bool FActorBrowsingMode::ReparentItemToFolder(const FFolder& FolderPath, const FSceneOutlinerTreeItemPtr& Item)
