Add Outliner - Searching and Filtering tests

#rnx

#virtualized

diff --git a/Engine/Source/Editor/LevelEditor/Public/LevelEditor.h b/Engine/Source/Editor/LevelEditor/Public/LevelEditor.h
--- a/Engine/Source/Editor/LevelEditor/Public/LevelEditor.h
+++ b/Engine/Source/Editor/LevelEditor/Public/LevelEditor.h
@@ -394,7 +394,7 @@ public:
 	}
 
 	/** Return the settings (containing the outliner filters) used to create the level editor's outliners  */
-	TSharedPtr<FLevelEditorOutlinerSettings> GetLevelEditorOutlinerSettings() const;
+	virtual TSharedPtr<FLevelEditorOutlinerSettings> GetLevelEditorOutlinerSettings() const;
 
 public:
 	
diff --git a/Engine/Source/Editor/SceneOutliner/Private/SSceneOutliner.cpp b/Engine/Source/Editor/SceneOutliner/Private/SSceneOutliner.cpp
--- a/Engine/Source/Editor/SceneOutliner/Private/SSceneOutliner.cpp
+++ b/Engine/Source/Editor/SceneOutliner/Private/SSceneOutliner.cpp
@@ -2681,6 +2681,27 @@ void SSceneOutliner::AddSourceControlMenuOptions(UToolMenu* Menu)
 	SourceControlHandler->AddSourceControlMenuOptions(Menu, SelectedItems);
 }
 
+void SSceneOutliner::FilterByType(TSharedPtr<FCustomClassFilterData> CustomClassFilterData)
+{
+	if (FilterBar)
+	{
+		FilterBar->FilterByTypeClicked(CustomClassFilterData);
+	}
+}
+
+void SSceneOutliner::FilterByTypeCategory(TSharedPtr<FFilterCategory> TypeCategory, TArray<TSharedPtr<FCustomClassFilterData>> Classes)
+{
+	if (FilterBar)
+	{
+		FilterBar->FilterByTypeCategoryClicked(TypeCategory, Classes);
+	}
+}
+
+bool SSceneOutliner::NeedsRefresh() const
+{
+	return bNeedsRefresh || bFullRefresh;
+}
+
 void SSceneOutliner::OnFilterBarFilterChanged()
 {
 	FilterCollection = FilterBar->GetAllActiveFilters();
diff --git a/Engine/Source/Editor/SceneOutliner/Public/SSceneOutliner.h b/Engine/Source/Editor/SceneOutliner/Public/SSceneOutliner.h
--- a/Engine/Source/Editor/SceneOutliner/Public/SSceneOutliner.h
+++ b/Engine/Source/Editor/SceneOutliner/Public/SSceneOutliner.h
@@ -1063,6 +1063,10 @@ public:
 	virtual TSharedPtr<FSceneOutlinerTreeItemSCC> GetItemSourceControl(const FSceneOutlinerTreeItemPtr& InItem) override;
 
 	void AddSourceControlMenuOptions(UToolMenu* Menu);
+
+	void FilterByType(TSharedPtr<FCustomClassFilterData> CustomClassFilterData);
+	void FilterByTypeCategory(TSharedPtr<FFilterCategory> TypeCategory, TArray<TSharedPtr<FCustomClassFilterData>> Classes);
+	bool NeedsRefresh() const;
 };
 
 struct SCENEOUTLINER_API FSceneOutlinerMenuHelper
