[Editor] SceneOutliner Unsaved column: handle actor package change
- Remove need for internal map
- Update filename when actor packaging mode changes

#rnx

diff --git a/Engine/Source/Editor/SceneOutliner/Private/SceneOutlinerUnsavedColumn.cpp b/Engine/Source/Editor/SceneOutliner/Private/SceneOutlinerUnsavedColumn.cpp
--- a/Engine/Source/Editor/SceneOutliner/Private/SceneOutlinerUnsavedColumn.cpp
+++ b/Engine/Source/Editor/SceneOutliner/Private/SceneOutlinerUnsavedColumn.cpp
@@ -2,15 +2,13 @@
 
 #include "SceneOutlinerUnsavedColumn.h"
 #include "ActorTreeItem.h"
-#include "ActorDescTreeItem.h"
-#include "ActorFolderTreeItem.h"
 #include "SourceControlHelpers.h"
-#include "WorldPartition/WorldPartitionActorDesc.h"
 #include "UnsavedAssetsTrackerModule.h"
 #include "SortHelper.h"
 #include "ISceneOutliner.h"
 #include "ISceneOutlinerMode.h"
 #include "SceneOutlinerHelpers.h"
+#include "GameFramework/Actor.h"
 
 #define LOCTEXT_NAMESPACE "SceneOutlinerActorUnsavedColumn"
 
@@ -20,33 +18,97 @@ public:
 	SLATE_BEGIN_ARGS(SUnsavedActorWidget) {}
 	SLATE_END_ARGS()
 
+	~SUnsavedActorWidget();
+
 	/** Construct this widget */
-	void Construct(const FArguments& InArgs, const FString& InActorExternalPackageFilename);
+	void Construct(const FArguments& InArgs, TWeakPtr<ISceneOutlinerTreeItem> InWeakTreeItem);
 	
 	bool IsUnsaved() const;
 	
-private:
-	void UpdateImage();
 	void OnUnsavedAssetAdded(const FString& FileAbsPathname);
 	void OnUnsavedAssetRemoved(const FString& FileAbsPathname);
 
+private:
+	void UpdateImage();
+	void UpdateExternalPackageFilename();
+	
 private:
 	FString ExternalPackageFilename;
+	TWeakPtr<ISceneOutlinerTreeItem> WeakTreeItem;
 	bool bIsUnsaved;
+
+	FDelegateHandle OnPackagingModeChangedHandle;
+	FDelegateHandle OnUnsavedAssetAddedHandle;
+	FDelegateHandle OnUnsavedAssetRemovedHandle;
 };
 
-void SUnsavedActorWidget::Construct(const FArguments& InArgs, const FString& InActorExternalPackageFilename)
+void SUnsavedActorWidget::Construct(const FArguments& InArgs, TWeakPtr<ISceneOutlinerTreeItem> InWeakTreeItem)
 {
-	ExternalPackageFilename = InActorExternalPackageFilename;
+	WeakTreeItem = InWeakTreeItem;
 
 	SImage::Construct(
 			SImage::FArguments()
 			.ColorAndOpacity(FSlateColor::UseForeground())
 			.Image(FStyleDefaults::GetNoBrush()));
-	
+
+	UpdateExternalPackageFilename();
+
+	// Handle Actor package change
+	if (FActorTreeItem* ActorItem = WeakTreeItem.Pin()->CastTo<FActorTreeItem>())
+	{
+		if (AActor* Actor = ActorItem->Actor.Get())
+		{
+			OnPackagingModeChangedHandle = Actor->OnPackagingModeChanged.AddLambda([WeakThis = AsWeak()](AActor* InActor, bool bExternal)
+			{
+				if (TSharedPtr<SUnsavedActorWidget> This = StaticCastSharedPtr<SUnsavedActorWidget>(WeakThis.Pin()))
+				{
+					This->UpdateExternalPackageFilename();
+				}
+			});
+		}
+	}
+}
+
+SUnsavedActorWidget::~SUnsavedActorWidget()
+{
+	if (FActorTreeItem* ActorItem = WeakTreeItem.IsValid() ? WeakTreeItem.Pin()->CastTo<FActorTreeItem>() : nullptr)
+	{
+		if (AActor* Actor = ActorItem->Actor.Get())
+		{
+			Actor->OnPackagingModeChanged.Remove(OnPackagingModeChangedHandle);
+		}
+	}
+
+	if (FUnsavedAssetsTrackerModule* UnsavedAssetsTrackerModule = FModuleManager::GetModulePtr<FUnsavedAssetsTrackerModule>("UnsavedAssetsTracker"))
+	{
+		UnsavedAssetsTrackerModule->OnUnsavedAssetAdded.Remove(OnUnsavedAssetAddedHandle);
+		UnsavedAssetsTrackerModule->OnUnsavedAssetRemoved.Remove(OnUnsavedAssetRemovedHandle);
+	}
+}
+
+void SUnsavedActorWidget::UpdateExternalPackageFilename()
+{
+	const bool bWasExternal = !ExternalPackageFilename.IsEmpty();
+		
+	const FString ExternalPackageName = WeakTreeItem.IsValid() ? WeakTreeItem.Pin()->GetPackageName() : FString();
+	ExternalPackageFilename = !ExternalPackageName.IsEmpty() ? USourceControlHelpers::PackageFilename(ExternalPackageName) : FString();
+
 	FUnsavedAssetsTrackerModule& UnsavedAssetsTrackerModule = FUnsavedAssetsTrackerModule::Get();
-	UnsavedAssetsTrackerModule.OnUnsavedAssetAdded.AddSP(this, &SUnsavedActorWidget::OnUnsavedAssetAdded);
-	UnsavedAssetsTrackerModule.OnUnsavedAssetRemoved.AddSP(this, &SUnsavedActorWidget::OnUnsavedAssetRemoved);
+
+	// Register/Unregister if needed
+	if (bWasExternal && ExternalPackageFilename.IsEmpty())
+	{
+		UnsavedAssetsTrackerModule.OnUnsavedAssetAdded.Remove(OnUnsavedAssetAddedHandle);
+		OnUnsavedAssetAddedHandle.Reset();
+
+		UnsavedAssetsTrackerModule.OnUnsavedAssetRemoved.Remove(OnUnsavedAssetRemovedHandle);
+		OnUnsavedAssetRemovedHandle.Reset();
+	}
+	else if (!bWasExternal && !ExternalPackageFilename.IsEmpty())
+	{
+		OnUnsavedAssetAddedHandle = UnsavedAssetsTrackerModule.OnUnsavedAssetAdded.AddSP(this, &SUnsavedActorWidget::OnUnsavedAssetAdded);
+		OnUnsavedAssetRemovedHandle = UnsavedAssetsTrackerModule.OnUnsavedAssetRemoved.AddSP(this, &SUnsavedActorWidget::OnUnsavedAssetRemoved);
+	}
 
 	bIsUnsaved = UnsavedAssetsTrackerModule.IsAssetUnsaved(ExternalPackageFilename);
 
@@ -117,44 +179,27 @@ SHeaderRow::FColumn::FArguments FSceneOutlinerActorUnsavedColumn::ConstructHeade
 }
 
 const TSharedRef<SWidget> FSceneOutlinerActorUnsavedColumn::ConstructRowWidget(FSceneOutlinerTreeItemRef TreeItem, const STableRow<FSceneOutlinerTreeItemPtr>& Row)
-{
-	const FString ExternalPackageName = TreeItem->GetPackageName();
-	const FString ExternalPackageFileName = !ExternalPackageName.IsEmpty() ? USourceControlHelpers::PackageFilename(ExternalPackageName) : FString();
-
-	if (ExternalPackageFileName.IsEmpty())
-	{
-		return SNullWidget::NullWidget;
-	}
-
-	const TSharedRef<SUnsavedActorWidget> ActualWidget = SNew(SUnsavedActorWidget, ExternalPackageFileName);
-
-	// Store a reference to the widget for this asset so we can grab the unsaved status for sorting
-	UnsavedActorWidgets.Add(ExternalPackageFileName, ActualWidget);
-	
+{	
 	return SNew(SHorizontalBox)
 			+ SHorizontalBox::Slot()
 			.AutoWidth()
 			.HAlign(HAlign_Center)
 			.VAlign(VAlign_Center)
 			[
-				ActualWidget
+				SNew(SUnsavedActorWidget, TreeItem.ToWeakPtr())
 			];
 }
 
 void FSceneOutlinerActorUnsavedColumn::SortItems(TArray<FSceneOutlinerTreeItemPtr>& RootItems, const EColumnSortMode::Type SortMode) const
 {
+	FUnsavedAssetsTrackerModule& UnsavedAssetsTrackerModule = FUnsavedAssetsTrackerModule::Get();
 	FSceneOutlinerSortHelper<bool, SceneOutliner::FNumericStringWrapper>()
 		/** Sort by unsaved first */
-		.Primary([this](const ISceneOutlinerTreeItem& Item)
+		.Primary([UnsavedAssetsTrackerModule](const ISceneOutlinerTreeItem& Item)
 		{
-			FString LongPackageName = Item.GetPackageName();
-			if (!LongPackageName.IsEmpty())
+			if (const FString PackageName = Item.GetPackageName(); !PackageName.IsEmpty())
 			{
-				if (const TSharedRef<SUnsavedActorWidget>* FoundWidget = UnsavedActorWidgets.Find(USourceControlHelpers::PackageFilename(LongPackageName)))
-				{
-					// return the inverse because we want items with Unsaved = TRUE to show up on top in ascending sort (default)
-					return !(*FoundWidget)->IsUnsaved();
-				}
+				return !UnsavedAssetsTrackerModule.IsAssetUnsaved(USourceControlHelpers::PackageFilename(PackageName));
 			}
 			return true;
 		}, SortMode)
diff --git a/Engine/Source/Editor/SceneOutliner/Private/SceneOutlinerUnsavedColumn.h b/Engine/Source/Editor/SceneOutliner/Private/SceneOutlinerUnsavedColumn.h
--- a/Engine/Source/Editor/SceneOutliner/Private/SceneOutlinerUnsavedColumn.h
+++ b/Engine/Source/Editor/SceneOutliner/Private/SceneOutlinerUnsavedColumn.h
@@ -32,8 +32,4 @@ public:
 
 private:
 	TWeakPtr<ISceneOutliner> WeakOutliner;
-
-	// Map of asset path -> widgets for each asset we encounter
-	TMap<FString, TSharedRef<SUnsavedActorWidget>> UnsavedActorWidgets;
-	
 };
