Add Outliner- Custom Filter tests

#rnx

diff --git a/Engine/Source/Editor/SceneOutliner/Private/SSceneOutliner.cpp b/Engine/Source/Editor/SceneOutliner/Private/SSceneOutliner.cpp
--- a/Engine/Source/Editor/SceneOutliner/Private/SSceneOutliner.cpp
+++ b/Engine/Source/Editor/SceneOutliner/Private/SSceneOutliner.cpp
@@ -2705,6 +2705,39 @@ bool SSceneOutliner::NeedsRefresh() const
 	return bNeedsRefresh || bFullRefresh;
 }
 
+void SSceneOutliner::CreateCustomTextFilter(const FCustomTextFilterData& InFilterData, bool bApplyFilter)
+{
+	if (FilterBar)
+	{
+		if (TSharedPtr<SSceneOutlinerFilterBar> OutlinerFilterBar = StaticCastSharedPtr<SSceneOutlinerFilterBar>(FilterBar))
+		{
+			OutlinerFilterBar->OnCreateCustomTextFilter(InFilterData, bApplyFilter);
+		}
+	}
+}
+
+void SSceneOutliner::ModifyCustomTextFilter(const FCustomTextFilterData& InFilterData, const TSharedPtr<ICustomTextFilter<SceneOutliner::FilterBarType>>& InFilter)
+{
+	if (FilterBar)
+	{
+		if (TSharedPtr<SSceneOutlinerFilterBar> OutlinerFilterBar = StaticCastSharedPtr<SSceneOutlinerFilterBar>(FilterBar))
+		{
+			OutlinerFilterBar->OnModifyCustomTextFilter(InFilterData, InFilter);
+		}
+	}
+}
+
+void SSceneOutliner::DeleteCustomTextFilter(const TSharedPtr<ICustomTextFilter<SceneOutliner::FilterBarType>>& InFilter)
+{
+	if (FilterBar)
+	{
+		if (TSharedPtr<SSceneOutlinerFilterBar> OutlinerFilterBar = StaticCastSharedPtr<SSceneOutlinerFilterBar>(FilterBar))
+		{
+			OutlinerFilterBar->OnDeleteCustomTextFilter(InFilter);
+		}
+	}
+}
+
 void SSceneOutliner::OnFilterBarFilterChanged()
 {
 	FilterCollection = FilterBar->GetAllActiveFilters();
diff --git a/Engine/Source/Editor/SceneOutliner/Public/SSceneOutliner.h b/Engine/Source/Editor/SceneOutliner/Public/SSceneOutliner.h
--- a/Engine/Source/Editor/SceneOutliner/Public/SSceneOutliner.h
+++ b/Engine/Source/Editor/SceneOutliner/Public/SSceneOutliner.h
@@ -41,6 +41,8 @@ struct FToolMenuSection;
 template<typename FilterType> class SFilterBar;
 template<typename ItemType> class STreeView;
 class SFilterSearchBox;
+struct FCustomTextFilterData;
+template<typename FilterType> class ICustomTextFilter;
 
 SCENEOUTLINER_API DECLARE_LOG_CATEGORY_EXTERN(LogSceneOutliner, Log, All);
 
@@ -1067,6 +1069,10 @@ public:
 	void FilterByType(TSharedPtr<FCustomClassFilterData> CustomClassFilterData);
 	void FilterByTypeCategory(TSharedPtr<FFilterCategory> TypeCategory, TArray<TSharedPtr<FCustomClassFilterData>> Classes);
 	bool NeedsRefresh() const;
+
+	void CreateCustomTextFilter(const FCustomTextFilterData& InFilterData, bool bApplyFilter);
+	void ModifyCustomTextFilter(const FCustomTextFilterData& InFilterData, const TSharedPtr<ICustomTextFilter<SceneOutliner::FilterBarType>>& InFilter);
+	void DeleteCustomTextFilter(const TSharedPtr<ICustomTextFilter<SceneOutliner::FilterBarType>>& InFilter);
 };
 
 struct SCENEOUTLINER_API FSceneOutlinerMenuHelper
