[Localization] Set the existing outliner columns with icon headers to have localized tooltips

diff --git a/Engine/Source/Editor/SceneOutliner/Private/SceneOutlinerGutter.cpp b/Engine/Source/Editor/SceneOutliner/Private/SceneOutlinerGutter.cpp
--- a/Engine/Source/Editor/SceneOutliner/Private/SceneOutlinerGutter.cpp
+++ b/Engine/Source/Editor/SceneOutliner/Private/SceneOutlinerGutter.cpp
@@ -325,7 +325,7 @@ SHeaderRow::FColumn::FArguments FSceneOutlinerGutter::ConstructHeaderRowColumn()
 		.VAlignHeader(VAlign_Center)
 		.HAlignCell(HAlign_Center)
 		.VAlignCell(VAlign_Center)
-		.DefaultTooltip(FText::FromName(GetColumnID()))
+		.DefaultTooltip(GetDisplayName())
 		.HeaderContentPadding(FMargin(4.0f, 0.0f, 0.0f, 0.0f))
 		[
 			SNew(SImage)
diff --git a/Engine/Source/Editor/SceneOutliner/Private/SceneOutlinerSourceControlColumn.cpp b/Engine/Source/Editor/SceneOutliner/Private/SceneOutlinerSourceControlColumn.cpp
--- a/Engine/Source/Editor/SceneOutliner/Private/SceneOutlinerSourceControlColumn.cpp
+++ b/Engine/Source/Editor/SceneOutliner/Private/SceneOutlinerSourceControlColumn.cpp
@@ -35,7 +35,7 @@ SHeaderRow::FColumn::FArguments FSceneOutlinerSourceControlColumn::ConstructHead
 		.VAlignHeader(VAlign_Center)
 		.HAlignCell(HAlign_Center)
 		.VAlignCell(VAlign_Center)
-		.DefaultTooltip(FText::FromName(GetColumnID()))
+		.DefaultTooltip(GetDisplayName())
 		[
 			HeaderRowIcon
 		];
diff --git a/Engine/Source/Editor/SceneOutliner/Private/SceneOutlinerUnsavedColumn.cpp b/Engine/Source/Editor/SceneOutliner/Private/SceneOutlinerUnsavedColumn.cpp
--- a/Engine/Source/Editor/SceneOutliner/Private/SceneOutlinerUnsavedColumn.cpp
+++ b/Engine/Source/Editor/SceneOutliner/Private/SceneOutlinerUnsavedColumn.cpp
@@ -170,7 +170,7 @@ SHeaderRow::FColumn::FArguments FSceneOutlinerActorUnsavedColumn::ConstructHeade
 		.VAlignHeader(VAlign_Center)
 		.HAlignCell(HAlign_Center)
 		.VAlignCell(VAlign_Center)
-		.DefaultTooltip(FText::FromName(GetColumnID()))
+		.DefaultTooltip(GetDisplayName())
 		[
 			SNew(SImage)
 			.Image(FAppStyle::GetBrush("Icons.DirtyBadge"))
diff --git a/Engine/Source/Editor/SceneOutliner/Private/SceneOutlinerUnsavedColumn.h b/Engine/Source/Editor/SceneOutliner/Private/SceneOutlinerUnsavedColumn.h
--- a/Engine/Source/Editor/SceneOutliner/Private/SceneOutlinerUnsavedColumn.h
+++ b/Engine/Source/Editor/SceneOutliner/Private/SceneOutlinerUnsavedColumn.h
@@ -19,6 +19,7 @@ public:
 	virtual ~FSceneOutlinerActorUnsavedColumn() = default;
 
 	static FName GetID() { return FSceneOutlinerBuiltInColumnTypes::Unsaved(); }
+	static FText GetDisplayName() { return FSceneOutlinerBuiltInColumnTypes::Unsaved_Localized(); }
 
 	//////////////////////////////////////////////////////////////////////////
 	// Begin ISceneOutlinerColumn Implementation
diff --git a/Engine/Source/Editor/SceneOutliner/Public/SceneOutlinerGutter.h b/Engine/Source/Editor/SceneOutliner/Public/SceneOutlinerGutter.h
--- a/Engine/Source/Editor/SceneOutliner/Public/SceneOutlinerGutter.h
+++ b/Engine/Source/Editor/SceneOutliner/Public/SceneOutlinerGutter.h
@@ -28,6 +28,8 @@ public:
 
 	static FName GetID() { return FSceneOutlinerBuiltInColumnTypes::Gutter(); }
 
+	static FText GetDisplayName() { return FSceneOutlinerBuiltInColumnTypes::Gutter_Localized(); }
+
 	// -----------------------------------------
 	// ISceneOutlinerColumn Implementation
 	virtual FName GetColumnID() override;
diff --git a/Engine/Source/Editor/SceneOutliner/Public/SceneOutlinerSourceControlColumn.h b/Engine/Source/Editor/SceneOutliner/Public/SceneOutlinerSourceControlColumn.h
--- a/Engine/Source/Editor/SceneOutliner/Public/SceneOutlinerSourceControlColumn.h
+++ b/Engine/Source/Editor/SceneOutliner/Public/SceneOutlinerSourceControlColumn.h
@@ -22,6 +22,8 @@ public:
 
 	static FName GetID() { return FSceneOutlinerBuiltInColumnTypes::SourceControl(); }
 
+	static FText GetDisplayName() { return FSceneOutlinerBuiltInColumnTypes::SourceControl_Localized(); }
+
 	//////////////////////////////////////////////////////////////////////////
 	// Begin ISceneOutlinerColumn Implementation
 
