Scene Outliner: Add functionality to customize the "primary" collumn instead of hardcoding it to the label column. The primary column is the one used to initially sort the Outliner and where the hierarchy arrows show up

diff --git a/Engine/Source/Editor/SceneOutliner/Private/SOutlinerTreeView.cpp b/Engine/Source/Editor/SceneOutliner/Private/SOutlinerTreeView.cpp
--- a/Engine/Source/Editor/SceneOutliner/Private/SOutlinerTreeView.cpp
+++ b/Engine/Source/Editor/SceneOutliner/Private/SOutlinerTreeView.cpp
@@ -291,7 +291,7 @@ TSharedRef<SWidget> SSceneOutlinerTreeRow::GenerateWidgetForColumn( const FName&
 		NewItemWidget = Column->ConstructRowWidget(ItemPtr.ToSharedRef(), *this);
 	}
 
-	if( ColumnName == FSceneOutlinerBuiltInColumnTypes::Label() )
+	if( ColumnName == Outliner->GetPrimaryColumnName() )
 	{
 		// The first column gets the tree expansion arrow for this row
 		return SNew(SBox)
diff --git a/Engine/Source/Editor/SceneOutliner/Private/SSceneOutliner.cpp b/Engine/Source/Editor/SceneOutliner/Private/SSceneOutliner.cpp
--- a/Engine/Source/Editor/SceneOutliner/Private/SSceneOutliner.cpp
+++ b/Engine/Source/Editor/SceneOutliner/Private/SSceneOutliner.cpp
@@ -86,8 +86,9 @@ void SSceneOutliner::Construct(const FArguments& InArgs, const FSceneOutlinerIni
 	bSelectionDirty = true;
 	SortOutlinerTimer = 0.0f;
 	bPendingFocusNextFrame = InInitOptions.bFocusSearchBoxWhenOpened;
+	PrimaryColumnName = InInitOptions.PrimaryColumnName;
 	
-	SortByColumn = FSceneOutlinerBuiltInColumnTypes::Label();
+	SortByColumn = PrimaryColumnName;
 	SortMode = EColumnSortMode::Ascending;
 
 	UOutlinerConfig::Initialize();
@@ -1792,6 +1793,11 @@ void SSceneOutliner::DuplicateFoldersHierarchy()
 	}
 }
 
+FName SSceneOutliner::GetPrimaryColumnName() const
+{
+	return PrimaryColumnName;
+}
+
 void SSceneOutliner::DeleteFoldersBegin()
 {
 	CacheFoldersDelete.Empty();
diff --git a/Engine/Source/Editor/SceneOutliner/Public/SSceneOutliner.h b/Engine/Source/Editor/SceneOutliner/Public/SSceneOutliner.h
--- a/Engine/Source/Editor/SceneOutliner/Public/SSceneOutliner.h
+++ b/Engine/Source/Editor/SceneOutliner/Public/SSceneOutliner.h
@@ -807,6 +807,8 @@ public:
 	/** Duplicates current folder and all descendants */
 	UE_API void DuplicateFoldersHierarchy();
 
+	FName GetPrimaryColumnName() const;
+
 private:
 	/** Cache selected folders during edit delete */
 	TArray<FFolderTreeItem*> CacheFoldersDelete;
@@ -1002,6 +1004,9 @@ private:
 	SceneOutliner::FTreeItemPtrEvent OnDoubleClickOnTreeEvent;
 
 	SceneOutliner::FOnItemSelectionChanged OnItemSelectionChanged;
+	
+	/** The primary column for this outliner, used to determine initial sorting and where the hierarchy arrows appear */
+	FName PrimaryColumnName;
 private:
 
 	virtual void AddReferencedObjects(FReferenceCollector& Collector) override {};
diff --git a/Engine/Source/Editor/SceneOutliner/Public/SceneOutlinerPublicTypes.h b/Engine/Source/Editor/SceneOutliner/Public/SceneOutlinerPublicTypes.h
--- a/Engine/Source/Editor/SceneOutliner/Public/SceneOutlinerPublicTypes.h
+++ b/Engine/Source/Editor/SceneOutliner/Public/SceneOutlinerPublicTypes.h
@@ -203,6 +203,9 @@ struct FSceneOutlinerInitializationOptions : FSharedSceneOutlinerData
 	/** Init options related to the filter bar */
 	FSceneOutlinerFilterBarOptions FilterBarOptions;
 
+	/** The primary column for this outliner, used to determine initial sorting and where the hierarchy arrows appear */
+	FName PrimaryColumnName;
+
 public:
 
 	/** Constructor */
@@ -215,6 +218,7 @@ public:
 		, Filters( new FSceneOutlinerFilters )
 		, ModeFactory()
 		, OutlinerIdentifier(NAME_None)
+		, PrimaryColumnName(FSceneOutlinerBuiltInColumnTypes::Label())
 	{}
 };
 
