SceneOutliner Unsaved Column: show a dimmed icon for entries that can be saved in their own file.
This makes it possible to see if an entity is saved in its own file or find the nearest parent in the hierarchy that will save this entity.

diff --git a/Engine/Source/Editor/EditorStyle/Private/StarshipStyle.cpp b/Engine/Source/Editor/EditorStyle/Private/StarshipStyle.cpp
--- a/Engine/Source/Editor/EditorStyle/Private/StarshipStyle.cpp
+++ b/Engine/Source/Editor/EditorStyle/Private/StarshipStyle.cpp
@@ -419,6 +419,7 @@ void FStarshipEditorStyle::FStyle::SetupGeneralStyles()
 		Set("Icons.SaveChanged", new IMAGE_BRUSH_SVG( "Starship/Common/SaveChanged", Icon16x16 ) );
 
 		Set("Icons.DirtyBadge", new IMAGE_BRUSH_SVG("Starship/Common/DirtyBadge", Icon12x12));
+		Set("Icons.SaveableBadge", new IMAGE_BRUSH_SVG("Starship/Common/DirtyBadge", Icon12x12, FLinearColor(0.5f, 0.5f, 0.5f, 0.3f)));
 		Set("Icons.MakeStaticMesh", new IMAGE_BRUSH_SVG("Starship/Common/MakeStaticMesh", Icon16x16));
 		Set("Icons.Documentation", new IMAGE_BRUSH_SVG("Starship/Common/Documentation", Icon16x16));
 		Set("Icons.Support", new IMAGE_BRUSH_SVG("Starship/Common/Support", Icon16x16));
diff --git a/Engine/Source/Editor/SceneOutliner/Private/SceneOutlinerUnsavedColumn.cpp b/Engine/Source/Editor/SceneOutliner/Private/SceneOutlinerUnsavedColumn.cpp
--- a/Engine/Source/Editor/SceneOutliner/Private/SceneOutlinerUnsavedColumn.cpp
+++ b/Engine/Source/Editor/SceneOutliner/Private/SceneOutlinerUnsavedColumn.cpp
@@ -36,6 +36,7 @@ private:
 	FString ExternalPackageFilename;
 	TWeakPtr<ISceneOutlinerTreeItem> WeakTreeItem;
 	bool bIsUnsaved;
+	bool bIsSavable;
 
 	FDelegateHandle OnPackagingModeChangedHandle;
 	FDelegateHandle OnUnsavedAssetAddedHandle;
@@ -111,6 +112,7 @@ void SUnsavedActorWidget::UpdateExternalPackageFilename()
 	}
 
 	bIsUnsaved = UnsavedAssetsTrackerModule.IsAssetUnsaved(ExternalPackageFilename);
+	bIsSavable = !ExternalPackageFilename.IsEmpty();
 
 	UpdateImage();
 }
@@ -150,6 +152,10 @@ void SUnsavedActorWidget::UpdateImage()
 	{
 		SetImage(FAppStyle::GetBrush("Icons.DirtyBadge"));
 	}
+	else if (bIsSavable)
+	{
+		SetImage(FAppStyle::GetBrush("Icons.SaveableBadge"));
+	}
 	else
 	{
 		SetImage(nullptr);
