SceneOutliner: Don't use placement new for simple value type

#rnx

diff --git a/Engine/Source/Editor/SceneOutliner/Public/SceneOutlinerStandaloneTypes.h b/Engine/Source/Editor/SceneOutliner/Public/SceneOutlinerStandaloneTypes.h
--- a/Engine/Source/Editor/SceneOutliner/Public/SceneOutlinerStandaloneTypes.h
+++ b/Engine/Source/Editor/SceneOutliner/Public/SceneOutlinerStandaloneTypes.h
@@ -70,7 +70,8 @@ public:
 	/** ID representing a generic tree item */
 	FSceneOutlinerTreeItemID(const FUniqueID& CustomID) : Type(EType::UniqueID)
 	{
-		new (Data) FUniqueID(CustomID);
+		// placement new not necessary for simple value type
+		*reinterpret_cast<FUniqueID*>(Data) = CustomID;
 		CachedHash = CalculateTypeHash();
 	}
 
