[Backout] - CL45302588
#rnx
Original CL Desc
-----------------------------------------------------------------
SceneOutliner: Don't use placement new for simple value type

#rnx

diff --git a/Engine/Source/Editor/SceneOutliner/Public/SceneOutlinerStandaloneTypes.h b/Engine/Source/Editor/SceneOutliner/Public/SceneOutlinerStandaloneTypes.h
--- a/Engine/Source/Editor/SceneOutliner/Public/SceneOutlinerStandaloneTypes.h
+++ b/Engine/Source/Editor/SceneOutliner/Public/SceneOutlinerStandaloneTypes.h
@@ -70,8 +70,7 @@ public:
 	/** ID representing a generic tree item */
 	FSceneOutlinerTreeItemID(const FUniqueID& CustomID) : Type(EType::UniqueID)
 	{
-		// placement new not necessary for simple value type
-		*reinterpret_cast<FUniqueID*>(Data) = CustomID;
+		new (Data) FUniqueID(CustomID);
 		CachedHash = CalculateTypeHash();
 	}
 
