SceneOutliner: alignas(8) Data to fix clang analysis warning

diff --git a/Engine/Source/Editor/SceneOutliner/Public/SceneOutlinerStandaloneTypes.h b/Engine/Source/Editor/SceneOutliner/Public/SceneOutlinerStandaloneTypes.h
--- a/Engine/Source/Editor/SceneOutliner/Public/SceneOutlinerStandaloneTypes.h
+++ b/Engine/Source/Editor/SceneOutliner/Public/SceneOutlinerStandaloneTypes.h
@@ -178,7 +178,7 @@ private:
 
 	uint32 CachedHash;
 	static const uint32 MaxSize = TMaxSizeof<FObjectKey, FFolderKey, FGuid, FUniqueID>::Value;
-	mutable uint8 Data[MaxSize];
+	alignas(8) mutable uint8 Data[MaxSize];
 };
 	
 struct FSceneOutlinerTreeItemType
