Add GetTypeHash to FSceneOutlinerTreeItemType

diff --git a/Engine/Source/Editor/SceneOutliner/Public/SceneOutlinerStandaloneTypes.h b/Engine/Source/Editor/SceneOutliner/Public/SceneOutlinerStandaloneTypes.h
--- a/Engine/Source/Editor/SceneOutliner/Public/SceneOutlinerStandaloneTypes.h
+++ b/Engine/Source/Editor/SceneOutliner/Public/SceneOutlinerStandaloneTypes.h
@@ -195,6 +195,11 @@ public:
 	{
 		return (ID == Other.ID) || (ParentType && ParentType->IsA(Other));
 	}
+	
+	friend uint32 GetTypeHash(const FSceneOutlinerTreeItemType& Type)
+	{
+		return Type.ID; 
+	}
 
 private:
 	static UE_API uint32 NextUniqueID;
