Added IsA(Type& OtherType) in addition to IsA<OtherType> ISceneOutlinerTreeItem
Added GetType to ISceneOutlinerTreeItem

diff --git a/Engine/Source/Editor/SceneOutliner/Public/ISceneOutlinerTreeItem.h b/Engine/Source/Editor/SceneOutliner/Public/ISceneOutlinerTreeItem.h
--- a/Engine/Source/Editor/SceneOutliner/Public/ISceneOutlinerTreeItem.h
+++ b/Engine/Source/Editor/SceneOutliner/Public/ISceneOutlinerTreeItem.h
@@ -130,6 +130,13 @@ public:
 	{
 		return TreeType.IsA(T::Type);
 	}
+	
+	bool IsA(const FSceneOutlinerTreeItemType& InType) const
+	{
+		return TreeType.IsA(InType);
+	}
+	
+	const FSceneOutlinerTreeItemType& GetType() const { return TreeType; }
 
 	/** Returns true if the data the item references is valid */
 	virtual bool IsValid() const = 0;
