Added ability for ISceneOutlinerMode implementations to handle updating DragDrop operation decorators (to handle custom non-FDecoratedDragDropOp implementations)

#FailedValidations Preflight

diff --git a/Engine/Source/Editor/SceneOutliner/Private/SOutlinerTreeView.cpp b/Engine/Source/Editor/SceneOutliner/Private/SOutlinerTreeView.cpp
--- a/Engine/Source/Editor/SceneOutliner/Private/SOutlinerTreeView.cpp
+++ b/Engine/Source/Editor/SceneOutliner/Private/SOutlinerTreeView.cpp
@@ -15,8 +15,16 @@
 
 #define LOCTEXT_NAMESPACE "SSceneOutliner"
 
-static void UpdateOperationDecorator(const FDragDropEvent& Event, const FSceneOutlinerDragValidationInfo& ValidationInfo)
+static void UpdateOperationDecorator(TWeakPtr<SSceneOutliner> WeakSceneOutlinerPtr, const FDragDropEvent& Event, const FSceneOutlinerDragValidationInfo& ValidationInfo)
 {
+	if (TSharedPtr<SSceneOutliner> SharedSceneOutliner = WeakSceneOutlinerPtr.Pin())
+	{
+		if (SharedSceneOutliner->UpdateOperationDecorator(Event, &ValidationInfo))
+		{
+			return;
+		}
+	}
+
 	const FSlateBrush* Icon = ValidationInfo.IsValid() ? FAppStyle::GetBrush(TEXT("Graph.ConnectorFeedback.OK")) : FAppStyle::GetBrush(TEXT("Graph.ConnectorFeedback.Error"));
 
 	FDragDropOperation* Operation = Event.GetOperation().Get();
@@ -27,8 +35,16 @@ static void UpdateOperationDecorator(const FDragDropEvent& Event, const FSceneOu
 	}
 }
 
-static void ResetOperationDecorator(const FDragDropEvent& Event)
+static void ResetOperationDecorator(TWeakPtr<SSceneOutliner> WeakSceneOutlinerPtr, const FDragDropEvent& Event)
 {
+	if (TSharedPtr<SSceneOutliner> SharedSceneOutliner = WeakSceneOutlinerPtr.Pin())
+	{
+		if (SharedSceneOutliner->UpdateOperationDecorator(Event, nullptr))
+		{
+			return;
+		}
+	}
+
 	FDragDropOperation* Operation = Event.GetOperation().Get();
 	if (Operation)
 	{
@@ -132,7 +148,7 @@ FReply SSceneOutlinerTreeView::OnDragOver(const FGeometry& MyGeometry, const FDr
 	auto Reply = HandleDropFromWeak(SceneOutlinerWeak, DragDropEvent, ValidationInfo);
 	if (Reply.IsEventHandled())
 	{
-		UpdateOperationDecorator(DragDropEvent, ValidationInfo);
+		UpdateOperationDecorator(SceneOutlinerWeak, DragDropEvent, ValidationInfo);
 	}
 
 	return Reply;
@@ -147,7 +163,7 @@ void SSceneOutlinerTreeView::OnDragLeave(const FDragDropEvent& DragDropEvent)
 
 	if( SceneOutlinerWeak.Pin()->GetSharedData().bShowParentTree )
 	{
-		ResetOperationDecorator(DragDropEvent);
+		ResetOperationDecorator(SceneOutlinerWeak, DragDropEvent);
 	}
 }
 
@@ -197,17 +213,14 @@ void SSceneOutlinerTreeRow::OnDragEnter( const FGeometry& MyGeometry, const FDra
 		FReply Reply = HandleDrop(SceneOutlinerPtr, DragDropEvent, *ItemPtr, ValidationInfo, false);
 		if (Reply.IsEventHandled())
 		{
-			UpdateOperationDecorator(DragDropEvent, ValidationInfo);
+			UpdateOperationDecorator(SceneOutlinerPtr, DragDropEvent, ValidationInfo);
 		}
 	}
 }
 
 void SSceneOutlinerTreeRow::OnDragLeave( const FDragDropEvent& DragDropEvent )
 {
-	auto ItemPtr = Item.Pin();
-	auto SceneOutlinerPtr = SceneOutlinerWeak.Pin();
-
-	ResetOperationDecorator(DragDropEvent);
+	ResetOperationDecorator(SceneOutlinerWeak, DragDropEvent);
 }
 
 FReply SSceneOutlinerTreeRow::OnDragOver(const FGeometry& MyGeometry, const FDragDropEvent& DragDropEvent)
@@ -553,14 +566,14 @@ void SSceneOutlinerPinnedTreeRow::OnDragEnter(const FGeometry& MyGeometry, const
 		FReply Reply = HandleDrop(SceneOutlinerPtr, DragDropEvent, *ItemPtr, ValidationInfo, false);
 		if (Reply.IsEventHandled())
 		{
-			UpdateOperationDecorator(DragDropEvent, ValidationInfo);
+			UpdateOperationDecorator(SceneOutlinerPtr, DragDropEvent, ValidationInfo);
 		}
 	}
 }
 
 void SSceneOutlinerPinnedTreeRow::OnDragLeave(const FDragDropEvent& DragDropEvent)
 {
-	ResetOperationDecorator(DragDropEvent);
+	ResetOperationDecorator(SceneOutlinerWeak, DragDropEvent);
 }
 
 FReply SSceneOutlinerPinnedTreeRow::OnDragOver(const FGeometry& MyGeometry, const FDragDropEvent& DragDropEvent)
diff --git a/Engine/Source/Editor/SceneOutliner/Private/SSceneOutliner.cpp b/Engine/Source/Editor/SceneOutliner/Private/SSceneOutliner.cpp
--- a/Engine/Source/Editor/SceneOutliner/Private/SSceneOutliner.cpp
+++ b/Engine/Source/Editor/SceneOutliner/Private/SSceneOutliner.cpp
@@ -2554,6 +2554,11 @@ FReply SSceneOutliner::OnDragOverItem(const FDragDropEvent& Event, const ISceneO
 	return Mode->OnDragOverItem(Event, Item); 
 }
 
+bool SSceneOutliner::UpdateOperationDecorator(const FDragDropEvent& Event, const FSceneOutlinerDragValidationInfo* ValidationInfo) const
+{
+	return Mode->UpdateOperationDecorator(Event, ValidationInfo); 
+}
+
 void SSceneOutliner::PinItems(const TArray<FSceneOutlinerTreeItemPtr>& InItems)
 {
 	Mode->PinItems(InItems);
diff --git a/Engine/Source/Editor/SceneOutliner/Public/ISceneOutlinerMode.h b/Engine/Source/Editor/SceneOutliner/Public/ISceneOutlinerMode.h
--- a/Engine/Source/Editor/SceneOutliner/Public/ISceneOutlinerMode.h
+++ b/Engine/Source/Editor/SceneOutliner/Public/ISceneOutlinerMode.h
@@ -63,6 +63,9 @@ public:
 
 	/** Called when a payload is dragged over an item */
 	virtual FReply OnDragOverItem(const FDragDropEvent& Event, const ISceneOutlinerTreeItem& Item) const { return FReply::Handled(); }
+	
+	/** Called when a drag-over/enter/leave warrants updating the DragDrop decorator, (ValidationInfo == nullptr, indicates resetting the decorator) */ 
+	virtual bool UpdateOperationDecorator(const FDragDropEvent& Event, const FSceneOutlinerDragValidationInfo* ValidationInfo) const { return false; }
 public:
 	/* Events */
 
diff --git a/Engine/Source/Editor/SceneOutliner/Public/SSceneOutliner.h b/Engine/Source/Editor/SceneOutliner/Public/SSceneOutliner.h
--- a/Engine/Source/Editor/SceneOutliner/Public/SSceneOutliner.h
+++ b/Engine/Source/Editor/SceneOutliner/Public/SSceneOutliner.h
@@ -551,6 +551,9 @@ public:
 
 	/** Called when a payload is dragged over an item */
 	UE_API FReply OnDragOverItem(const FDragDropEvent& Event, const ISceneOutlinerTreeItem& Item) const;
+	
+	/** Called when a drag-over/enter/leave warrants updating the DragDrop decorator */ 
+	UE_API bool UpdateOperationDecorator(const FDragDropEvent& Event, const FSceneOutlinerDragValidationInfo* ValidationInfo) const;	
 
 	/** Used to test if Outliner related selection changes have already been handled */
 	bool GetIsReentrant() const { return bIsReentrant; }
