Scene Outliner: Add an OnFullRefreshDone delegate so that other systems are able to determine when the outliner is done scanning.

#rnx

diff --git a/Engine/Source/Editor/SceneOutliner/Private/SSceneOutliner.cpp b/Engine/Source/Editor/SceneOutliner/Private/SSceneOutliner.cpp
--- a/Engine/Source/Editor/SceneOutliner/Private/SSceneOutliner.cpp
+++ b/Engine/Source/Editor/SceneOutliner/Private/SSceneOutliner.cpp
@@ -772,6 +772,8 @@ void SSceneOutliner::Populate()
 	}
 	PendingFoldersSelect.Empty();
 
+	bool bFullRefreshDone = false;
+
 	// Check if we need to sort because we are finished with the populating operations
 	bool bFinalSort = false;
 	if (PendingOperations.Num() == 0)
@@ -789,6 +791,8 @@ void SSceneOutliner::Populate()
 					ScrollItemIntoView(ItemToScroll);
 				}
 			}
+
+			bFullRefreshDone = true;
 		}
 
 		bProcessingFullRefresh = false;
@@ -807,6 +811,11 @@ void SSceneOutliner::Populate()
 	{
 		RequestSort();
 	}
+
+	if (bFullRefreshDone)
+	{
+		OnFullRefreshDone.Broadcast();
+	}
 }
 
 bool SSceneOutliner::ShouldShowFolders() const
diff --git a/Engine/Source/Editor/SceneOutliner/Public/SSceneOutliner.h b/Engine/Source/Editor/SceneOutliner/Public/SSceneOutliner.h
--- a/Engine/Source/Editor/SceneOutliner/Public/SSceneOutliner.h
+++ b/Engine/Source/Editor/SceneOutliner/Public/SSceneOutliner.h
@@ -431,6 +431,9 @@ public:
 		*/
 	SceneOutliner::FOnItemSelectionChanged& GetOnItemSelectionChanged() { return OnItemSelectionChanged; }
 
+	/** Event that's broadcast after the scene outliner has been re-populated. */
+	UE_API FSimpleMulticastDelegate& GetOnFullRefreshDone() { return OnFullRefreshDone; }
+
 	/** Set the item selection of the outliner based on a selector function. Any items which return true will be added */
 	UE_API virtual void SetSelection(const TFunctionRef<bool(ISceneOutlinerTreeItem&)> Selector) override;
 
@@ -1007,6 +1010,8 @@ private:
 	SceneOutliner::FTreeItemPtrEvent OnDoubleClickOnTreeEvent;
 
 	SceneOutliner::FOnItemSelectionChanged OnItemSelectionChanged;
+
+	FSimpleMulticastDelegate OnFullRefreshDone;
 	
 	/** The primary column for this outliner, used to determine initial sorting and where the hierarchy arrows appear */
 	FName PrimaryColumnName;
