Slate Dynamic Invalidation - Scene Outliner

diff --git a/Engine/Source/Editor/SceneOutliner/Private/ActorFolderTreeItem.cpp b/Engine/Source/Editor/SceneOutliner/Private/ActorFolderTreeItem.cpp
--- a/Engine/Source/Editor/SceneOutliner/Private/ActorFolderTreeItem.cpp
+++ b/Engine/Source/Editor/SceneOutliner/Private/ActorFolderTreeItem.cpp
@@ -24,6 +24,12 @@
 
 const FSceneOutlinerTreeItemType FActorFolderTreeItem::Type(&FFolderTreeItem::Type);
 
+template<>
+struct TWidgetTypeTraits<struct SActorFolderTreeLabel>
+{
+	static constexpr bool SupportsInvalidation() { return true; }
+};
+
 struct SActorFolderTreeLabel : FSceneOutlinerCommonLabelData, public SCompoundWidget
 {
 	SLATE_BEGIN_ARGS(SActorFolderTreeLabel) {}
diff --git a/Engine/Source/Editor/SceneOutliner/Private/ActorTreeItem.cpp b/Engine/Source/Editor/SceneOutliner/Private/ActorTreeItem.cpp
--- a/Engine/Source/Editor/SceneOutliner/Private/ActorTreeItem.cpp
+++ b/Engine/Source/Editor/SceneOutliner/Private/ActorTreeItem.cpp
@@ -28,6 +28,12 @@
 
 const FSceneOutlinerTreeItemType FActorTreeItem::Type(&IActorBaseTreeItem::Type);
 
+template<>
+struct TWidgetTypeTraits<struct SActorTreeLabel>
+{
+	static constexpr bool SupportsInvalidation() { return true; }
+};
+
 struct SActorTreeLabel : FSceneOutlinerCommonLabelData, public SCompoundWidget
 {
 	SLATE_BEGIN_ARGS(SActorTreeLabel) {}
diff --git a/Engine/Source/Editor/SceneOutliner/Private/SOutlinerTreeView.cpp b/Engine/Source/Editor/SceneOutliner/Private/SOutlinerTreeView.cpp
--- a/Engine/Source/Editor/SceneOutliner/Private/SOutlinerTreeView.cpp
+++ b/Engine/Source/Editor/SceneOutliner/Private/SOutlinerTreeView.cpp
@@ -357,7 +357,15 @@ const float SSceneOutlinerTreeRow::LabelChangedAnimOffsetPercent = 0.2f;
 
 void SSceneOutlinerTreeRow::FlashHighlight()
 {
-    LastHighlightInteractionTime = FSlateApplication::Get().GetCurrentTime();
+	LastHighlightInteractionTime = FSlateApplication::Get().GetCurrentTime();
+	Invalidate(EInvalidateWidgetReason::PaintAndVolatility);
+}
+
+bool SSceneOutlinerTreeRow::ComputeVolatility() const
+{
+	const double CurrentTime = FSlateApplication::Get().GetCurrentTime();
+	const float TimeSinceHighlightInteraction = (float)( CurrentTime - LastHighlightInteractionTime );
+	return TimeSinceHighlightInteraction <= HighlightTargetEffectDuration;
 }
 
 void SSceneOutlinerTreeRow::Tick( const FGeometry& AllottedGeometry, const double InCurrentTime, const float InDeltaTime )
@@ -365,9 +373,6 @@ void SSceneOutlinerTreeRow::Tick( const FGeometry& AllottedGeometry, const doubl
 	// Call parent implementation.
 	SCompoundWidget::Tick( AllottedGeometry, InCurrentTime, InDeltaTime );
 
-	// We'll draw with the 'focused' look if we're either focused or we have a context menu summoned
-	const bool bShouldAppearFocused = HasKeyboardFocus();
-
 	// Update highlight 'target' effect
 	{
 		const float HighlightLeftX = HighlightRectLeftOffset;
@@ -377,11 +382,17 @@ void SSceneOutlinerTreeRow::Tick( const FGeometry& AllottedGeometry, const doubl
 		HighlightTargetRightSpring.SetTarget( HighlightRightX );
 
 		float TimeSinceHighlightInteraction = (float)( InCurrentTime - LastHighlightInteractionTime );
-		if( TimeSinceHighlightInteraction <= HighlightTargetEffectDuration || bShouldAppearFocused )
+		if( TimeSinceHighlightInteraction <= HighlightTargetEffectDuration)
 		{
 			HighlightTargetLeftSpring.Tick( InDeltaTime );
 			HighlightTargetRightSpring.Tick( InDeltaTime );
 		}
+		else if (const float LastTimeSinceHighlightInteraction = (float)( (InCurrentTime - InDeltaTime) - LastHighlightInteractionTime );
+			LastTimeSinceHighlightInteraction <= HighlightTargetEffectDuration)
+		{
+			// We only want to invalidate once after the highlight animation is done, so if it was shown last frame, we recompute volatility
+			Invalidate(EInvalidateWidgetReason::PaintAndVolatility);
+		}
 	}
 }
 
@@ -397,9 +408,6 @@ int32 SSceneOutlinerTreeRow::OnPaint( const FPaintArgs& Args, const FGeometry& A
 
 	const double CurrentTime = FSlateApplication::Get().GetCurrentTime();
 
-	// We'll draw with the 'focused' look if we're either focused or we have a context menu summoned
-	const bool bShouldAppearFocused = HasKeyboardFocus();
-
 	// Draw highlight targeting effect
 	const float TimeSinceHighlightInteraction = (float)( CurrentTime - LastHighlightInteractionTime );
 	if( TimeSinceHighlightInteraction <= HighlightTargetEffectDuration )
diff --git a/Engine/Source/Editor/SceneOutliner/Private/SSceneOutliner.cpp b/Engine/Source/Editor/SceneOutliner/Private/SSceneOutliner.cpp
--- a/Engine/Source/Editor/SceneOutliner/Private/SSceneOutliner.cpp
+++ b/Engine/Source/Editor/SceneOutliner/Private/SSceneOutliner.cpp
@@ -31,6 +31,7 @@
 #include "Widgets/Notifications/SNotificationList.h"
 #include "Widgets/Input/SMultiLineEditableTextBox.h"
 #include "DetailLayoutBuilder.h"
+#include "Widgets/SInvalidationPanel.h"
 
 DEFINE_LOG_CATEGORY(LogSceneOutliner);
 
@@ -263,55 +264,59 @@ void SSceneOutliner::Construct(const FArguments& InArgs, const FSceneOutlinerIni
 
 		+SOverlay::Slot()
 		[
-			SAssignNew( OutlinerTreeView, SSceneOutlinerTreeView, StaticCastSharedRef<SSceneOutliner>(AsShared()) )
+			SNew(SInvalidationPanel)
+			.UseDynamicInvalidation(true)
+			[
+				SAssignNew( OutlinerTreeView, SSceneOutlinerTreeView, StaticCastSharedRef<SSceneOutliner>(AsShared()) )
 
-			// Determined by the mode
-			.SelectionMode( this, &SSceneOutliner::GetSelectionMode )
+				// Determined by the mode
+				.SelectionMode( this, &SSceneOutliner::GetSelectionMode )
 
-			// Point the tree to our array of root-level items.  Whenever this changes, we'll call RequestTreeRefresh()
-			.TreeItemsSource( &RootTreeItems )
+				// Point the tree to our array of root-level items.  Whenever this changes, we'll call RequestTreeRefresh()
+				.TreeItemsSource( &RootTreeItems )
 
-			// Find out when the user selects something in the tree
-			.OnSelectionChanged( this, &SSceneOutliner::OnOutlinerTreeSelectionChanged )
+				// Find out when the user selects something in the tree
+				.OnSelectionChanged( this, &SSceneOutliner::OnOutlinerTreeSelectionChanged )
 
-			// Called when the user double-clicks with LMB on an item in the list
-			.OnMouseButtonDoubleClick( this, &SSceneOutliner::OnOutlinerTreeDoubleClick )
+				// Called when the user double-clicks with LMB on an item in the list
+				.OnMouseButtonDoubleClick( this, &SSceneOutliner::OnOutlinerTreeDoubleClick )
 
-			// Called when the user single-clicks with LMB on an item in the list
-			.OnMouseButtonClick( this, &SSceneOutliner::OnOutlinerTreeSingleClick )
+				// Called when the user single-clicks with LMB on an item in the list
+				.OnMouseButtonClick( this, &SSceneOutliner::OnOutlinerTreeSingleClick )
 
-			// Called when an item is scrolled into view
-			.OnItemScrolledIntoView( this, &SSceneOutliner::OnOutlinerTreeItemScrolledIntoView )
+				// Called when an item is scrolled into view
+				.OnItemScrolledIntoView( this, &SSceneOutliner::OnOutlinerTreeItemScrolledIntoView )
 
-			// Called when an item is expanded or collapsed
-			.OnExpansionChanged(this, &SSceneOutliner::OnItemExpansionChanged)
+				// Called when an item is expanded or collapsed
+				.OnExpansionChanged(this, &SSceneOutliner::OnItemExpansionChanged)
 
-			// Called to child items for any given parent item
-			.OnGetChildren( this, &SSceneOutliner::OnGetChildrenForOutlinerTree )
+				// Called to child items for any given parent item
+				.OnGetChildren( this, &SSceneOutliner::OnGetChildrenForOutlinerTree )
 
-			// Generates the actual widget for a tree item
-			.OnGenerateRow( this, &SSceneOutliner::OnGenerateRowForOutlinerTree ) 
+				// Generates the actual widget for a tree item
+				.OnGenerateRow( this, &SSceneOutliner::OnGenerateRowForOutlinerTree ) 
 
-			// Generates the actual widget for a pinned tree item
-			.OnGeneratePinnedRow(this, &SSceneOutliner::OnGeneratePinnedRowForOutlinerTree)
+				// Generates the actual widget for a pinned tree item
+				.OnGeneratePinnedRow(this, &SSceneOutliner::OnGeneratePinnedRowForOutlinerTree)
 
-			// Use the level viewport context menu as the right click menu for tree items
-			.OnContextMenuOpening(this, &SSceneOutliner::OnOpenContextMenu)
+				// Use the level viewport context menu as the right click menu for tree items
+				.OnContextMenuOpening(this, &SSceneOutliner::OnOpenContextMenu)
 
-			// Header for the tree
-			.HeaderRow( HeaderRowWidget )
+				// Header for the tree
+				.HeaderRow( HeaderRowWidget )
 
-			// Called when an item is expanded or collapsed with the shift-key pressed down
-			.OnSetExpansionRecursive(this, &SSceneOutliner::SetItemExpansionRecursive)
+				// Called when an item is expanded or collapsed with the shift-key pressed down
+				.OnSetExpansionRecursive(this, &SSceneOutliner::SetItemExpansionRecursive)
 
-			// Make it easier to see hierarchies when there are a lot of items
-			.HighlightParentNodesForSelection(true)
+				// Make it easier to see hierarchies when there are a lot of items
+				.HighlightParentNodesForSelection(true)
 
-			// Show the Hierarchy of actors pinned at the top of the tree view
-			.ShouldStackHierarchyHeaders(this, &SSceneOutliner::ShouldStackHierarchyHeaders)
+				// Show the Hierarchy of actors pinned at the top of the tree view
+				.ShouldStackHierarchyHeaders(this, &SSceneOutliner::ShouldStackHierarchyHeaders)
 
-			// Preserve the selection when the selected item is hidden due to a parent collapsing
-			.AllowInvisibleItemSelection(true)
+				// Preserve the selection when the selected item is hidden due to a parent collapsing
+				.AllowInvisibleItemSelection(true)
+			]
 		]
 	];
 
diff --git a/Engine/Source/Editor/SceneOutliner/Private/SceneOutlinerUnsavedColumn.cpp b/Engine/Source/Editor/SceneOutliner/Private/SceneOutlinerUnsavedColumn.cpp
--- a/Engine/Source/Editor/SceneOutliner/Private/SceneOutlinerUnsavedColumn.cpp
+++ b/Engine/Source/Editor/SceneOutliner/Private/SceneOutlinerUnsavedColumn.cpp
@@ -12,6 +12,12 @@
 
 #define LOCTEXT_NAMESPACE "SceneOutlinerActorUnsavedColumn"
 
+template<>
+struct TWidgetTypeTraits<class SUnsavedActorWidget>
+{
+	static constexpr bool SupportsInvalidation() { return true; }
+};
+
 class SUnsavedActorWidget : public SImage
 {
 public:
diff --git a/Engine/Source/Editor/SceneOutliner/Private/WorldTreeItem.cpp b/Engine/Source/Editor/SceneOutliner/Private/WorldTreeItem.cpp
--- a/Engine/Source/Editor/SceneOutliner/Private/WorldTreeItem.cpp
+++ b/Engine/Source/Editor/SceneOutliner/Private/WorldTreeItem.cpp
@@ -77,6 +77,12 @@ namespace SceneOutliner
 const FSceneOutlinerTreeItemType FWorldTreeItem::Type(&ISceneOutlinerTreeItem::Type);
 TMap<FName, FWorldTreeItem::FContextMenuProviderFunc> FWorldTreeItem::ContextMenuProviders;
 
+template<>
+struct TWidgetTypeTraits<struct SWorldTreeLabel>
+{
+	static constexpr bool SupportsInvalidation() { return true; }
+};
+
 struct SWorldTreeLabel : FSceneOutlinerCommonLabelData, public SCompoundWidget
 {
 	SLATE_BEGIN_ARGS(SWorldTreeLabel) {}
@@ -140,7 +146,7 @@ private:
 	{
 		auto Item = TreeItemPtr.Pin();
 		FText PersistentLevelDisplayName = Item.IsValid() ? FText::FromString(Item->GetWorldName()) : FText();
-		if (Item->CanInteract())
+		if (Item.IsValid() && Item->CanInteract())
 		{
 			return FText::Format(LOCTEXT("WorldLabel_Tooltip", "The {0} level. Double-click to edit World Settings."), PersistentLevelDisplayName);
 		}
@@ -163,9 +169,9 @@ private:
 	FReply HandleDoubleClick(const FGeometry& MyGeometry, const FPointerEvent& PointerEvent) const
 	{
 		auto Item = TreeItemPtr.Pin();
-		if (Item->CanInteract())
+		if (Item.IsValid() && Item->CanInteract())
 		{
-			TreeItemPtr.Pin()->OpenWorldSettings();
+			Item->OpenWorldSettings();
 		}
 		return FReply::Handled();
 	}
diff --git a/Engine/Source/Editor/SceneOutliner/Public/SOutlinerTreeView.h b/Engine/Source/Editor/SceneOutliner/Public/SOutlinerTreeView.h
--- a/Engine/Source/Editor/SceneOutliner/Public/SOutlinerTreeView.h
+++ b/Engine/Source/Editor/SceneOutliner/Public/SOutlinerTreeView.h
@@ -10,6 +10,12 @@
 #include "Widgets/Views/STreeView.h"
 #include "Animation/SlateSprings.h"
 
+template<>
+struct TWidgetTypeTraits<class SSceneOutlinerTreeView>
+{
+	static constexpr bool SupportsInvalidation() { return true; }
+};
+
 class SSceneOutlinerTreeView : public STreeView< FSceneOutlinerTreeItemPtr >
 {
 public:
@@ -35,6 +41,12 @@ private:
 	TWeakPtr<SSceneOutliner> SceneOutlinerWeak;
 };
 
+template<>
+struct TWidgetTypeTraits<class SSceneOutlinerTreeRow>
+{
+	static constexpr bool SupportsInvalidation() { return true; }
+};
+
 /** Widget that represents a row in the outliner's tree control.  Generates widgets for each column on demand. */
 class SSceneOutlinerTreeRow
 	: public SMultiColumnTableRow< FSceneOutlinerTreeItemPtr >
@@ -59,6 +71,7 @@ public:
 	void FlashHighlight();
 
 protected:
+	virtual bool ComputeVolatility() const override;
 
 	virtual FReply OnMouseButtonDown(const FGeometry& MyGeometry, const FPointerEvent& MouseEvent) override;
 	virtual FReply OnMouseButtonUp(const FGeometry& MyGeometry, const FPointerEvent& MouseEvent) override;
diff --git a/Engine/Source/Editor/SceneOutliner/Public/SceneOutlinerGutter.h b/Engine/Source/Editor/SceneOutliner/Public/SceneOutlinerGutter.h
--- a/Engine/Source/Editor/SceneOutliner/Public/SceneOutlinerGutter.h
+++ b/Engine/Source/Editor/SceneOutliner/Public/SceneOutlinerGutter.h
@@ -62,6 +62,12 @@ protected:
 	FSceneOutlinerVisibilityCache VisibilityCache;
 };
 
+template<>
+struct TWidgetTypeTraits<class SVisibilityWidget>
+{
+	static constexpr bool SupportsInvalidation() { return true; }
+};
+
 /** Widget responsible for managing the visibility for a single item */
 class SVisibilityWidget : public SImage
 {
