EditorUsability : SceneOutliner
- Fixed hotkeys do not work in the Outliner when undocked, reason is that SLevelEditor handles commands processing but when scene outliner is undocked, events do not bubble up to SLevelEditor since they stop when they reach top window, and SLevelEditor is part of the background main editor window
- Remove scene outliner search bar tooltip since it gets in the way and is obvious
- Reorder includes alphabetically

diff --git a/Engine/Source/Editor/LevelEditor/Private/SLevelEditor.cpp b/Engine/Source/Editor/LevelEditor/Private/SLevelEditor.cpp
--- a/Engine/Source/Editor/LevelEditor/Private/SLevelEditor.cpp
+++ b/Engine/Source/Editor/LevelEditor/Private/SLevelEditor.cpp
@@ -827,6 +827,13 @@ TSharedRef<ISceneOutliner> SLevelEditor::CreateSceneOutliner(FName TabIdentifier
 {
 	FSceneOutlinerInitializationOptions InitOptions;
 	InitOptions.bShowTransient = true;
+
+	/**
+	 * Pass level editor commands to allow scene outliner tab to handle them properly
+	 * even if outliner scene tab is undocked and events are not processed by SLevelEditor since they do not share the same window anymore
+	 */
+	InitOptions.CommandList = LevelEditorCommands;
+
 	{
 		UToolMenus* ToolMenus = UToolMenus::Get();
 		static const FName MenuName = "LevelEditor.LevelEditorSceneOutliner.ContextMenu";
diff --git a/Engine/Source/Editor/SceneOutliner/Private/ActorBrowsingMode.cpp b/Engine/Source/Editor/SceneOutliner/Private/ActorBrowsingMode.cpp
--- a/Engine/Source/Editor/SceneOutliner/Private/ActorBrowsingMode.cpp
+++ b/Engine/Source/Editor/SceneOutliner/Private/ActorBrowsingMode.cpp
@@ -1408,10 +1408,15 @@ void FActorBrowsingMode::BindCommands(const TSharedRef<FUICommandList>& OutComma
 		FExecuteAction::CreateRaw(this, &FActorBrowsingMode::OnExecuteRename ),
 		FCanExecuteAction::CreateRaw(this, &FActorBrowsingMode::CanExecuteRename));
 
-	OutCommandList->MapAction(
-		FGenericCommands::Get().Delete,
-		FExecuteAction::CreateRaw(this, &FActorBrowsingMode::OnExecuteDelete),
-		FCanExecuteAction::CreateRaw(this, &FActorBrowsingMode::CanDelete));
+	// Override DELETE command with custom delete
+	if (SceneOutliner && SceneOutliner->GetSharedData().CustomDelete.IsBound())
+	{
+		OutCommandList->MapAction(
+			FGenericCommands::Get().Delete,
+			FExecuteAction::CreateRaw(this, &FActorBrowsingMode::OnExecuteDelete),
+			FCanExecuteAction::CreateRaw(this, &FActorBrowsingMode::CanDelete)
+		);
+	}
 	
 	OutCommandList->MapAction(
 		FActorBrowsingModeCommands::Get().Refresh,
diff --git a/Engine/Source/Editor/SceneOutliner/Private/SSceneOutliner.cpp b/Engine/Source/Editor/SceneOutliner/Private/SSceneOutliner.cpp
--- a/Engine/Source/Editor/SceneOutliner/Private/SSceneOutliner.cpp
+++ b/Engine/Source/Editor/SceneOutliner/Private/SSceneOutliner.cpp
@@ -2,36 +2,36 @@
 
 #include "SSceneOutliner.h"
 
+#include "DetailLayoutBuilder.h"
 #include "Editor.h"
-#include "Styling/AppStyle.h"
+#include "EditorFolderUtils.h"
 #include "Engine/GameViewportClient.h"
+#include "FolderTreeItem.h"
 #include "Framework/Application/SlateApplication.h"
 #include "Framework/Commands/UIAction.h"
 #include "Framework/MultiBox/MultiBoxBuilder.h"
+#include "Framework/Notifications/NotificationManager.h"
 #include "HAL/PlatformApplicationMisc.h"
 #include "ISceneOutlinerColumn.h"
+#include "ISceneOutlinerMode.h"
 #include "Layout/WidgetPath.h"
 #include "Modules/ModuleManager.h"
+#include "SceneOutlinerConfig.h"
+#include "SceneOutlinerFilterBar.h"
 #include "SceneOutlinerFilters.h"
 #include "SceneOutlinerModule.h"
 #include "ScopedTransaction.h"
+#include "Styling/AppStyle.h"
 #include "Textures/SlateIcon.h"
 #include "ToolMenus.h"
 #include "UObject/PackageReload.h"
+#include "Widgets/SInvalidationPanel.h"
+#include "Widgets/SOverlay.h"
 #include "Widgets/Images/SImage.h"
 #include "Widgets/Input/SButton.h"
 #include "Widgets/Input/SComboButton.h"
-#include "Widgets/SOverlay.h"
-#include "ISceneOutlinerMode.h"
-#include "FolderTreeItem.h"
-#include "EditorFolderUtils.h"
-#include "SceneOutlinerConfig.h"
-#include "SceneOutlinerFilterBar.h"
-#include "Framework/Notifications/NotificationManager.h"
-#include "Widgets/Notifications/SNotificationList.h"
 #include "Widgets/Input/SMultiLineEditableTextBox.h"
-#include "DetailLayoutBuilder.h"
-#include "Widgets/SInvalidationPanel.h"
+#include "Widgets/Notifications/SNotificationList.h"
 
 DEFINE_LOG_CATEGORY(LogSceneOutliner);
 
@@ -65,7 +65,6 @@ void SSceneOutliner::Construct(const FArguments& InArgs, const FSceneOutlinerIni
 		FilterTextBoxWidget = SNew(SFilterSearchBox)
 		.Visibility( InInitOptions.bShowSearchBox ? EVisibility::Visible : EVisibility::Collapsed )
 		.HintText( LOCTEXT( "FilterSearch", "Search..." ) )
-		.ToolTipText( LOCTEXT("FilterSearchHint", "Type here to search (pressing enter selects the results)") )
 		.OnTextChanged( this, &SSceneOutliner::OnFilterTextChanged )
 		.OnTextCommitted( this, &SSceneOutliner::OnFilterTextCommitted );
 	}
@@ -107,7 +106,7 @@ void SSceneOutliner::Construct(const FArguments& InArgs, const FSceneOutlinerIni
 	// @todo outliner: Should save spacing for list view in layout
 
 	// Setup Commands
-	BindCommands();
+	BindCommands(InInitOptions.CommandList);
 
 	FSceneOutlinerModule& SceneOutlinerModule = FModuleManager::LoadModuleChecked<FSceneOutlinerModule>("SceneOutliner");
 	SceneOutlinerModule.OnColumnPermissionListChanged().AddSP(this, &SSceneOutliner::OnColumnPermissionListChanged);
@@ -2198,12 +2197,21 @@ void SSceneOutliner::OnItemLabelChanged(FSceneOutlinerTreeItemPtr ChangedItem, b
 	}
 }
 
-void SSceneOutliner::BindCommands()
+void SSceneOutliner::BindCommands(const TSharedPtr<FUICommandList>& InCommandListToAppend)
 {
 	CommandList = MakeShared<FUICommandList>();
 
 	// Bind Mode Commands
-	Mode->BindCommands(CommandList.ToSharedRef());
+	if (Mode)
+	{
+		Mode->BindCommands(CommandList.ToSharedRef());
+	}
+
+	// Append additional commands
+	if (InCommandListToAppend.IsValid())
+	{
+		CommandList->Append(InCommandListToAppend.ToSharedRef());
+	}
 }
 
 void SSceneOutliner::OnAssetReloaded(const EPackageReloadPhase InPackageReloadPhase, FPackageReloadedEvent* InPackageReloadedEvent)
diff --git a/Engine/Source/Editor/SceneOutliner/Public/SSceneOutliner.h b/Engine/Source/Editor/SceneOutliner/Public/SSceneOutliner.h
--- a/Engine/Source/Editor/SceneOutliner/Public/SSceneOutliner.h
+++ b/Engine/Source/Editor/SceneOutliner/Public/SSceneOutliner.h
@@ -1,4 +1,5 @@
 // Copyright Epic Games, Inc. All Rights Reserved.
+
 #pragma once
 
 #include "EditorUndoClient.h"
@@ -243,13 +244,13 @@ public:
 	{}
 	SLATE_END_ARGS()
 
-		/**
-			* Construct this widget.  Called by the SNew() Slate macro.
-			*
-			* @param	InArgs		Declaration used by the SNew() macro to construct this widget
-			* @param	InitOptions	Programmer-driven initialization options for this widget
-			*/
-		UE_API void Construct(const FArguments& InArgs, const FSceneOutlinerInitializationOptions& InitOptions);
+	/**
+	* Construct this widget.  Called by the SNew() Slate macro.
+	*
+	* @param	InArgs		Declaration used by the SNew() macro to construct this widget
+	* @param	InitOptions	Programmer-driven initialization options for this widget
+	*/
+	UE_API void Construct(const FArguments& InArgs, const FSceneOutlinerInitializationOptions& InitOptions);
 
 	/** Default constructor - initializes data that is shared between all tree items */
 	SSceneOutliner() 
@@ -687,10 +688,10 @@ private:
 
 	/** Handler for when a property changes on any item. Called by the mode */
 	UE_API void OnItemLabelChanged(FSceneOutlinerTreeItemPtr ChangedItem, bool bFlashHighlight = true);
+	
 private:
-
-	/** BindCommands for this Outliner */
-	UE_API void BindCommands();
+	/** BindCommands for this Outliner with additional commands to append */
+	UE_API void BindCommands(const TSharedPtr<FUICommandList>& InCommandListToAppend);
 
 	/** Map of columns that are shown on this outliner. */
 	TMap<FName, TSharedPtr<ISceneOutlinerColumn>> Columns;
diff --git a/Engine/Source/Editor/SceneOutliner/Public/SceneOutlinerPublicTypes.h b/Engine/Source/Editor/SceneOutliner/Public/SceneOutlinerPublicTypes.h
--- a/Engine/Source/Editor/SceneOutliner/Public/SceneOutlinerPublicTypes.h
+++ b/Engine/Source/Editor/SceneOutliner/Public/SceneOutlinerPublicTypes.h
@@ -2,20 +2,20 @@
 
 #pragma once
 
-#include "Layout/Margin.h"
-#include "Framework/SlateDelegates.h"
 #include "Engine/World.h"
-
+#include "Framework/SlateDelegates.h"
+#include "Layout/Margin.h"
 #include "SceneOutlinerFilters.h"
 
 #define UE_API SCENEOUTLINER_API
 
 #define LOCTEXT_NAMESPACE "SceneOutlinerPublicTypes"
 
-class FExtender;
-struct FToolMenuContext;
 class FCustomClassFilterData;
+class FExtender;
 class FFilterCategory;
+struct FToolMenuContext;
+class FUICommandList;
 template<typename FilterType> class FFilterBase;
 
 DECLARE_DELEGATE_TwoParams(FSceneOutlinerModifyContextMenu, FName& /* MenuName */, FToolMenuContext& /* MenuContext */);
@@ -206,7 +206,8 @@ struct FSceneOutlinerInitializationOptions : FSharedSceneOutlinerData
 	/** The primary column for this outliner, used to determine initial sorting and where the hierarchy arrows appear */
 	FName PrimaryColumnName;
 
-public:
+	/** Additional commands to append to the scene outliner command list */
+	TSharedPtr<FUICommandList> CommandList;
 
 	/** Constructor */
 	FSceneOutlinerInitializationOptions()
