SceneOutliner / SourceControl: Fixed SceneOutliner triggering an FUpdateStatus after duplicating an actor.

At that point, the actor merely exists in memory - it has not been saved to disk yet. So there is no SCC state to query.

#rnx

#Synced-CL 47898077

diff --git a/Engine/Source/Editor/SceneOutliner/Private/SceneOutlinerTreeItemSCC.cpp b/Engine/Source/Editor/SceneOutliner/Private/SceneOutlinerTreeItemSCC.cpp
--- a/Engine/Source/Editor/SceneOutliner/Private/SceneOutlinerTreeItemSCC.cpp
+++ b/Engine/Source/Editor/SceneOutliner/Private/SceneOutlinerTreeItemSCC.cpp
@@ -127,6 +127,14 @@ void FSceneOutlinerTreeItemSCC::ConnectSourceControl()
 	{
 		BroadcastNewState(SourceControlState);
 	}
+	else if (UPackage* Package = GetPackage())
+	{
+		// There's nothing to query if the package didn't come from disk or hasn't yet been saved to disk
+		if (Package->GetHasBeenEndLoaded())
+		{
+			SCCModule.QueueStatusUpdate(ExternalPackageFileName);
+		}
+	}
 	else
 	{
 		SCCModule.QueueStatusUpdate(ExternalPackageFileName);
