EditorUsability : SceneOutliner
- Fixed base selection is lost when ctrl is held and a new search is made to select additional elements in the scene outliner

#Synced-CL 48903505

diff --git a/Engine/Source/Editor/SceneOutliner/Private/ActorBrowsingMode.cpp b/Engine/Source/Editor/SceneOutliner/Private/ActorBrowsingMode.cpp
--- a/Engine/Source/Editor/SceneOutliner/Private/ActorBrowsingMode.cpp
+++ b/Engine/Source/Editor/SceneOutliner/Private/ActorBrowsingMode.cpp
@@ -1064,6 +1064,20 @@ void FActorBrowsingMode::OnItemSelectionChanged(FSceneOutlinerTreeItemPtr TreeIt
 	if (UTypedElementSelectionSet* SelectionSet = ActorSelection->GetElementSelectionSet())
 	{
 		TSet<AActor*> EditorSelectedActors(SelectionSet->GetSelectedObjects<AActor>());
+		
+		const FModifierKeysState ModifierKeys = FSlateApplication::Get().GetModifierKeys();
+		if (ModifierKeys.IsControlDown())
+		{
+			for (AActor* SelectedActor : EditorSelectedActors)
+			{
+				// If we can't find item in tree == item is filtered out, since CTRL is held, add it to current selection
+				FSceneOutlinerTreeItemPtr Outliner = SceneOutliner->GetTreeItem(SelectedActor);
+				if (!Outliner && SelectedActor != nullptr)
+				{
+					OutlinerSelectedActors.Add(SelectedActor);
+				}
+			}
+		}
 
 		bool bChanged = false;
 		bool bAnyInPIE = false;
diff --git a/Engine/Source/Editor/SceneOutliner/Private/SSceneOutliner.cpp b/Engine/Source/Editor/SceneOutliner/Private/SSceneOutliner.cpp
--- a/Engine/Source/Editor/SceneOutliner/Private/SSceneOutliner.cpp
+++ b/Engine/Source/Editor/SceneOutliner/Private/SSceneOutliner.cpp
@@ -771,7 +771,7 @@ void SSceneOutliner::Populate()
 	{
 		if (FSceneOutlinerTreeItemPtr* Item = TreeItemMap.Find(Folder))
 		{
-			OutlinerTreeView->SetItemSelection(*Item, true);
+			AddToSelection(*Item);
 		}
 	}
 	PendingFoldersSelect.Empty();
@@ -1037,8 +1037,7 @@ bool SSceneOutliner::AddItemToTree(FSceneOutlinerTreeItemRef Item)
 		{
 			if (*ActionMask & SceneOutliner::ENewItemAction::Select)
 			{
-				OutlinerTreeView->ClearSelection();
-				OutlinerTreeView->SetItemSelection(Item, true);
+				SetItemSelection(Item, true);
 			}
 
 			if (*ActionMask & SceneOutliner::ENewItemAction::Rename && CanExecuteRenameRequest(*Item))
@@ -1303,6 +1302,11 @@ void SSceneOutliner::RemoveFromSelection(const TArray<FSceneOutlinerTreeItemPtr>
 	OutlinerTreeView->SetItemSelection(InItems, false, SelectInfo);
 }
 
+void SSceneOutliner::RemoveFromSelection(const FSceneOutlinerTreeItemPtr& InItem, ESelectInfo::Type SelectInfo)
+{
+	OutlinerTreeView->SetItemSelection(InItem, false, SelectInfo);
+}
+
 void SSceneOutliner::AddFolderToSelection(const FName& FolderName)
 {
 	// Not used (but public) : For backward compatibility, we use Mode->GetRootObject()
@@ -1311,7 +1315,7 @@ void SSceneOutliner::AddFolderToSelection(const FName& FolderName)
 	{
 		if (FSceneOutlinerTreeItemPtr* ItemPtr = TreeItemMap.Find(FFolder(RootObject, FolderName)))
 		{
-			OutlinerTreeView->SetItemSelection(*ItemPtr, true);
+			AddToSelection(*ItemPtr);
 		}
 	}
 }
@@ -1324,7 +1328,7 @@ void SSceneOutliner::RemoveFolderFromSelection(const FName& FolderName)
 	{
 		if (FSceneOutlinerTreeItemPtr* ItemPtr = TreeItemMap.Find(FFolder(RootObject, FolderName)))
 		{
-			OutlinerTreeView->SetItemSelection(*ItemPtr, false);
+			RemoveFromSelection(*ItemPtr);
 		}
 	}
 }
@@ -1777,7 +1781,7 @@ void SSceneOutliner::DuplicateFoldersHierarchy()
 	{
 		if (Item->IsA<FFolderTreeItem>())
 		{
-			OutlinerTreeView->SetItemSelection(Item, true);
+			AddToSelection(Item);
 		}
 		for (const TWeakPtr<ISceneOutlinerTreeItem>& Child : Item->GetChildren())
 		{
diff --git a/Engine/Source/Editor/SceneOutliner/Public/SSceneOutliner.h b/Engine/Source/Editor/SceneOutliner/Public/SSceneOutliner.h
--- a/Engine/Source/Editor/SceneOutliner/Public/SSceneOutliner.h
+++ b/Engine/Source/Editor/SceneOutliner/Public/SSceneOutliner.h
@@ -448,7 +448,10 @@ public:
 	UE_API void AddToSelection(const TArray<FSceneOutlinerTreeItemPtr>& InItems, ESelectInfo::Type SelectInfo = ESelectInfo::Direct);
 
 	/** Remove a set of items from the current selection */
-	UE_API void RemoveFromSelection(const TArray<FSceneOutlinerTreeItemPtr>& InItems, ESelectInfo::Type SelectInfo = ESelectInfo::Direct);\
+	UE_API void RemoveFromSelection(const TArray<FSceneOutlinerTreeItemPtr>& InItems, ESelectInfo::Type SelectInfo = ESelectInfo::Direct);
+	
+	/** Remove a single item from the current selection */
+	UE_API void RemoveFromSelection(const FSceneOutlinerTreeItemPtr& InItem, ESelectInfo::Type SelectInfo = ESelectInfo::Direct);
 
 	/** Adds a single item to the current selection */
 	UE_API void AddToSelection(const FSceneOutlinerTreeItemPtr& InItem, ESelectInfo::Type SelectInfo = ESelectInfo::Direct);
