Editor: Expose SetItemExpansion in ISceneOutliner.

#Synced-CL 48918043

diff --git a/Engine/Source/Editor/SceneOutliner/Public/ISceneOutliner.h b/Engine/Source/Editor/SceneOutliner/Public/ISceneOutliner.h
--- a/Engine/Source/Editor/SceneOutliner/Public/ISceneOutliner.h
+++ b/Engine/Source/Editor/SceneOutliner/Public/ISceneOutliner.h
@@ -177,6 +177,9 @@ public:
 	/** Scrolls the outliner to the specified item if it's found. No action is done otherwise. */
 	virtual void FrameItem(const FSceneOutlinerTreeItemID& Item) = 0;
 
+	/** Set whether a particular item in the outliner is expanded or not. */
+	virtual void SetItemExpansion(const FSceneOutlinerTreeItemPtr& Item, bool bIsExpanded) {}
+
 	/** Get the active SceneOutlinerMode */
 	const ISceneOutlinerMode* GetMode() const { return Mode; }
 
diff --git a/Engine/Source/Editor/SceneOutliner/Public/SSceneOutliner.h b/Engine/Source/Editor/SceneOutliner/Public/SSceneOutliner.h
--- a/Engine/Source/Editor/SceneOutliner/Public/SSceneOutliner.h
+++ b/Engine/Source/Editor/SceneOutliner/Public/SSceneOutliner.h
@@ -732,7 +732,7 @@ public:
 	/** Scroll the specified item into view */
 	UE_API void ScrollItemIntoView(const FSceneOutlinerTreeItemPtr& Item);
 
-	UE_API void SetItemExpansion(const FSceneOutlinerTreeItemPtr& Item, bool bIsExpanded);
+	UE_API virtual void SetItemExpansion(const FSceneOutlinerTreeItemPtr& Item, bool bIsExpanded) override;
 		
 	UE_API bool IsItemExpanded(const FSceneOutlinerTreeItemPtr& Item) const;
 
