Dynamic Invalidation - SSceneOutlinerPinnedTreeRow

#Synced-CL 48711348

diff --git a/Engine/Source/Editor/SceneOutliner/Public/SOutlinerTreeView.h b/Engine/Source/Editor/SceneOutliner/Public/SOutlinerTreeView.h
--- a/Engine/Source/Editor/SceneOutliner/Public/SOutlinerTreeView.h
+++ b/Engine/Source/Editor/SceneOutliner/Public/SOutlinerTreeView.h
@@ -130,6 +130,12 @@ private:
 	double LastHighlightInteractionTime;
 };
 
+template<>
+struct TWidgetTypeTraits<class SSceneOutlinerPinnedTreeRow>
+{
+	static constexpr bool SupportsInvalidation() { return true; }
+};
+
 /** Widget that represents a pinned row in the outliner's tree control.*/
 class SSceneOutlinerPinnedTreeRow
 	: public SMultiColumnTableRow< FSceneOutlinerTreeItemPtr >
