Moved world outliner's Pinned column in WorldPartitionEditorModule.

#Synced-CL 49593614

diff --git a/Engine/Source/Editor/SceneOutliner/Private/SceneOutlinerModule.cpp b/Engine/Source/Editor/SceneOutliner/Private/SceneOutlinerModule.cpp
--- a/Engine/Source/Editor/SceneOutliner/Private/SceneOutlinerModule.cpp
+++ b/Engine/Source/Editor/SceneOutliner/Private/SceneOutlinerModule.cpp
@@ -230,40 +230,18 @@ void FSceneOutlinerModule::CreateActorBrowserColumns(FSceneOutlinerInitializatio
 	InInitOptions.ColumnMap.Add(FSceneOutlinerBuiltInColumnTypes::Gutter(), FSceneOutlinerColumnInfo(ESceneOutlinerColumnVisibility::Visible, 0, FCreateSceneOutlinerColumn(), true, TOptional<float>(), FSceneOutlinerBuiltInColumnTypes::Gutter_Localized()));
 	InInitOptions.ColumnMap.Add(FSceneOutlinerBuiltInColumnTypes::ActorInfo(), FSceneOutlinerColumnInfo(ESceneOutlinerColumnVisibility::Visible, 20, FCreateSceneOutlinerColumn(), true, TOptional<float>(), FSceneOutlinerBuiltInColumnTypes::ActorInfo_Localized()));
 
-	CreateWorldPartitionColumns(InInitOptions, InWorld);
-	CreateActorInfoColumns(InInitOptions, InWorld);
-
-	CreateActorBrowserColumnsDelegate.Broadcast(InInitOptions, InWorld);
-}
-
-void FSceneOutlinerModule::CreateWorldPartitionColumns(FSceneOutlinerInitializationOptions& InInitOptions, UWorld* WorldPtr) const
-{
-	ESceneOutlinerColumnVisibility SourceControlColumnVisibility = ESceneOutlinerColumnVisibility::Invisible;
-
 	ISourceControlModule& SourceControlModule = ISourceControlModule::Get();
-	if (SourceControlModule.IsEnabled())
-	{
-		if (WorldPtr && WorldPtr->PersistentLevel->IsUsingExternalActors())
-		{
-			// The source control column should be visible by default in source-controlled levels using external actors
-			SourceControlColumnVisibility = ESceneOutlinerColumnVisibility::Visible;
-		}
-	}
+	// The source control column should be visible by default in source-controlled levels using external actors
+	ESceneOutlinerColumnVisibility SourceControlColumnVisibility = (SourceControlModule.IsEnabled() && InWorld && InWorld->PersistentLevel->IsUsingExternalActors()) ? ESceneOutlinerColumnVisibility::Visible : ESceneOutlinerColumnVisibility::Invisible;
 	InInitOptions.ColumnMap.Add(FSceneOutlinerBuiltInColumnTypes::SourceControl(), FSceneOutlinerColumnInfo(SourceControlColumnVisibility, 30, FCreateSceneOutlinerColumn(), true, TOptional<float>(), FSceneOutlinerBuiltInColumnTypes::SourceControl_Localized()));
 
-	ESceneOutlinerColumnVisibility UnsavedColumnVisibility = ESceneOutlinerColumnVisibility::Invisible;
-		
-	if (WorldPtr && WorldPtr->IsPartitionedWorld())
-	{
-		// We don't want the pinned column in non wp levels
-		InInitOptions.ColumnMap.Add(FSceneOutlinerBuiltInColumnTypes::Pinned(), FSceneOutlinerColumnInfo(ESceneOutlinerColumnVisibility::Visible, 5, FCreateSceneOutlinerColumn(), true, TOptional<float>(), FSceneOutlinerBuiltInColumnTypes::Pinned_Localized()));
-
-		// We want the unsaved column to be visible by default in partitioned levels
-		UnsavedColumnVisibility = ESceneOutlinerColumnVisibility::Visible;
-	}
-	
+	// We want the unsaved column to be visible by default in partitioned levels
+	ESceneOutlinerColumnVisibility UnsavedColumnVisibility = (InWorld && InWorld->IsPartitionedWorld()) ? ESceneOutlinerColumnVisibility::Visible : ESceneOutlinerColumnVisibility::Invisible;
 	InInitOptions.ColumnMap.Add(FSceneOutlinerBuiltInColumnTypes::Unsaved(), FSceneOutlinerColumnInfo(UnsavedColumnVisibility, 1, FCreateSceneOutlinerColumn(), true, TOptional<float>(), FSceneOutlinerBuiltInColumnTypes::Unsaved_Localized()));
 
+	CreateActorInfoColumns(InInitOptions, InWorld);
+
+	CreateActorBrowserColumnsDelegate.Broadcast(InInitOptions, InWorld);
 }
 
 void FSceneOutlinerModule::CreateActorInfoColumns(FSceneOutlinerInitializationOptions& InInitOptions, UWorld *WorldPtr) const
diff --git a/Engine/Source/Editor/SceneOutliner/Public/SceneOutlinerModule.h b/Engine/Source/Editor/SceneOutliner/Public/SceneOutlinerModule.h
--- a/Engine/Source/Editor/SceneOutliner/Public/SceneOutlinerModule.h
+++ b/Engine/Source/Editor/SceneOutliner/Public/SceneOutlinerModule.h
@@ -140,7 +140,6 @@ public:
 	}
 
 	void CreateActorInfoColumns(FSceneOutlinerInitializationOptions& InInitOptions, UWorld* WorldPtr = nullptr) const;
-	void CreateWorldPartitionColumns(FSceneOutlinerInitializationOptions& InInitOptions, UWorld* WorldPtr = nullptr) const;
 	
 	/** Map of column type name -> default column info */
 	TMap< FName, FSceneOutlinerColumnInfo> DefaultColumnMap;
diff --git a/Engine/Source/Editor/WorldPartitionEditor/Private/WorldPartitionEditorModule.cpp b/Engine/Source/Editor/WorldPartitionEditor/Private/WorldPartitionEditorModule.cpp
--- a/Engine/Source/Editor/WorldPartitionEditor/Private/WorldPartitionEditorModule.cpp
+++ b/Engine/Source/Editor/WorldPartitionEditor/Private/WorldPartitionEditorModule.cpp
@@ -63,6 +63,8 @@
 #include "ISourceControlModule.h"
 #include "ISourceControlProvider.h"
 #include "SourceControlOperations.h"
+#include "SceneOutlinerPublicTypes.h"
+#include "SceneOutlinerModule.h"
 
 #include "Styling/AppStyle.h"
 #include "WorldPartition/ContentBundle/SContentBundleBrowser.h"
@@ -306,7 +308,10 @@ void FWorldPartitionEditorModule::StartupModule()
 			UEditorStateSubsystem::UnregisterEditorStateType<UWorldPartitionEditorState>();
 			GEditor->OnEditorClose().Remove(EditorCloseHandle);
 		});
-	});		
+	});	
+
+	FSceneOutlinerModule& SceneOutlinerModule = FModuleManager::LoadModuleChecked<FSceneOutlinerModule>("SceneOutliner");
+	SceneOutlinerModule.OnCreateActorBrowserColumns().AddRaw(this, &FWorldPartitionEditorModule::OnCreateActorBrowserColumns);
 
 	IAssetReferenceFilter::OnIsCrossPluginReferenceAllowed().BindRaw(this, &FWorldPartitionEditorModule::OnIsCrossPluginReferenceAllowed);
 
@@ -318,6 +323,11 @@ void FWorldPartitionEditorModule::ShutdownModule()
 	FWorldPartitionClassDescRegistry().Get().Uninitialize();
 	FWorldPartitionClassDescRegistry().Get().TearDown();
 
+	if (FSceneOutlinerModule* SceneOutlinerModule = FModuleManager::GetModulePtr<FSceneOutlinerModule>("SceneOutliner"))
+	{
+		SceneOutlinerModule->OnCreateActorBrowserColumns().RemoveAll(this);
+	}
+
 	IAssetReferenceFilter::OnIsCrossPluginReferenceAllowed().Unbind();
 
 	if (!IsRunningGame())
@@ -349,6 +359,15 @@ void FWorldPartitionEditorModule::ShutdownModule()
 	FEditorDelegates::OnEditorInitialized.Remove(EditorInitializedHandle);
 }
 
+void FWorldPartitionEditorModule::OnCreateActorBrowserColumns(FSceneOutlinerInitializationOptions& InitOptions, UWorld* World)
+{
+	// We don't want the pinned column in non wp levels
+	if (World && World->IsPartitionedWorld())
+	{
+		InitOptions.ColumnMap.Add(FSceneOutlinerBuiltInColumnTypes::Pinned(), FSceneOutlinerColumnInfo(ESceneOutlinerColumnVisibility::Visible, 5, FCreateSceneOutlinerColumn(), true, TOptional<float>(), FSceneOutlinerBuiltInColumnTypes::Pinned_Localized()));
+	}
+}
+
 bool FWorldPartitionEditorModule::OnIsCrossPluginReferenceAllowed(const FAssetData& ReferencingAssetData, const FAssetData& ReferencedAssetData)
 {
 	// Allow External Data Layer Actor (ReferencingAssetData) from a plugin X to reference its world (ReferencedAssetData) from a plugin Y
diff --git a/Engine/Source/Editor/WorldPartitionEditor/Public/WorldPartitionEditorModule.h b/Engine/Source/Editor/WorldPartitionEditor/Public/WorldPartitionEditorModule.h
--- a/Engine/Source/Editor/WorldPartitionEditor/Public/WorldPartitionEditorModule.h
+++ b/Engine/Source/Editor/WorldPartitionEditor/Public/WorldPartitionEditorModule.h
@@ -10,6 +10,7 @@ class SDockTab;
 class FSpawnTabArgs;
 class SContentBundleBrowser;
 struct FAssetData;
+struct FSceneOutlinerInitializationOptions;
 
 
 /**
@@ -171,6 +172,7 @@ private:
 	void RunCommandletAsExternalProcess(const FString& InCommandletArgs, const FText& InOperationDescription, int32& OutResult, bool& bOutCancelled);
 	void OnConvertMap();
 	bool OnIsCrossPluginReferenceAllowed(const FAssetData& ReferencingAssetData, const FAssetData& ReferencedAssetData);
+	void OnCreateActorBrowserColumns(FSceneOutlinerInitializationOptions& InitOptions, UWorld* World);
 
 	FDelegateHandle EditorInitializedHandle;
 	FDelegateHandle EditorCloseHandle;
