[TEDS Outliner] Change the SSceneOutliner CustomSceneOutlinerFactory registration functionsto be UE_INTERNAL and disable warnings so only the Hybrid Outliner can use it

#Synced-CL 51758326

#Synced-CL 51758326

diff --git a/Engine/Source/Editor/LevelEditor/Private/SLevelEditor.cpp b/Engine/Source/Editor/LevelEditor/Private/SLevelEditor.cpp
--- a/Engine/Source/Editor/LevelEditor/Private/SLevelEditor.cpp
+++ b/Engine/Source/Editor/LevelEditor/Private/SLevelEditor.cpp
@@ -890,10 +890,12 @@ TSharedRef<ISceneOutliner> SLevelEditor::CreateSceneOutliner(FName TabIdentifier
 	TSharedPtr<ISceneOutliner> NewSceneOutlinerPtr;
 
 	// Check if the Typed Element Registry has registered a custom outliner factory for us to use
+	PRAGMA_DISABLE_INTERNAL_WARNINGS
 	if(SceneOutlinerModule.IsCustomSceneOutlinerFactoryRegistered(UTypedElementRegistry::GetInstance()->GetFName()))
 	{
 		NewSceneOutlinerPtr = SceneOutlinerModule.CreateCustomRegisteredOutliner(UTypedElementRegistry::GetInstance()->GetFName(), InitOptions);
 	}
+	PRAGMA_ENABLE_INTERNAL_WARNINGS
 	// Fallback to the regular Actor Browser otherwise
 	else
 	{
diff --git a/Engine/Source/Editor/SceneOutliner/Public/SceneOutlinerModule.h b/Engine/Source/Editor/SceneOutliner/Public/SceneOutlinerModule.h
--- a/Engine/Source/Editor/SceneOutliner/Public/SceneOutlinerModule.h
+++ b/Engine/Source/Editor/SceneOutliner/Public/SceneOutlinerModule.h
@@ -61,16 +61,16 @@ public:
 
 
 	/** Register a factory to create a custom Scene Outliner */
-	virtual void RegisterCustomSceneOutlinerFactory(FName ID, FSceneOutlinerFactory InOutlinerFactory);
+	virtual UE_INTERNAL SCENEOUTLINER_API void RegisterCustomSceneOutlinerFactory(FName ID, FSceneOutlinerFactory InOutlinerFactory) final;
 
 	/** Unregister a factory to create a custom Scene Outliner */
-	virtual void UnregisterCustomSceneOutlinerFactory(FName ID);
+	virtual UE_INTERNAL SCENEOUTLINER_API void UnregisterCustomSceneOutlinerFactory(FName ID) final;
 
 	/** Try to create a custom scene outliner using a registered factory (nullptr if ID was not registered) */
-	virtual TSharedPtr<ISceneOutliner> CreateCustomRegisteredOutliner(FName ID, FSceneOutlinerInitializationOptions InInitOptions);
+	virtual UE_INTERNAL SCENEOUTLINER_API TSharedPtr<ISceneOutliner> CreateCustomRegisteredOutliner(FName ID, FSceneOutlinerInitializationOptions InInitOptions) final;
 
 	/** Check if a custom scene outliner factory is registered with the given ID */
-	virtual bool IsCustomSceneOutlinerFactoryRegistered(FName ID);
+	virtual UE_INTERNAL SCENEOUTLINER_API bool IsCustomSceneOutlinerFactoryRegistered(FName ID) final;
 
 	/** Add the columns present in the level editor's Outliner (Actor Browser) to the given init options
 	 *  @param InWorld The world the Outliner initialized by InInitOptions will look it, defaults to the level editor's world if nullptr
