ActorBrowsingMode: Remove explicit GEditor->NoteSelectionChange() call since SelectionSet->NotifyPendingChanges() already fires the notification, leading to duplicate notifications.

#Synced-CL 51760981

diff --git a/Engine/Source/Editor/SceneOutliner/Private/ActorBrowsingMode.cpp b/Engine/Source/Editor/SceneOutliner/Private/ActorBrowsingMode.cpp
--- a/Engine/Source/Editor/SceneOutliner/Private/ActorBrowsingMode.cpp
+++ b/Engine/Source/Editor/SceneOutliner/Private/ActorBrowsingMode.cpp
@@ -1146,15 +1146,11 @@ void FActorBrowsingMode::OnItemSelectionChanged(FSceneOutlinerTreeItemPtr TreeIt
 					.SetAllowLegacyNotifications(false)
 					.SetAllowSubRootSelection(true);
 
-				// Avoid senting out notification via typed element and call NoteSelectionChange to preserve previous behavior
 				FTypedElementList::FScopedClearNewPendingChange ClearNewPendingChange = SelectionSet->GetScopedClearNewPendingChange();
 				SelectionSet->SetSelection(ElementsToSelect, SelectionOptions);
 				SelectionSet->NotifyPendingChanges();
 			}
 
-			// Fire selection changed event
-			GEditor->NoteSelectionChange();
-
 			// Set this outliner as the most recently interacted with
 			SetAsMostRecentOutliner();
 		}
