[Outliner] Move 'LoadSettings' for the FilterBar to after the mode is initialized in case it adds custom filters so they can be saved appropriately

#Synced-CL 52418168

diff --git a/Engine/Source/Editor/SceneOutliner/Private/SSceneOutliner.cpp b/Engine/Source/Editor/SceneOutliner/Private/SSceneOutliner.cpp
--- a/Engine/Source/Editor/SceneOutliner/Private/SSceneOutliner.cpp
+++ b/Engine/Source/Editor/SceneOutliner/Private/SSceneOutliner.cpp
@@ -77,6 +77,13 @@ void SSceneOutliner::Construct(const FArguments& InArgs, const FSceneOutlinerIni
 	Mode = InInitOptions.ModeFactory.Execute(this);
 	check(Mode);
 
+	if (FilterBar)
+	{
+		// Defer loading the Filter settings until after the mode is initialized as they can
+		// add custom filters and we want those to be loaded/saved along with default filters
+		FilterBar->LoadSettings();
+	}	
+
 	bProcessingFullRefresh = false;
 	bFullRefresh = true;
 	bNeedsRefresh = true;
@@ -2839,8 +2846,6 @@ void SSceneOutliner::CreateFilterBar(const FSceneOutlinerFilterBarOptions& Filte
 
 		return MakeShareable(new FCustomTextFilter<SceneOutliner::FilterBarType>(Filter));
 	}));
-
-	FilterBar->LoadSettings();
 }
 
 bool SSceneOutliner::IsFilterActive(const FString& FilterName) const
