[TEDS Outliner] Allow the SceneOutlinerFilterBar accept generic Filters as class filters and add them to the add filter menu accordingly. (This allows Teds Filters to appear as type filters)

#Synced-CL 52730888

diff --git a/Engine/Plugins/Experimental/EditorDataStorageFeatures/Source/TedsOutliner/Private/TedsOutlinerHelpers.cpp b/Engine/Plugins/Experimental/EditorDataStorageFeatures/Source/TedsOutliner/Private/TedsOutlinerHelpers.cpp
--- a/Engine/Plugins/Experimental/EditorDataStorageFeatures/Source/TedsOutliner/Private/TedsOutlinerHelpers.cpp
+++ b/Engine/Plugins/Experimental/EditorDataStorageFeatures/Source/TedsOutliner/Private/TedsOutlinerHelpers.cpp
@@ -10,6 +10,7 @@
 #include "ISceneOutliner.h"
 #include "LevelEditor.h"
 #include "TedsOutlinerImpl.h"
+#include "Filters/CustomClassFilterData.h"
 #include "Widgets/Docking/SDockTab.h"
 
 namespace UE::Editor::Outliner::Helpers
@@ -259,4 +260,46 @@ namespace UE::Editor::Outliner::Helpers
     	}
     	return false;
     }
+    
+    void ConvertLegacyFiltersToTedsFilters(FSceneOutlinerFilterBarOptions& LegacyFilters, TArray<TSharedPtr<FTedsOutlinerFilter>>& TedsFilters,
+    	const TMap<FName, TVariant<QueryHandle, TConstQueryFunction<bool>>>& FilterConversionMap)
+	{
+		for (const TSharedRef<FCustomClassFilterData>& LegacyClassFilter : LegacyFilters.CustomClassFilters)
+		{
+			TedsFilters.Add(MakeShared<FTedsOutlinerFilter>(
+				LegacyClassFilter->GetClass(), LegacyClassFilter->GetCategories().IsEmpty() ? nullptr : LegacyClassFilter->GetCategories()[0]));
+		}
+		LegacyFilters.CustomClassFilters.Empty();
+
+		for (const TSharedRef<FFilterBase<SceneOutliner::FilterBarType>>& CustomFilter : LegacyFilters.CustomFilters)
+		{
+			const FName FilterName = FName(*CustomFilter->GetName());
+			const TVariant<QueryHandle, TConstQueryFunction<bool>>* FilterFunction = FilterConversionMap.Find(FilterName);
+			if (ensureAlwaysMsgf(FilterFunction, TEXT("Legacy Outliner Filter '%s' does not have a TEDS Mapping. It will not be "
+													"present in the Add Filter Menu in the TEDS Outliner."), *FilterName.ToString()))
+			{
+				if (FilterFunction->IsType<QueryHandle>())
+				{
+					TedsFilters.Add(MakeShared<FTedsOutlinerFilter>(
+					FilterName,
+					CustomFilter->GetDisplayName(),
+					CustomFilter->GetToolTipText(),
+					CustomFilter->GetIconName(),
+					CustomFilter->GetCategory(),
+					FilterFunction->Get<QueryHandle>()));
+				}
+				else if (FilterFunction->IsType<TConstQueryFunction<bool>>())
+				{
+					TedsFilters.Add(MakeShared<FTedsOutlinerFilter>(
+						FilterName,
+						CustomFilter->GetDisplayName(),
+						CustomFilter->GetToolTipText(),
+						CustomFilter->GetIconName(),
+						CustomFilter->GetCategory(),
+						FilterFunction->Get<TConstQueryFunction<bool>>()));
+				}
+			}
+		}
+		LegacyFilters.CustomFilters.Empty();
+	}
 }
diff --git a/Engine/Plugins/Experimental/EditorDataStorageFeatures/Source/TedsOutliner/Private/TedsOutlinerImpl.cpp b/Engine/Plugins/Experimental/EditorDataStorageFeatures/Source/TedsOutliner/Private/TedsOutlinerImpl.cpp
--- a/Engine/Plugins/Experimental/EditorDataStorageFeatures/Source/TedsOutliner/Private/TedsOutlinerImpl.cpp
+++ b/Engine/Plugins/Experimental/EditorDataStorageFeatures/Source/TedsOutliner/Private/TedsOutlinerImpl.cpp
@@ -350,7 +350,14 @@ void FTedsOutlinerImpl::CreateFilterQueries()
 			Filter->SetSceneOutlinerImpl(AsShared());
 			if (Filter->IsInteractiveFilter())
 			{
-				SceneOutliner->AddFilterToFilterBar(Filter.ToSharedRef());
+				if (Filter->IsClassFilter())
+				{
+					SceneOutliner->AddTypeFilterToFilterBar(Filter.ToSharedRef());
+				}
+				else
+				{
+					SceneOutliner->AddFilterToFilterBar(Filter.ToSharedRef());
+				}
 			}
 			else
 			{
diff --git a/Engine/Plugins/Experimental/EditorDataStorageFeatures/Source/TedsOutliner/Private/Widgets/SceneOutlinerWidget.cpp b/Engine/Plugins/Experimental/EditorDataStorageFeatures/Source/TedsOutliner/Private/Widgets/SceneOutlinerWidget.cpp
--- a/Engine/Plugins/Experimental/EditorDataStorageFeatures/Source/TedsOutliner/Private/Widgets/SceneOutlinerWidget.cpp
+++ b/Engine/Plugins/Experimental/EditorDataStorageFeatures/Source/TedsOutliner/Private/Widgets/SceneOutlinerWidget.cpp
@@ -23,6 +23,7 @@
 #include "LevelInstance/LevelInstanceInterface.h"
 #include "LevelInstance/LevelInstanceSubsystem.h"
 #include "WorldTreeItem.h"
+#include "Elements/Columns/TypedElementPackageColumns.h"
 
 #define LOCTEXT_NAMESPACE "SceneOutlinerWidget"
 
@@ -94,6 +95,8 @@ TSharedPtr<SWidget> FTedsSceneOutlinerWidgetConstructor::CreateWidget(UE::Editor
 	FTedsOutlinerParams Params(nullptr);
 	Params.QueryDescription = GetQueryDescription();
 	Params.ColumnDescription = GetColumnDescription();
+	
+	Helpers::ConvertLegacyFiltersToTedsFilters(InitOptions.FilterBarOptions, Params.Filters, GetLegacyFilterConversionMap(DataStorage));
 	GetFilters(DataStorage, Params.Filters);
 	
 	for (const TPair<FName, TSharedPtr<FTedsOutlinerFilter>>& Filter : ShowOptionsFilters)
@@ -170,7 +173,6 @@ void FTedsSceneOutlinerWidgetConstructor::GetFilters(UE::Editor::DataStorage::IC
 {
 	using namespace UE::Editor::Outliner;
 	
-	Filters.Add(MakeShared<FTedsOutlinerFilter>(AActor::StaticClass()));
 	if(WorldFilter)
 	{
 		Filters.Add(WorldFilter);
@@ -481,6 +483,42 @@ UTypedElementSelectionSet* FTedsSceneOutlinerWidgetConstructor::GetSelectionSet(
 	return nullptr;
 }
 
+TMap<FName, TVariant<UE::Editor::DataStorage::QueryHandle, UE::Editor::DataStorage::Queries::TConstQueryFunction<bool>>>
+	FTedsSceneOutlinerWidgetConstructor::GetLegacyFilterConversionMap(UE::Editor::DataStorage::ICoreProvider* DataStorage) const
+{
+	using namespace UE::Editor::DataStorage::Queries;
+	if (DataStorage)
+	{
+		return
+		{
+			{
+				FName("UnsavedAssetsFilter"),
+				TVariant<QueryHandle, TConstQueryFunction<bool>>( TInPlaceType<QueryHandle>(),
+				DataStorage->RegisterQuery(
+				Select()
+				.Where()
+				.All<FTedsPackageDirtyTag>()
+				.Compile()))
+			},
+			// TODO: Actually add UncontrolledAssetsFilter after improving Revision control integration, right now add this way to avoid ensure
+			{
+				FName("UncontrolledAssetsFilter"),
+				TVariant<QueryHandle, TConstQueryFunction<bool>>( TInPlaceType<TConstQueryFunction<bool>>(),
+					BuildConstQueryFunction<bool>([](TConstQueryContext<RowBatchInfo> Context, TResult<bool>& Result)
+					{
+						Context.ForEachRow([&Result](RowHandle Row)
+						{
+							Result.Add(Row, true);
+						});
+					})
+				)
+			},
+		};
+	}
+	return {};
+}
+
+
 void FTedsSceneOutlinerWidgetConstructor::GetLevelEditorSceneOutlinerInitOptions(FName TabIdentifier, FSceneOutlinerInitializationOptions& OutInitOptions) const
 {
 	TWeakPtr<ILevelEditor> LevelEditor = FModuleManager::GetModuleChecked<FLevelEditorModule>(LevelEditorModuleName).GetLevelEditorInstance();
diff --git a/Engine/Plugins/Experimental/EditorDataStorageFeatures/Source/TedsOutliner/Public/TedsOutlinerHelpers.h b/Engine/Plugins/Experimental/EditorDataStorageFeatures/Source/TedsOutliner/Public/TedsOutlinerHelpers.h
--- a/Engine/Plugins/Experimental/EditorDataStorageFeatures/Source/TedsOutliner/Public/TedsOutlinerHelpers.h
+++ b/Engine/Plugins/Experimental/EditorDataStorageFeatures/Source/TedsOutliner/Public/TedsOutlinerHelpers.h
@@ -7,6 +7,11 @@
 #include "SceneOutlinerFwd.h"
 #include "Columns/TedsOutlinerColumns.h"
 
+namespace UE::Editor::Outliner
+{
+	class FTedsOutlinerFilter;
+}
+
 class FTedsOutlinerColumnDescription;
 
 // Helper functions for Teds Outliner related functionality
@@ -46,4 +51,9 @@ namespace UE::Editor::Outliner::Helpers
 	
 	// Helper function to check if the given QueryHandle is valid to be used in a filter (is not an observer).
     bool CheckValidFilterQueryHandle(const DataStorage::QueryHandle& InQueryHandle);
+
+	// Helper function to convert legacy filters to Teds counterparts. All class filters will be automatically converted, any custom filters
+	// will need to have a mapped filter function counterpart to be converted successfully.
+	void TEDSOUTLINER_API ConvertLegacyFiltersToTedsFilters(FSceneOutlinerFilterBarOptions& LegacyFilters, TArray<TSharedPtr<FTedsOutlinerFilter>>& TedsFilters,
+		const TMap<FName, TVariant<DataStorage::QueryHandle, DataStorage::Queries::TConstQueryFunction<bool>>>& FilterConversionMap);
 }
diff --git a/Engine/Plugins/Experimental/EditorDataStorageFeatures/Source/TedsOutliner/Public/Widgets/SceneOutlinerWidget.h b/Engine/Plugins/Experimental/EditorDataStorageFeatures/Source/TedsOutliner/Public/Widgets/SceneOutlinerWidget.h
--- a/Engine/Plugins/Experimental/EditorDataStorageFeatures/Source/TedsOutliner/Public/Widgets/SceneOutlinerWidget.h
+++ b/Engine/Plugins/Experimental/EditorDataStorageFeatures/Source/TedsOutliner/Public/Widgets/SceneOutlinerWidget.h
@@ -89,6 +89,9 @@ public:
 	TEDSOUTLINER_API virtual TSharedPtr<UE::Editor::Outliner::ITedsOutlinerHierarchyDataInterface> GetHierarchyData(UE::Editor::DataStorage::ICoreProvider* DataStorage) const;
 	/** Gets the selection set to use for this Outliner, if nullptr it doesn't propagate tree selection through TypedElementSelectionInterface */
 	TEDSOUTLINER_API virtual UTypedElementSelectionSet* GetSelectionSet() const;
+	/** Gets a conversion map to map legacy Outliner filters to TEDS Filters, if a filter is not mapped, it will not appear in the add filter menu */
+	TEDSOUTLINER_API virtual TMap<FName, TVariant<UE::Editor::DataStorage::QueryHandle, UE::Editor::DataStorage::Queries::TConstQueryFunction<bool>>>
+		GetLegacyFilterConversionMap(UE::Editor::DataStorage::ICoreProvider* DataStorage) const;
 	
 protected:
 	/** Gets the FSceneOutlinerInitializationOptions defined in the Level Editor for the Scene Outliner to be used as the default */
diff --git a/Engine/Plugins/Experimental/EditorDataStorageFeatures/Source/TedsOutliner/TedsOutliner.Build.cs b/Engine/Plugins/Experimental/EditorDataStorageFeatures/Source/TedsOutliner/TedsOutliner.Build.cs
--- a/Engine/Plugins/Experimental/EditorDataStorageFeatures/Source/TedsOutliner/TedsOutliner.Build.cs
+++ b/Engine/Plugins/Experimental/EditorDataStorageFeatures/Source/TedsOutliner/TedsOutliner.Build.cs
@@ -23,6 +23,7 @@ public class TedsOutliner : ModuleRules
 					"CoreUObject",
 					"EditorConfig",
 					"EditorFramework",
+					"EditorWidgets",
 					"Engine",
 					"SceneOutliner",
 					"Slate",
diff --git a/Engine/Source/Editor/SceneOutliner/Private/SSceneOutliner.cpp b/Engine/Source/Editor/SceneOutliner/Private/SSceneOutliner.cpp
--- a/Engine/Source/Editor/SceneOutliner/Private/SSceneOutliner.cpp
+++ b/Engine/Source/Editor/SceneOutliner/Private/SSceneOutliner.cpp
@@ -2776,6 +2776,17 @@ bool SSceneOutliner::NeedsRefresh() const
 	return bNeedsRefresh || bFullRefresh;
 }
 
+void SSceneOutliner::AddTypeFilterToFilterBar(const TSharedRef<FFilterBase<SceneOutliner::FilterBarType>>& InFilter) const
+{
+	if (FilterBar)
+	{
+		if (TSharedPtr<SSceneOutlinerFilterBar> OutlinerFilterBar = StaticCastSharedPtr<SSceneOutlinerFilterBar>(FilterBar))
+		{
+			OutlinerFilterBar->AddTypeFilter(InFilter);
+		}
+	}
+}
+
 void SSceneOutliner::CreateCustomTextFilter(const FCustomTextFilterData& InFilterData, bool bApplyFilter)
 {
 	if (FilterBar)
diff --git a/Engine/Source/Editor/SceneOutliner/Private/SceneOutlinerFilterBar.cpp b/Engine/Source/Editor/SceneOutliner/Private/SceneOutlinerFilterBar.cpp
--- a/Engine/Source/Editor/SceneOutliner/Private/SceneOutlinerFilterBar.cpp
+++ b/Engine/Source/Editor/SceneOutliner/Private/SceneOutlinerFilterBar.cpp
@@ -2,6 +2,8 @@
 
 #include "SceneOutlinerFilterBar.h"
 
+#define LOCTEXT_NAMESPACE "SceneOutlinerFilterBar"
+
 const FName SSceneOutlinerFilterBar::SharedIdentifier("SceneOutlinerFilterBarSharedSettings");
 SSceneOutlinerFilterBar::FCustomTextFilterEvent SSceneOutlinerFilterBar::CustomTextFilterEvent;
 
@@ -143,11 +145,179 @@ void SSceneOutlinerFilterBar::LoadSettings()
 	this->OnFilterChanged.ExecuteIfBound();
 }
 
+void SSceneOutlinerFilterBar::AddTypeFilter(const TSharedRef<FFilterBase<SceneOutliner::FilterBarType>>& InFilter)
+{
+	TypeFilters.Add(InFilter.ToSharedPtr());
+	AddFilter(InFilter);
+}
+
+void SSceneOutlinerFilterBar::PopulateAddFilterMenu(UToolMenu* Menu, TSharedPtr<FFilterCategory> MenuExpansion, FOnFilterAssetType OnFilterAssetType)
+{
+	using namespace UE::Editor::Widgets;
+
+	// Build a map holding both FCustomClassFilterData entries and FFilterBase type
+	// filters so the rest of the function only iterates one category rather than having them separate.
+	TMap<TSharedPtr<FFilterCategory>, FTypeFilterCategoryEntry> TypeFilterMap;
+	PopulateTypeFilterMap(OnFilterAssetType, TypeFilterMap);
+	
+	PopulateCommonFilterSections(Menu);
+
+	auto BuildClassFilterDataEntry = [this](FToolMenuSection& Section, const TSharedPtr<FFilterCategory>& Category,
+		const TArray<TSharedPtr<FCustomClassFilterData>>& Classes)
+	{
+		Section.AddSubMenu(
+			FName(FText::AsCultureInvariant(Category->Title).ToString()),
+			Category->Title,
+			Category->Tooltip,
+			FNewToolMenuDelegate::CreateSP(this, &SSceneOutlinerFilterBar::CreateFiltersMenuCategory, Classes),
+			FUIAction(
+				FExecuteAction::CreateSP(this, &SSceneOutlinerFilterBar::FilterByTypeCategory, Category, Classes),
+				FCanExecuteAction(),
+				FGetActionCheckState::CreateSP(this, &SSceneOutlinerFilterBar::IsTypeCategoryChecked, Category, Classes)),
+			EUserInterfaceActionType::ToggleButton);
+	};
+
+	auto BuildBaseFilterEntry = [this](FToolMenuSection& Section, const TSharedPtr<FFilterCategory>& Category)
+	{
+		Section.AddSubMenu(
+			FName(FText::AsCultureInvariant(Category->Title).ToString()),
+			Category->Title,
+			Category->Tooltip,
+			FNewToolMenuDelegate::CreateSP(this, &SSceneOutlinerFilterBar::CreateOtherFiltersMenuCategory, Category),
+			FUIAction(
+				FExecuteAction::CreateSP(this, &SSceneOutlinerFilterBar::FrontendFilterCategoryClicked, Category),
+				FCanExecuteAction(),
+				FGetActionCheckState::CreateSP(this, &SSceneOutlinerFilterBar::IsFrontendFilterCategoryChecked, Category)),
+			EUserInterfaceActionType::ToggleButton);
+	};
+
+	// "Other Filters" section: FFilterBase categories not present in TypeFilterMap
+	TArray<TSharedPtr<FFilterCategory>> OtherCategories = AllFilterCategories;
+
+	// Expand the designated category inline — before TypeSection so it appears above it.
+	if (MenuExpansion)
+	{
+		if (FTypeFilterCategoryEntry* ExpandedEntry = TypeFilterMap.Find(MenuExpansion))
+		{
+			FToolMenuSection& Section = Menu->AddSection(ExpandedEntry->SectionExtensionHook, ExpandedEntry->SectionHeading);
+			FUIAction EntryAction;
+			FNewToolMenuDelegate NewMenuDelegate;
+			if (!ExpandedEntry->Classes.IsEmpty())
+			{
+				EntryAction = FUIAction(
+					FExecuteAction::CreateSP(this, &SSceneOutlinerFilterBar::FilterByTypeCategory, MenuExpansion, ExpandedEntry->Classes),
+					FCanExecuteAction(),
+					FGetActionCheckState::CreateSP(this, &SSceneOutlinerFilterBar::IsTypeCategoryChecked, MenuExpansion, ExpandedEntry->Classes));
+				NewMenuDelegate = FNewToolMenuDelegate::CreateSP(this, &SSceneOutlinerFilterBar::CreateFiltersMenuCategory, ExpandedEntry->Classes);
+			}
+			else
+			{
+				EntryAction = FUIAction(
+				FExecuteAction::CreateSP(this, &SSceneOutlinerFilterBar::FrontendFilterCategoryClicked, MenuExpansion),
+					FCanExecuteAction(),
+					FGetActionCheckState::CreateSP(this, &SSceneOutlinerFilterBar::IsFrontendFilterCategoryChecked, MenuExpansion));
+				NewMenuDelegate = FNewToolMenuDelegate::CreateSP(this, &SSceneOutlinerFilterBar::CreateOtherFiltersMenuCategory, MenuExpansion);
+			}
+
+			Section.AddMenuEntry(
+				FName(FText::AsCultureInvariant(ExpandedEntry->SectionHeading).ToString()),
+				ExpandedEntry->SectionHeading,
+				MenuExpansion->Tooltip,
+				FSlateIcon(FAppStyle::Get().GetStyleSetName(), "PlacementBrowser.Icons.Basic"),
+				EntryAction,
+				EUserInterfaceActionType::ToggleButton);
+
+			Section.AddSeparator("ExpandedCategorySeparator");
+
+			NewMenuDelegate.Execute(Menu);
+
+			TypeFilterMap.Remove(MenuExpansion);
+			OtherCategories.Remove(MenuExpansion);
+		}
+	}
+	
+	// Sort all categories and their entries alphabetically
+	TypeFilterMap.KeySort([](const TSharedPtr<FFilterCategory>& A, const TSharedPtr<FFilterCategory>& B) {
+		return A->Title.CompareTo(B->Title) < 0;
+	});
+	
+	for (TPair<TSharedPtr<FFilterCategory>, FTypeFilterCategoryEntry>& Pair : TypeFilterMap)
+	{
+		Pair.Value.Classes.Sort([](const TSharedPtr<FCustomClassFilterData>& A, const TSharedPtr<FCustomClassFilterData>& B) {
+			return A->GetName().CompareTo(B->GetName()) < 0;
+		});
+	}
+
+	FToolMenuSection& TypeSection = Menu->AddSection("AssetFilterBarFilterAdvancedAsset", LOCTEXT("TypeFiltersHeading", "Type Filters"));
+	for (const TPair<TSharedPtr<FFilterCategory>, FTypeFilterCategoryEntry>& Pair : TypeFilterMap)
+	{
+		if (!Pair.Value.Classes.IsEmpty())
+		{
+			BuildClassFilterDataEntry(TypeSection, Pair.Key, Pair.Value.Classes);
+		}
+		else
+		{
+			BuildBaseFilterEntry(TypeSection, Pair.Key);
+		}
+
+		OtherCategories.Remove(Pair.Key);
+	}
+
+	if (!OtherCategories.IsEmpty())
+	{
+		FToolMenuSection& OtherSection = Menu->AddSection(
+			"BasicFilterBarFiltersMenu",
+			LOCTEXT("OtherFiltersHeading", "Other Filters"));
+
+		for (const TSharedPtr<FFilterCategory>& Category : OtherCategories)
+		{
+			BuildBaseFilterEntry(OtherSection, Category);
+		}
+	}
+}
+
+void SSceneOutlinerFilterBar::PopulateTypeFilterMap(const FOnFilterAssetType& OnFilterAssetType, TMap<TSharedPtr<FFilterCategory>, FTypeFilterCategoryEntry>& TypeFilterMap)
+{
+	// Populate from FCustomClassFilterData via the SAssetFilterBar class helper
+	TMap<TSharedPtr<FFilterCategory>, UE::Editor::Widgets::FFilterCategoryMenu> ClassMap =
+		UE::Editor::Widgets::BuildCategoryToMenuMap(AssetFilterCategories, CustomClassFilters, OnFilterAssetType);
+	for (TPair<TSharedPtr<FFilterCategory>, UE::Editor::Widgets::FFilterCategoryMenu>& Pair : ClassMap)
+	{
+		FTypeFilterCategoryEntry& Entry = TypeFilterMap.FindOrAdd(Pair.Key);
+		Entry.SectionExtensionHook = Pair.Value.SectionExtensionHook;
+		Entry.SectionHeading = Pair.Value.SectionHeading;
+		Entry.Classes = MoveTemp(Pair.Value.Classes);
+	}
+
+	// Populate from FFilterBase type filters
+	for (const TSharedPtr<FFilterBase<SceneOutliner::FilterBarType>>& TypeFilter : TypeFilters)
+	{
+		if (TSharedPtr<FFilterCategory> Category = TypeFilter->GetCategory())
+		{
+			FTypeFilterCategoryEntry& Entry = TypeFilterMap.FindOrAdd(Category);
+			if (Entry.SectionHeading.IsEmpty())
+			{
+				Entry.SectionHeading = FText::Format(LOCTEXT("ExpandedCategoryHeading", "{0} Filters"), Category->Title);
+				Entry.SectionExtensionHook = FName(FText::AsCultureInvariant(Entry.SectionHeading).ToString());
+			}
+			Entry.Filters.Add(TypeFilter);
+		}
+	}
+}
+
 UAssetFilterBarContext* SSceneOutlinerFilterBar::CreateAssetFilterBarContext()
 {
 	UAssetFilterBarContext* AssetFilterBarContext = SFilterBar<const ISceneOutlinerTreeItem&>::CreateAssetFilterBarContext();
 	AssetFilterBarContext->MenuExpansion = CategoryToExpand;
 
+	// Use our PopulateAddFilterMenu whenever there are TEDS type filters to route into the
+	// "Type Filters" section, or when a category needs to be expanded inline (the base class
+	// PopulateCustomFilters has no per-category expansion support).
+	if (!TypeFilters.IsEmpty() || CategoryToExpand.IsValid())
+	{
+		AssetFilterBarContext->PopulateFilterMenu = FOnPopulateAddAssetFilterMenu::CreateSP(this, &SSceneOutlinerFilterBar::PopulateAddFilterMenu);
+	}
+
 	return AssetFilterBarContext;
 }
 
@@ -294,4 +464,6 @@ void SSceneOutlinerFilterBar::OnExternalCustomTextFilterCreated(TSharedPtr<SWidg
 
 	// Save the settings for this instance
 	SaveSettings();
-}
\ No newline at end of file
+}
+
+#undef LOCTEXT_NAMESPACE
\ No newline at end of file
diff --git a/Engine/Source/Editor/SceneOutliner/Private/SceneOutlinerFilterBar.h b/Engine/Source/Editor/SceneOutliner/Private/SceneOutlinerFilterBar.h
--- a/Engine/Source/Editor/SceneOutliner/Private/SceneOutlinerFilterBar.h
+++ b/Engine/Source/Editor/SceneOutliner/Private/SceneOutlinerFilterBar.h
@@ -70,6 +70,14 @@ public:
 
 	void Construct( const FArguments& InArgs );
 
+	/**
+	 * Register a filter as a "type filter" so it appears as an entry in the "Type Filters" section of the Add Filter menu
+	 * rather than under "Other Filters". The filter must already have its category set before this is called.
+	 */
+	void AddTypeFilter(const TSharedRef<FFilterBase<SceneOutliner::FilterBarType>>& InFilter);
+	/** Override from the AssetFilterBar to allow non-FCustomClassFilterData to populate the TypeFilter section */
+	virtual void PopulateAddFilterMenu(UToolMenu* Menu, TSharedPtr<FFilterCategory> MenuExpansion, FOnFilterAssetType OnFilterAssetType) override;
+
 	virtual void SaveSettings() override;
 	virtual void LoadSettings() override;
 
@@ -84,6 +92,23 @@ protected:
 	virtual UAssetFilterBarContext* CreateAssetFilterBarContext() override;
 
 private:
+	/**
+	 * Unified representation of a single "Type Filters" menu category.
+	 * A category may contain FCustomClassFilterData entries, FFilterBase entries, or both.
+	 */
+	struct FTypeFilterCategoryEntry
+	{
+		/** Extension hook and heading used when this category is rendered as a section or submenu */
+		FName SectionExtensionHook;
+		FText SectionHeading;
+
+		/** Asset class filters managed by the asset filter bar infrastructure */
+		TArray<TSharedPtr<FCustomClassFilterData>> Classes;
+
+		/** Type filters managed via AllFrontendFilters */
+		TArray<TSharedPtr<FFilterBase<SceneOutliner::FilterBarType>>> Filters;
+	};
+	
 	/** Handler for when another Filter Bar using shared settings creates a custom text filter */
 	void OnExternalCustomTextFilterCreated(TSharedPtr<SWidget> BroadcastingFilterBar);
 	
@@ -94,6 +119,8 @@ private:
 	 *  @return True if the filter was restored successfully, false if not
 	 */
 	bool RestoreCustomTextFilterState(const FCustomTextFilterState& InFilterState);
+
+	void PopulateTypeFilterMap(const FOnFilterAssetType& OnFilterAssetType, TMap<TSharedPtr<FFilterCategory>, FTypeFilterCategoryEntry>& TypeFilterMap);
 	
 private:
 	/** The event that executes when a custom text filter is created/modified/deleted in any filter list that is using Shared Settings */
@@ -105,5 +132,8 @@ private:
 	/** Whether this Filter bar wants to load/save from the settings common to all instances */
 	bool bUseSharedSettings;
 
-	TSharedPtr<FFilterCategory> CategoryToExpand;	
+	TSharedPtr<FFilterCategory> CategoryToExpand;
+
+	/** Filters registered via AddTypeFilter, used when building the 'Type Filters' menu section */
+	TSet<TSharedPtr<FFilterBase<SceneOutliner::FilterBarType>>> TypeFilters;
 };
\ No newline at end of file
diff --git a/Engine/Source/Editor/SceneOutliner/Public/SSceneOutliner.h b/Engine/Source/Editor/SceneOutliner/Public/SSceneOutliner.h
--- a/Engine/Source/Editor/SceneOutliner/Public/SSceneOutliner.h
+++ b/Engine/Source/Editor/SceneOutliner/Public/SSceneOutliner.h
@@ -1092,6 +1092,12 @@ public:
 	UE_API void FilterByTypeCategory(TSharedPtr<FFilterCategory> TypeCategory, TArray<TSharedPtr<FCustomClassFilterData>> Classes);
 	UE_API bool NeedsRefresh() const;
 
+	/**
+	 * Add a type filter to the scene outliner's filter bar. Added separately so it can appear in the 'Type Filters' section
+	 * of the Add Filter menu rather than "Other Filters". The filter must have its category set before this call.
+	 */
+	UE_API void AddTypeFilterToFilterBar(const TSharedRef<FFilterBase<SceneOutliner::FilterBarType>>& InFilter) const;
+
 	UE_API void CreateCustomTextFilter(const FCustomTextFilterData& InFilterData, bool bApplyFilter);
 	UE_API void ModifyCustomTextFilter(const FCustomTextFilterData& InFilterData, const TSharedPtr<ICustomTextFilter<SceneOutliner::FilterBarType>>& InFilter);
 	UE_API void DeleteCustomTextFilter(const TSharedPtr<ICustomTextFilter<SceneOutliner::FilterBarType>>& InFilter);
